Non-unity fixes for Mac.

diff --git a/Engine/Source/Runtime/ApplicationCore/Private/Mac/MacConsoleOutputDevice.cpp b/Engine/Source/Runtime/ApplicationCore/Private/Mac/MacConsoleOutputDevice.cpp
--- a/Engine/Source/Runtime/ApplicationCore/Private/Mac/MacConsoleOutputDevice.cpp
+++ b/Engine/Source/Runtime/ApplicationCore/Private/Mac/MacConsoleOutputDevice.cpp
@@ -7,6 +7,7 @@
 #include "MacApplication.h"
 #include "MacPlatformApplicationMisc.h"
 #include "Misc/OutputDeviceHelper.h"
+#include "Misc/ScopeLock.h"
 #include "CocoaThread.h"
 #include "HAL/PlatformApplicationMisc.h"
 
diff --git a/Engine/Source/Runtime/ApplicationCore/Private/Mac/MacErrorOutputDevice.cpp b/Engine/Source/Runtime/ApplicationCore/Private/Mac/MacErrorOutputDevice.cpp
--- a/Engine/Source/Runtime/ApplicationCore/Private/Mac/MacErrorOutputDevice.cpp
+++ b/Engine/Source/Runtime/ApplicationCore/Private/Mac/MacErrorOutputDevice.cpp
@@ -6,6 +6,7 @@
 #include "Misc/ConfigCacheIni.h"
 #include "FeedbackContextAnsi.h"
 #include "MacApplication.h"
+#include "MacConsoleOutputDevice.h"
 #include "MacPlatformApplicationMisc.h"
 #include "CocoaThread.h"
 
diff --git a/Engine/Source/Runtime/ApplicationCore/Private/Mac/MacPlatformApplicationMisc.cpp b/Engine/Source/Runtime/ApplicationCore/Private/Mac/MacPlatformApplicationMisc.cpp
--- a/Engine/Source/Runtime/ApplicationCore/Private/Mac/MacPlatformApplicationMisc.cpp
+++ b/Engine/Source/Runtime/ApplicationCore/Private/Mac/MacPlatformApplicationMisc.cpp
@@ -10,6 +10,10 @@
 #include "MacConsoleOutputDevice.h"
 #include "MacErrorOutputDevice.h"
 #include "MacFeedbackContext.h"
+#include "Misc/OutputDeviceHelper.h"
+#include "Misc/OutputDeviceRedirector.h"
+#include "HAL/ThreadHeartBeat.h"
+#include "Mac/CocoaMenu.h"
 
 #include <dlfcn.h>
 #include <IOKit/IOKitLib.h>
diff --git a/Engine/Source/Runtime/Core/Private/Mac/MacCriticalSection.cpp b/Engine/Source/Runtime/Core/Private/Mac/MacCriticalSection.cpp
--- a/Engine/Source/Runtime/Core/Private/Mac/MacCriticalSection.cpp
+++ b/Engine/Source/Runtime/Core/Private/Mac/MacCriticalSection.cpp
@@ -6,6 +6,7 @@
 
 #include "Mac/MacCriticalSection.h"
 #include "Misc/App.h"
+#include "Misc/DateTime.h"
 #include <mach-o/dyld.h>
 #include <libproc.h>
 
