CIS fixes

diff --git a/Engine/Plugins/Runtime/Advertising/IOSTapJoy/Source/IOSTapJoy/IOSTapJoy.Build.cs b/Engine/Plugins/Runtime/Advertising/IOSTapJoy/Source/IOSTapJoy/IOSTapJoy.Build.cs
--- a/Engine/Plugins/Runtime/Advertising/IOSTapJoy/Source/IOSTapJoy/IOSTapJoy.Build.cs
+++ b/Engine/Plugins/Runtime/Advertising/IOSTapJoy/Source/IOSTapJoy/IOSTapJoy.Build.cs
@@ -30,6 +30,7 @@ namespace UnrealBuildTool.Rules
 				new string[]
 				{
 					"Advertising",
+					"ApplicationCore",
 					// ... add private dependencies that you statically link with here ...
 				}
 				);
diff --git a/Engine/Plugins/Runtime/GoogleVR/GoogleVRHMD/Source/GoogleVRHMD/GoogleVRHMD.Build.cs b/Engine/Plugins/Runtime/GoogleVR/GoogleVRHMD/Source/GoogleVRHMD/GoogleVRHMD.Build.cs
--- a/Engine/Plugins/Runtime/GoogleVR/GoogleVRHMD/Source/GoogleVRHMD/GoogleVRHMD.Build.cs
+++ b/Engine/Plugins/Runtime/GoogleVR/GoogleVRHMD/Source/GoogleVRHMD/GoogleVRHMD.Build.cs
@@ -66,7 +66,7 @@ namespace UnrealBuildTool.Rules
 			}
 			else if (Target.Platform == UnrealTargetPlatform.IOS)
 			{
-				PrivateDependencyModuleNames.AddRange(new string[] { "GoogleVR" });
+				PrivateDependencyModuleNames.AddRange(new string[] { "GoogleVR", "ApplicationCore" });
 
 				string PluginPath = Utils.MakePathRelativeTo(ModuleDirectory, Target.RelativeEnginePath);
 				AdditionalPropertiesForReceipt.Add(new ReceiptProperty("IOSPlugin", Path.Combine(PluginPath, "GoogleVRHMD_APL.xml")));
diff --git a/Engine/Source/Runtime/Advertising/IOS/IOSAdvertising/IOSAdvertising.Build.cs b/Engine/Source/Runtime/Advertising/IOS/IOSAdvertising/IOSAdvertising.Build.cs
--- a/Engine/Source/Runtime/Advertising/IOS/IOSAdvertising/IOSAdvertising.Build.cs
+++ b/Engine/Source/Runtime/Advertising/IOS/IOSAdvertising/IOSAdvertising.Build.cs
@@ -31,6 +31,7 @@ namespace UnrealBuildTool.Rules
 				new string[]
 				{
 					"Advertising",
+					"ApplicationCore"
                     // ... add private dependencies that you statically link with here ...
 				}
 				);
diff --git a/Engine/Source/Runtime/Advertising/IOS/IOSAdvertising/Private/IOSAdvertising.cpp b/Engine/Source/Runtime/Advertising/IOS/IOSAdvertising/Private/IOSAdvertising.cpp
--- a/Engine/Source/Runtime/Advertising/IOS/IOSAdvertising/Private/IOSAdvertising.cpp
+++ b/Engine/Source/Runtime/Advertising/IOS/IOSAdvertising/Private/IOSAdvertising.cpp
@@ -4,6 +4,7 @@
 #include "IOSAppDelegate.h"
 #include "Modules/ModuleManager.h"
 #include "IOS/IOSAsyncTask.h"
+#include "IOS/IOSAppDelegate.h"
 
 #import <iAd/ADBannerView.h>
 
diff --git a/Engine/Source/Runtime/ApplicationCore/ApplicationCore.Build.cs b/Engine/Source/Runtime/ApplicationCore/ApplicationCore.Build.cs
--- a/Engine/Source/Runtime/ApplicationCore/ApplicationCore.Build.cs
+++ b/Engine/Source/Runtime/ApplicationCore/ApplicationCore.Build.cs
@@ -43,31 +43,35 @@ public class ApplicationCore : ModuleRules
 				PublicAdditionalLibraries.Add("/System/Library/PrivateFrameworks/MultitouchSupport.framework/Versions/Current/MultitouchSupport");
 			}
 		}
-        else if ((Target.Platform == UnrealTargetPlatform.Linux))
-        {
+		else if ((Target.Platform == UnrealTargetPlatform.Linux))
+		{
 			AddEngineThirdPartyPrivateStaticDependencies(Target, 
 				"SDL2"
-                );
+				);
 
-            // We need FreeType2 and GL for the Splash, but only in the Editor
-            if (Target.Type == TargetType.Editor)
-            {
-                AddEngineThirdPartyPrivateStaticDependencies(Target, "FreeType2");
+			// We need FreeType2 and GL for the Splash, but only in the Editor
+			if (Target.Type == TargetType.Editor)
+			{
+				AddEngineThirdPartyPrivateStaticDependencies(Target, "FreeType2");
 				AddEngineThirdPartyPrivateStaticDependencies(Target, "OpenGL");
 				PrivateIncludePathModuleNames.Add("ImageWrapper");
 			}
-        }
+		}
 		else if (Target.Platform == UnrealTargetPlatform.HTML5 && Target.Architecture == "-win32")
 		{
 			AddEngineThirdPartyPrivateStaticDependencies(Target, "SDL2");
 			AddEngineThirdPartyPrivateStaticDependencies(Target, "OpenAL");
 		}
-        else if (Target.Platform == UnrealTargetPlatform.HTML5 && Target.Architecture != "-win32")
-        {
-            AddEngineThirdPartyPrivateStaticDependencies(Target, "SDL2");
-            PrivateDependencyModuleNames.Add("HTML5JS");
-            PrivateDependencyModuleNames.Add("MapPakDownloader");
-        }
+		else if (Target.Platform == UnrealTargetPlatform.HTML5 && Target.Architecture != "-win32")
+		{
+			AddEngineThirdPartyPrivateStaticDependencies(Target, "SDL2");
+			PrivateDependencyModuleNames.Add("HTML5JS");
+			PrivateDependencyModuleNames.Add("MapPakDownloader");
+		}
+		else if (Target.Platform == UnrealTargetPlatform.IOS || Target.Platform == UnrealTargetPlatform.TVOS)
+		{
+			PublicIncludePaths.AddRange(new string[] {"Runtime/ApplicationCore/Public/Apple", "Runtime/ApplicationCore/Public/IOS"});
+		}
 
 	}
 }
diff --git a/Engine/Source/Runtime/ApplicationCore/Private/Android/AndroidErrorOutputDevice.cpp b/Engine/Source/Runtime/ApplicationCore/Private/Android/AndroidErrorOutputDevice.cpp
--- a/Engine/Source/Runtime/ApplicationCore/Private/Android/AndroidErrorOutputDevice.cpp
+++ b/Engine/Source/Runtime/ApplicationCore/Private/Android/AndroidErrorOutputDevice.cpp
@@ -3,6 +3,7 @@
 #include "AndroidErrorOutputDevice.h"
 #include "OutputDeviceHelper.h"
 #include "HAL/PlatformMisc.h"
+#include "Misc/OutputDeviceRedirector.h"
 
 FAndroidErrorOutputDevice::FAndroidErrorOutputDevice()
 {
diff --git a/Engine/Source/Runtime/ApplicationCore/Private/Mac/MacPlatformApplicationMisc.cpp b/Engine/Source/Runtime/ApplicationCore/Private/Mac/MacPlatformApplicationMisc.cpp
--- a/Engine/Source/Runtime/ApplicationCore/Private/Mac/MacPlatformApplicationMisc.cpp
+++ b/Engine/Source/Runtime/ApplicationCore/Private/Mac/MacPlatformApplicationMisc.cpp
@@ -14,6 +14,7 @@
 #include "Misc/OutputDeviceRedirector.h"
 #include "HAL/ThreadHeartBeat.h"
 #include "Mac/CocoaMenu.h"
+#include "Mac/CocoaThread.h"
 
 #include <dlfcn.h>
 #include <IOKit/IOKitLib.h>
diff --git a/Engine/Source/Runtime/ApplicationCore/Public/Mac/MacFeedbackContext.h b/Engine/Source/Runtime/ApplicationCore/Public/Mac/MacFeedbackContext.h
--- a/Engine/Source/Runtime/ApplicationCore/Public/Mac/MacFeedbackContext.h
+++ b/Engine/Source/Runtime/ApplicationCore/Public/Mac/MacFeedbackContext.h
@@ -5,6 +5,8 @@
 #include "CoreMinimal.h"
 #include "Misc/FeedbackContext.h"
 #include "Misc/OutputDeviceConsole.h"
+#include "Misc/OutputDeviceHelper.h"
+#include "Mac/OutputDeviceRedirector.h"
 
 /**
  * Feedback context implementation for Mac.
diff --git a/Engine/Source/Runtime/InputCore/InputCore.Build.cs b/Engine/Source/Runtime/InputCore/InputCore.Build.cs
--- a/Engine/Source/Runtime/InputCore/InputCore.Build.cs
+++ b/Engine/Source/Runtime/InputCore/InputCore.Build.cs
@@ -8,7 +8,7 @@ public class InputCore : ModuleRules
 	{
 		PublicDependencyModuleNames.AddRange(new string[] { "Core", "CoreUObject" });
 
-		if(Target.Platform == UnrealTargetPlatform.IOS)
+		if(Target.Platform == UnrealTargetPlatform.IOS || Target.Platform == UnrealTargetPlatform.TVOS)
 		{
 			PrivateIncludePathModuleNames.Add("ApplicationCore");
 		}
