Add missing include for Mac UE4Game non-unity.

diff --git a/Engine/Source/Runtime/ApplicationCore/Private/Mac/MacPlatformApplicationMisc.cpp b/Engine/Source/Runtime/ApplicationCore/Private/Mac/MacPlatformApplicationMisc.cpp
--- a/Engine/Source/Runtime/ApplicationCore/Private/Mac/MacPlatformApplicationMisc.cpp
+++ b/Engine/Source/Runtime/ApplicationCore/Private/Mac/MacPlatformApplicationMisc.cpp
@@ -13,6 +13,7 @@
 #include "Misc/OutputDeviceHelper.h"
 #include "Misc/OutputDeviceRedirector.h"
 #include "HAL/ThreadHeartBeat.h"
+#include "HAL/FeedbackContextAnsi.h"
 #include "Mac/CocoaMenu.h"
 #include "Mac/CocoaThread.h"
 
