Rolling back game specific hack marked 'DO NOT MERGE' that was mistakenly introduced to Main.

diff --git a/Engine/Source/Runtime/ApplicationCore/Private/IOS/IOSAppDelegate.cpp b/Engine/Source/Runtime/ApplicationCore/Private/IOS/IOSAppDelegate.cpp
--- a/Engine/Source/Runtime/ApplicationCore/Private/IOS/IOSAppDelegate.cpp
+++ b/Engine/Source/Runtime/ApplicationCore/Private/IOS/IOSAppDelegate.cpp
@@ -173,9 +173,8 @@ void EngineCrashHandler(const FGenericCrashContext& GenericContext)
 	//}
 
 	// check to see if we are using the network file system, if so, disable the idle timer
-    // DO NOT MERGE TO MAIN - temporarily turn off the idle timer to see if some backgrounding issues are reduced
-//	FString HostIP;
-//	if (FParse::Value(FCommandLine::Get(), TEXT("-FileHostIP="), HostIP))
+	FString HostIP;
+	if (FParse::Value(FCommandLine::Get(), TEXT("-FileHostIP="), HostIP))
 	{
 		[UIApplication sharedApplication].idleTimerDisabled = YES;
 	}
