Adjust Lumin to cursor argument for input interface
#rnx

diff --git a/Engine/Source/Runtime/ApplicationCore/Private/Android/AndroidApplication.cpp b/Engine/Source/Runtime/ApplicationCore/Private/Android/AndroidApplication.cpp
--- a/Engine/Source/Runtime/ApplicationCore/Private/Android/AndroidApplication.cpp
+++ b/Engine/Source/Runtime/ApplicationCore/Private/Android/AndroidApplication.cpp
@@ -28,7 +28,7 @@ FAndroidApplication::FAndroidApplication()
 }
 
 FAndroidApplication::FAndroidApplication(TSharedPtr<class FAndroidInputInterface> InInputInterface)
-	: GenericApplication(InInputInterface.IsValid() ? InInputInterface->GetCursor() : MakeShareable(new FAndroidCursor()))
+	: GenericApplication((InInputInterface.IsValid() && InInputInterface->GetCursor()) ? InInputInterface->GetCursor() : MakeShareable(new FAndroidCursor()))
 	, InputInterface(InInputInterface)
 	, bHasLoadedInputPlugins(false)
 {
diff --git a/Engine/Source/Runtime/ApplicationCore/Private/Lumin/LuminInputInterface.cpp b/Engine/Source/Runtime/ApplicationCore/Private/Lumin/LuminInputInterface.cpp
--- a/Engine/Source/Runtime/ApplicationCore/Private/Lumin/LuminInputInterface.cpp
+++ b/Engine/Source/Runtime/ApplicationCore/Private/Lumin/LuminInputInterface.cpp
@@ -39,7 +39,7 @@
 #include "IHapticDevice.h"
 #include "CoreMinimal.h"
 
-TSharedRef<FAndroidInputInterface> FAndroidInputInterface::Create(const TSharedRef<FGenericApplicationMessageHandler>& InMessageHandler)
+TSharedRef<FAndroidInputInterface> FAndroidInputInterface::Create(const TSharedRef<FGenericApplicationMessageHandler>& InMessageHandler, const TSharedPtr< ICursor >& InCursor)
 {
 	return MakeShareable(new FAndroidInputInterface(InMessageHandler));
 }
diff --git a/Engine/Source/Runtime/ApplicationCore/Public/Lumin/LuminInputInterface.h b/Engine/Source/Runtime/ApplicationCore/Public/Lumin/LuminInputInterface.h
--- a/Engine/Source/Runtime/ApplicationCore/Public/Lumin/LuminInputInterface.h
+++ b/Engine/Source/Runtime/ApplicationCore/Public/Lumin/LuminInputInterface.h
@@ -46,7 +46,7 @@ class FAndroidInputInterface : public IInputInterface
 public:
 	virtual ~FAndroidInputInterface() = default;
 
-	static TSharedRef<FAndroidInputInterface> Create(const TSharedRef<FGenericApplicationMessageHandler>& InMessageHandler);
+	static TSharedRef<FAndroidInputInterface> Create(const TSharedRef<FGenericApplicationMessageHandler>& InMessageHandler, const TSharedPtr< ICursor >& InCursor);
 
 	/** Tick the interface (i.e check for new controllers) */
 	void Tick(float DeltaTime);
@@ -69,6 +69,8 @@ public:
 	virtual void AddExternalInputDevice(TSharedPtr<class IInputDevice>& InputDevice);
 	virtual bool IsGamepadAttached() const;
 
+	const TSharedPtr< ICursor > GetCursor() const { return nullptr; }
+
 private:
 	FAndroidInputInterface(const TSharedRef<FGenericApplicationMessageHandler>& InMessageHandler);
 
