#rnx

diff --git a/Engine/Source/Runtime/ApplicationCore/Private/HTML5/HTML5InputInterface.cpp b/Engine/Source/Runtime/ApplicationCore/Private/HTML5/HTML5InputInterface.cpp
--- a/Engine/Source/Runtime/ApplicationCore/Private/HTML5/HTML5InputInterface.cpp
+++ b/Engine/Source/Runtime/ApplicationCore/Private/HTML5/HTML5InputInterface.cpp
@@ -100,11 +100,8 @@ void FHTML5InputInterface::Tick(float DeltaTime, const SDL_Event& Event,TSharedR
 				const SDL_Keycode KeyCode = KeyEvent.keysym.scancode;
 				const bool bIsRepeated = KeyEvent.repeat != 0;
 
-				if ( KeyCode != 227 && KeyCode != 231 ) // UE-54056 -- filtering out Windows/Super key
-				{
-					// First KeyDown, then KeyChar. This is important, as in-game console ignores first character otherwise
-					MessageHandler->OnKeyDown(KeyCode, KeyEvent.keysym.sym, bIsRepeated);
-				}
+				// First KeyDown, then KeyChar. This is important, as in-game console ignores first character otherwise
+				MessageHandler->OnKeyDown(KeyCode, KeyEvent.keysym.sym, bIsRepeated);
 
 				// Backspace/Return input caught here.
 				// Note that TextInput still seems to get characters messages too but slate seems to not process them.
diff --git a/Engine/Source/Runtime/InputCore/Private/HTML5/HTML5PlatformInput.cpp b/Engine/Source/Runtime/InputCore/Private/HTML5/HTML5PlatformInput.cpp
--- a/Engine/Source/Runtime/InputCore/Private/HTML5/HTML5PlatformInput.cpp
+++ b/Engine/Source/Runtime/InputCore/Private/HTML5/HTML5PlatformInput.cpp
@@ -73,10 +73,13 @@ uint32 FHTML5PlatformInput::GetKeyMap(uint32* KeyCodes, FString* KeyNames, uint3
 		ADDKEYMAP(SDL_SCANCODE_KP_MINUS, TEXT("Subtract"));
 		ADDKEYMAP(SDL_SCANCODE_KP_DECIMAL, TEXT("Decimal"));
 		ADDKEYMAP(SDL_SCANCODE_KP_DIVIDE, TEXT("Divide"));
+		ADDKEYMAP(SDL_SCANCODE_KP_PERIOD, TEXT("Period"));
 
 		ADDKEYMAP(SDL_SCANCODE_CAPSLOCK, TEXT("CapsLock"));
 		ADDKEYMAP(SDL_SCANCODE_NUMLOCKCLEAR, TEXT("NumLock"));
 		ADDKEYMAP(SDL_SCANCODE_SCROLLLOCK, TEXT("ScrollLock"));
+		ADDKEYMAP(SDL_SCANCODE_CANCEL, TEXT("Cancel")); // UE-58441 -- CTRL+Pause/Break and CTRL+ScrollLock
+		ADDKEYMAP(SDL_SCANCODE_APPLICATION, TEXT("Menu")); // UE-58441 -- contextual menu
 	}
 #undef ADDKEYMAP
 
