Fixed a crash in FMacApplication::ProcessMouseUpEvent()

diff --git a/Engine/Source/Runtime/ApplicationCore/Private/Mac/MacApplication.cpp b/Engine/Source/Runtime/ApplicationCore/Private/Mac/MacApplication.cpp
--- a/Engine/Source/Runtime/ApplicationCore/Private/Mac/MacApplication.cpp
+++ b/Engine/Source/Runtime/ApplicationCore/Private/Mac/MacApplication.cpp
@@ -900,17 +900,20 @@ void FMacApplication::ProcessMouseUpEvent(const FDeferredMacEvent& Event, TShare
 	}
 
 	MessageHandler->OnMouseUp(Button);
-	
-	// 10.12.6 Fix for window position after dragging window to desktop selector in mission control
-	// 10.13.0 Doesn't need this as it always fires the window move event after desktop drag operation
-	if(DraggedWindow != nullptr && EventWindow->GetWindowHandle() == DraggedWindow)
-	{
-		OnWindowDidMove(EventWindow.ToSharedRef());
-	}
 
-	if (EventWindow.IsValid() && EventWindow->GetWindowHandle() && !DraggedWindow && !GetCapture())
+	if (EventWindow.IsValid())
 	{
-		MessageHandler->OnCursorSet();
+		// 10.12.6 Fix for window position after dragging window to desktop selector in mission control
+		// 10.13.0 Doesn't need this as it always fires the window move event after desktop drag operation
+		if (DraggedWindow != nullptr && EventWindow->GetWindowHandle() == DraggedWindow)
+		{
+			OnWindowDidMove(EventWindow.ToSharedRef());
+		}
+
+		if (EventWindow->GetWindowHandle() && !DraggedWindow && !GetCapture())
+		{
+			MessageHandler->OnCursorSet();
+		}
 	}
 
 	FPlatformApplicationMisc::bChachedMacMenuStateNeedsUpdate = true;
