Fix for crash in Mac editor copying unusual characters in the output log.

Current operation is unsafe.  String Length, number of bytes and code units are not always interchangeable in this way - use FString as that is the target type and does the right thing.

diff --git a/Engine/Source/Runtime/ApplicationCore/Private/Mac/MacPlatformApplicationMisc.cpp b/Engine/Source/Runtime/ApplicationCore/Private/Mac/MacPlatformApplicationMisc.cpp
--- a/Engine/Source/Runtime/ApplicationCore/Private/Mac/MacPlatformApplicationMisc.cpp
+++ b/Engine/Source/Runtime/ApplicationCore/Private/Mac/MacPlatformApplicationMisc.cpp
@@ -527,10 +527,7 @@ void FMacPlatformApplicationMisc::ClipboardPaste(class FString& Result)
 	NSString *CocoaString = [Pasteboard stringForType: NSPasteboardTypeString];
 	if (CocoaString)
 	{
-		TArray<TCHAR> Ch;
-		Ch.AddUninitialized([CocoaString length] + 1);
-		FPlatformString::CFStringToTCHAR((CFStringRef)CocoaString, Ch.GetData());
-		Result = Ch.GetData();
+		Result = FString(CocoaString);
 	}
 	else
 	{
