Speculative fix for crash in FMacWindow::ApplySizeAndModeChanges().  I had seen a crash in here while working on UE-46999 (fix for that submitted in CL 4016062) due to the window handle becoming nil mid way through the ApplySizeAndModeChanges() function then a deref'd.  This CL adds extra null pointer protection to make sure we check the validity of the WindowHandle after calls to UpdateFullScreenState().

diff --git a/Engine/Source/Runtime/ApplicationCore/Private/Mac/MacWindow.cpp b/Engine/Source/Runtime/ApplicationCore/Private/Mac/MacWindow.cpp
--- a/Engine/Source/Runtime/ApplicationCore/Private/Mac/MacWindow.cpp
+++ b/Engine/Source/Runtime/ApplicationCore/Private/Mac/MacWindow.cpp
@@ -366,7 +366,10 @@ void FMacWindow::Hide()
 
 void FMacWindow::SetWindowMode(EWindowMode::Type NewWindowMode)
 {
-	ApplySizeAndModeChanges(PositionX, PositionY, WindowHandle.contentView.frame.size.width, WindowHandle.contentView.frame.size.height, NewWindowMode);
+	if(WindowHandle)
+	{
+		ApplySizeAndModeChanges(PositionX, PositionY, WindowHandle.contentView.frame.size.width, WindowHandle.contentView.frame.size.height, NewWindowMode);
+	}
 }
 
 EWindowMode::Type FMacWindow::GetWindowMode() const
@@ -545,6 +548,9 @@ void FMacWindow::OnWindowDidChangeScreen()
 
 void FMacWindow::ApplySizeAndModeChanges(int32 X, int32 Y, int32 Width, int32 Height, EWindowMode::Type WindowMode)
 {
+	// It's possible for Window handle to become nil in FMacWindow::Destroy() after a call to UpdateFullScreenState() in this function,
+	// in rare cases due to FPlatformApplicationMisc::PumpMessages
+	
 	SCOPED_AUTORELEASE_POOL;
 
 	bool bIsFullScreen = [WindowHandle windowMode] == EWindowMode::WindowedFullscreen || [WindowHandle windowMode] == EWindowMode::Fullscreen;
@@ -581,39 +587,36 @@ void FMacWindow::ApplySizeAndModeChanges(int32 X, int32 Y, int32 Width, int32 He
 			UpdateFullScreenState(true);
 			bIsFullScreen = false;
 		}
+		
+		if(WindowHandle)
+		{
+			WindowHandle.TargetWindowMode = WindowMode;
 
-		WindowHandle.TargetWindowMode = WindowMode;
-
-		const float DPIScaleFactor = FPlatformApplicationMisc::GetDPIScaleFactorAtPoint(X, Y);
-		Width /= DPIScaleFactor;
-		Height /= DPIScaleFactor;
+			const float DPIScaleFactor = FPlatformApplicationMisc::GetDPIScaleFactorAtPoint(X, Y);
+			Width /= DPIScaleFactor;
+			Height /= DPIScaleFactor;
 
-		const FVector2D CocoaPosition = FMacApplication::ConvertSlatePositionToCocoa(X, Y);
-		NSRect Rect = NSMakeRect(CocoaPosition.X, CocoaPosition.Y - Height + 1, FMath::Max(Width, 1), FMath::Max(Height, 1));
-		if (Definition->HasOSWindowBorder)
-		{
-			Rect = [WindowHandle frameRectForContentRect:Rect];
-		}
+			const FVector2D CocoaPosition = FMacApplication::ConvertSlatePositionToCocoa(X, Y);
+			NSRect Rect = NSMakeRect(CocoaPosition.X, CocoaPosition.Y - Height + 1, FMath::Max(Width, 1), FMath::Max(Height, 1));
+			if (Definition->HasOSWindowBorder)
+			{
+				Rect = [WindowHandle frameRectForContentRect:Rect];
+			}
 
-		UpdateFullScreenState(bWantsFullScreen != bIsFullScreen);
-		
-		// It's possible for Window handle to be nil after a call to UpdateFullScreenState() in rare cases due to FPlatformApplicationMisc::PumpMessages
-		if(WindowHandle == nil)
-		{
-			return;
-		}
+			UpdateFullScreenState(bWantsFullScreen != bIsFullScreen);
 
-		if (WindowMode == EWindowMode::Windowed && !NSEqualRects([WindowHandle frame], Rect))
-		{
-			MainThreadCall(^{
-				SCOPED_AUTORELEASE_POOL;
-				[WindowHandle setFrame:Rect display:YES];
+			if (WindowHandle && WindowMode == EWindowMode::Windowed && !NSEqualRects([WindowHandle frame], Rect))
+			{
+				MainThreadCall(^{
+					SCOPED_AUTORELEASE_POOL;
+					[WindowHandle setFrame:Rect display:YES];
 
-				if (Definition->ShouldPreserveAspectRatio)
-				{
-					[WindowHandle setContentAspectRatio:NSMakeSize((float)Width / (float)Height, 1.0f)];
-				}
-			}, UE4ResizeEventMode, true);
+					if (Definition->ShouldPreserveAspectRatio)
+					{
+						[WindowHandle setContentAspectRatio:NSMakeSize((float)Width / (float)Height, 1.0f)];
+					}
+				}, UE4ResizeEventMode, true);
+			}
 		}
 	}
 	else
@@ -643,7 +646,10 @@ void FMacWindow::ApplySizeAndModeChanges(int32 X, int32 Y, int32 Width, int32 He
 		}
 	}
 
-	WindowHandle->bZoomed = WindowHandle.zoomed;
+	if(WindowHandle)
+	{
+		WindowHandle->bZoomed = WindowHandle.zoomed;
+	}
 
 	if (FadeReservationToken != 0)
 	{
