Another macOS TIS non-MainThread spam warning fix.

diff --git a/Engine/Source/Runtime/ApplicationCore/Private/Mac/MacApplication.cpp b/Engine/Source/Runtime/ApplicationCore/Private/Mac/MacApplication.cpp
--- a/Engine/Source/Runtime/ApplicationCore/Private/Mac/MacApplication.cpp
+++ b/Engine/Source/Runtime/ApplicationCore/Private/Mac/MacApplication.cpp
@@ -62,6 +62,7 @@ FMacApplication::FMacApplication()
 ,	bEmulatingRightClick(false)
 ,	bIgnoreMouseMoveDelta(0)
 ,	bIsWorkspaceSessionActive(true)
+, 	KeyBoardLayoutData(nil)
 {
 	TextInputMethodSystem = MakeShareable(new FMacTextInputMethodSystem);
 	if (!TextInputMethodSystem->Initialize())
@@ -99,6 +100,8 @@ FMacApplication::FMacApplication()
 		EventMonitor = [NSEvent addLocalMonitorForEventsMatchingMask:NSEventMaskAny handler:^(NSEvent* Event) { return HandleNSEvent(Event); }];
 
 		CGDisplayRegisterReconfigurationCallback(FMacApplication::OnDisplayReconfiguration, this);
+		
+		CacheKeyboardInputSource();
 	}, NSDefaultRunLoopMode, true);
 
 #if WITH_EDITOR
@@ -156,6 +159,12 @@ FMacApplication::~FMacApplication()
 		}
 
 		CGDisplayRemoveReconfigurationCallback(FMacApplication::OnDisplayReconfiguration, this);
+		
+		if(KeyBoardLayoutData != nil)
+		{
+			[KeyBoardLayoutData release];
+			KeyBoardLayoutData = nil;
+		}
 	}, NSDefaultRunLoopMode, true);
 
 	if (TextInputMethodSystem.IsValid())
@@ -1756,8 +1765,40 @@ TCHAR FMacApplication::ConvertChar(TCHAR Character) const
 	}
 }
 
+void FMacApplication::CacheKeyboardInputSource()
+{
+	// Cocoa main thread only
+	check([NSThread isMainThread]);
+	
+	if(KeyBoardLayoutData != nil)
+	{
+		[KeyBoardLayoutData release];
+		KeyBoardLayoutData = nil;
+	}
+	
+	TISInputSourceRef CurrentKeyboard = TISCopyCurrentKeyboardLayoutInputSource();
+	if (CurrentKeyboard)
+	{
+		CFDataRef CurrentLayoutData = (CFDataRef)TISGetInputSourceProperty(CurrentKeyboard, kTISPropertyUnicodeKeyLayoutData);
+		if(CurrentLayoutData)
+		{
+			const void* Bytes = CFDataGetBytePtr(CurrentLayoutData);
+			size_t DataLength = CFDataGetLength(CurrentLayoutData);
+			
+			if(Bytes && DataLength > 0)
+			{
+				KeyBoardLayoutData = [[NSData alloc] initWithBytes:Bytes length:DataLength];
+			}
+		}
+		
+		CFRelease(CurrentKeyboard);
+	}
+}
+
 unichar FMacApplication::TranslateKeyCodeToUniCode(uint32 KeyCode, uint32 Modifier)
 {
+	// Any thread allowed
+	
 	// Some just don't work as expected
 	switch(KeyCode)
 	{
@@ -1779,27 +1820,19 @@ unichar FMacApplication::TranslateKeyCodeToUniCode(uint32 KeyCode, uint32 Modifi
 		case kVK_F12: 		return NSF12FunctionKey;
 		default:
 		{
-			// Everything else - use the system unicode function
-			TISInputSourceRef CurrentKeyboard = TISCopyCurrentKeyboardLayoutInputSource();
-			if (CurrentKeyboard)
+			if (MacApplication && MacApplication->KeyBoardLayoutData != nil)
 			{
-				CFDataRef CurrentLayoutData = (CFDataRef)TISGetInputSourceProperty(CurrentKeyboard, kTISPropertyUnicodeKeyLayoutData);
-				CFRelease(CurrentKeyboard);
-
-				if (CurrentLayoutData)
+				const UCKeyboardLayout *KeyboardLayout = (UCKeyboardLayout*)MacApplication->KeyBoardLayoutData.bytes;
+				if (KeyboardLayout)
 				{
-					const UCKeyboardLayout *KeyboardLayout = (UCKeyboardLayout*)CFDataGetBytePtr(CurrentLayoutData);
-					if (KeyboardLayout)
-					{
-						UniChar Buffer[256] = { 0 };
-						UniCharCount BufferLength = 256;
-						uint32 DeadKeyState = 0;
+					UniChar Buffer[256] = { 0 };
+					UniCharCount BufferLength = 256;
+					uint32 DeadKeyState = 0;
 
-						OSStatus Status = UCKeyTranslate(KeyboardLayout, KeyCode, kUCKeyActionDown, ((Modifier) >> 8) & 0xFF, LMGetKbdType(), kUCKeyTranslateNoDeadKeysMask, &DeadKeyState, BufferLength, &BufferLength, Buffer);
-						if (Status == noErr)
-						{
-							return Buffer[0];
-						}
+					OSStatus Status = UCKeyTranslate(KeyboardLayout, KeyCode, kUCKeyActionDown, ((Modifier) >> 8) & 0xFF, LMGetKbdType(), kUCKeyTranslateNoDeadKeysMask, &DeadKeyState, BufferLength, &BufferLength, Buffer);
+					if (Status == noErr)
+					{
+						return Buffer[0];
 					}
 				}
 			}
diff --git a/Engine/Source/Runtime/ApplicationCore/Public/Mac/MacApplication.h b/Engine/Source/Runtime/ApplicationCore/Public/Mac/MacApplication.h
--- a/Engine/Source/Runtime/ApplicationCore/Public/Mac/MacApplication.h
+++ b/Engine/Source/Runtime/ApplicationCore/Public/Mac/MacApplication.h
@@ -286,6 +286,8 @@ private:
 	void OnApplicationWillResignActive();
 	void OnWindowsReordered();
 	void OnActiveSpaceDidChange();
+	
+	void CacheKeyboardInputSource();
 
 	void ConditionallyUpdateModifierKeys(const FDeferredMacEvent& Event);
 	void HandleModifierChange(NSUInteger NewModifierFlags, NSUInteger FlagsShift, NSUInteger UE4Shift, EMacModifierKeys TranslatedCode);
@@ -383,6 +385,8 @@ private:
 	int32 GestureUsage[(int32)EGestureEvent::Count];
 #endif
 
+	NSData* KeyBoardLayoutData;
+
 	friend class FMacWindow;
 };
 
