Don't allow invalid cursor clip rects when locking the cursor on Mac

diff --git a/Engine/Source/Runtime/ApplicationCore/Private/Mac/MacCursor.cpp b/Engine/Source/Runtime/ApplicationCore/Private/Mac/MacCursor.cpp
--- a/Engine/Source/Runtime/ApplicationCore/Private/Mac/MacCursor.cpp
+++ b/Engine/Source/Runtime/ApplicationCore/Private/Mac/MacCursor.cpp
@@ -287,8 +287,8 @@ void FMacCursor::Lock(const RECT* const Bounds)
 	{
 		CursorClipRect.Min.X = FMath::TruncToInt(Bounds->left);
 		CursorClipRect.Min.Y = FMath::TruncToInt(Bounds->top);
-		CursorClipRect.Max.X = FMath::TruncToInt(Bounds->right) - 1;
-		CursorClipRect.Max.Y = FMath::TruncToInt(Bounds->bottom) - 1;
+		CursorClipRect.Max.X = Bounds->right > Bounds->left ? FMath::TruncToInt(Bounds->right) - 1 : Bounds->left;
+		CursorClipRect.Max.Y = Bounds->bottom > Bounds->top ? FMath::TruncToInt(Bounds->bottom) - 1 : Bounds->top;
 	}
 
 	MacApplication->OnCursorLock();
