Prevent alt-f4 during slow tasks as it causes entry into the shutdown sequence at abnormal times

diff --git a/Engine/Source/Runtime/ApplicationCore/Private/Windows/WindowsApplication.cpp b/Engine/Source/Runtime/ApplicationCore/Private/Windows/WindowsApplication.cpp
--- a/Engine/Source/Runtime/ApplicationCore/Private/Windows/WindowsApplication.cpp
+++ b/Engine/Source/Runtime/ApplicationCore/Private/Windows/WindowsApplication.cpp
@@ -1,4 +1,4 @@
-﻿// Copyright 1998-2018 Epic Games, Inc. All Rights Reserved.
+// Copyright 1998-2018 Epic Games, Inc. All Rights Reserved.
 
 #include "Windows/WindowsApplication.h"
 #include "Containers/StringConv.h"
@@ -1290,7 +1290,11 @@ int32 FWindowsApplication::ProcessMessage( HWND hwnd, uint32 msg, WPARAM wParam,
 					break;
 				case SC_CLOSE:
 					{
-						DeferMessage(CurrentNativeEventWindowPtr, hwnd, WM_CLOSE, 0, 0);
+						// do not allow Alt-f4 during slow tasks.  This causes entry into the shutdown sequence at abonrmal times which causes crashes.
+						if (!GIsSlowTask)
+						{
+							DeferMessage(CurrentNativeEventWindowPtr, hwnd, WM_CLOSE, 0, 0);
+						}
 						return 1;
 					}
 					break;
