#rnx

diff --git a/Engine/Source/Runtime/ApplicationCore/Private/HTML5/HTML5InputInterface.cpp b/Engine/Source/Runtime/ApplicationCore/Private/HTML5/HTML5InputInterface.cpp
--- a/Engine/Source/Runtime/ApplicationCore/Private/HTML5/HTML5InputInterface.cpp
+++ b/Engine/Source/Runtime/ApplicationCore/Private/HTML5/HTML5InputInterface.cpp
@@ -100,8 +100,11 @@ void FHTML5InputInterface::Tick(float DeltaTime, const SDL_Event& Event,TSharedR
 				const SDL_Keycode KeyCode = KeyEvent.keysym.scancode;
 				const bool bIsRepeated = KeyEvent.repeat != 0;
 
-				// First KeyDown, then KeyChar. This is important, as in-game console ignores first character otherwise
-				MessageHandler->OnKeyDown(KeyCode, KeyEvent.keysym.sym, bIsRepeated);
+				if ( KeyCode != SDL_SCANCODE_F5 ) // UE-58440 -- note: this will be removed when new toolchain comes in
+				{
+					// First KeyDown, then KeyChar. This is important, as in-game console ignores first character otherwise
+					MessageHandler->OnKeyDown(KeyCode, KeyEvent.keysym.sym, bIsRepeated);
+				}
 
 				// Backspace/Return input caught here.
 				// Note that TextInput still seems to get characters messages too but slate seems to not process them.
@@ -120,7 +123,10 @@ void FHTML5InputInterface::Tick(float DeltaTime, const SDL_Event& Event,TSharedR
 				const SDL_Keycode KeyCode = keyEvent.keysym.scancode;
 				const bool IsRepeat = keyEvent.repeat != 0;
 
-				MessageHandler->OnKeyUp( KeyCode, keyEvent.keysym.sym, IsRepeat );
+				if ( KeyCode != SDL_SCANCODE_F5 ) // UE-58440 -- note: this will be removed when new toolchain comes in
+				{
+					MessageHandler->OnKeyUp( KeyCode, keyEvent.keysym.sym, IsRepeat );
+				}
 				UE_LOG(LogHTML5Input, Verbose, TEXT("KeyUp Code:%d"), KeyCode);
 			}
 			break;
