fix for crash on IOS 9 and TVOS 9

diff --git a/Engine/Source/Runtime/ApplicationCore/Private/IOS/IOSPlatformFramePacer.cpp b/Engine/Source/Runtime/ApplicationCore/Private/IOS/IOSPlatformFramePacer.cpp
--- a/Engine/Source/Runtime/ApplicationCore/Private/IOS/IOSPlatformFramePacer.cpp
+++ b/Engine/Source/Runtime/ApplicationCore/Private/IOS/IOSPlatformFramePacer.cpp
@@ -86,9 +86,14 @@ static NSMutableSet<FIOSFramePacerHandler>* Handlers = [NSMutableSet new];
 
 	{
 		FScopeLock Lock(&HandlersMutex);
-		CADisplayLink* displayLink = (CADisplayLink*)param;
-		double OutputSeconds = displayLink.targetTimestamp;
-		double OutputDuration = displayLink.duration;
+		double OutputSeconds = 0;
+		double OutputDuration = 0;
+		if (!UE4_TARGET_PRE_IOS10)
+		{
+			CADisplayLink* displayLink = (CADisplayLink*)param;
+			OutputSeconds = displayLink.targetTimestamp;
+			OutputDuration = displayLink.duration;
+		}
 		for (FIOSFramePacerHandler Handler in Handlers)
 		{
 			Handler(0, OutputSeconds, OutputDuration);
