fix for auto-backgrounding not working
fix for incorrect launch images on iPad Pro 9.7 and 10.5

diff --git a/Engine/Source/Runtime/ApplicationCore/Private/IOS/IOSAppDelegate.cpp b/Engine/Source/Runtime/ApplicationCore/Private/IOS/IOSAppDelegate.cpp
--- a/Engine/Source/Runtime/ApplicationCore/Private/IOS/IOSAppDelegate.cpp
+++ b/Engine/Source/Runtime/ApplicationCore/Private/IOS/IOSAppDelegate.cpp
@@ -298,15 +298,11 @@ void EngineCrashHandler(const FGenericCrashContext& GenericContext)
 {
 	float TimerDuration = 0.0F;
 	GConfig->GetFloat(TEXT("/Script/IOSRuntimeSettings.IOSRuntimeSettings"), TEXT("IdleTimerEnablePeriod"), TimerDuration, GEngineIni);
-    if (TimerDuration > 0.f)
-    {
-        IdleTimerEnablePeriod = TimerDuration;
-        self.IdleTimerEnableTimer = nil;
-    }
-    else
-    {
-        [self EnableIdleTimer: NO];
-    }
+    IdleTimerEnablePeriod = TimerDuration;
+	self.IdleTimerEnableTimer = nil;
+	bool bEnableTimer = YES;
+	GConfig->GetBool(TEXT("/Script/IOSRuntimeSettings.IOSRuntimeSettings"), TEXT("bEnableIdleTimer"), bEnableTimer, GEngineIni);
+	[self EnableIdleTimer : bEnableTimer];
 }
 
 -(void)DeferredEnableIdleTimer
@@ -736,7 +732,7 @@ bool GIsSuspended = 0;
 		// use IPhone6 image for now
 		[ImageString appendString : @"-IPhone6Plus-Landscape"];
 	}
-	else if (Device == FPlatformMisc::IOS_IPadPro_129 || Device == FPlatformMisc::IOS_IPadPro2_129)
+	else if (Device == FPlatformMisc::IOS_IPadPro_129 || Device == FPlatformMisc::IOS_IPadPro2_129 || Device == FPlatformMisc::IOS_IPadPro_97 || Device == FPlatformMisc::IOS_IPadPro_105)
 	{
 		if (!self.bDeviceInPortraitMode)
 		{
