[UE - 58463] Thermals are Nominal message appears on DebugGame and Development iOS packages

diff --git a/Engine/Source/Runtime/ApplicationCore/Private/IOS/IOSAppDelegate.cpp b/Engine/Source/Runtime/ApplicationCore/Private/IOS/IOSAppDelegate.cpp
--- a/Engine/Source/Runtime/ApplicationCore/Private/IOS/IOSAppDelegate.cpp
+++ b/Engine/Source/Runtime/ApplicationCore/Private/IOS/IOSAppDelegate.cpp
@@ -23,6 +23,7 @@
 
 #include <AudioToolbox/AudioToolbox.h>
 #include <AVFoundation/AVAudioSession.h>
+#include "HAL/IConsoleManager.h"
 
 // this is the size of the game thread stack, it must be a multiple of 4k
 #if (UE_BUILD_SHIPPING || UE_BUILD_TEST)
@@ -647,6 +648,13 @@ bool GIsSuspended = 0;
 	return YES;
 }
 
+static int32 GEnableThermalsReport = 0;
+static FAutoConsoleVariableRef CVarGEnableThermalsReport(
+	TEXT("ios.EnableThermalsReport"),
+	GEnableThermalsReport,
+	TEXT("When set to 1, will enable on-screen thermals debug display.")
+);
+
 - (BOOL)application:(UIApplication *)application didFinishLaunchingWithOptions:(NSDictionary*)launchOptions
 {
 	// save launch options
@@ -850,24 +858,26 @@ bool GIsSuspended = 0;
 	if (@available(iOS 11, *))
 	{
 		FCoreDelegates::OnGetOnScreenMessages.AddLambda(
-			[](TMultiMap<FCoreDelegates::EOnScreenMessageSeverity, FText >& OutMessages)
+			[&EnableThermalsReport = GEnableThermalsReport](TMultiMap<FCoreDelegates::EOnScreenMessageSeverity, FText >& OutMessages)
 			{
-				switch ([[NSProcessInfo processInfo] thermalState])
+				if (EnableThermalsReport)
 				{
-					case NSProcessInfoThermalStateNominal:	OutMessages.Add(FCoreDelegates::EOnScreenMessageSeverity::Info, FText::FromString(TEXT("Thermals are Nominal"))); break;
-					case NSProcessInfoThermalStateFair:		OutMessages.Add(FCoreDelegates::EOnScreenMessageSeverity::Info, FText::FromString(TEXT("Thermals are Fair"))); break;
-					case NSProcessInfoThermalStateSerious:	OutMessages.Add(FCoreDelegates::EOnScreenMessageSeverity::Warning, FText::FromString(TEXT("Thermals are Serious"))); break;
-					case NSProcessInfoThermalStateCritical:	OutMessages.Add(FCoreDelegates::EOnScreenMessageSeverity::Error, FText::FromString(TEXT("Thermals are Critical"))); break;
+					switch ([[NSProcessInfo processInfo] thermalState])
+					{
+						case NSProcessInfoThermalStateNominal:	OutMessages.Add(FCoreDelegates::EOnScreenMessageSeverity::Info, FText::FromString(TEXT("Thermals are Nominal"))); break;
+						case NSProcessInfoThermalStateFair:		OutMessages.Add(FCoreDelegates::EOnScreenMessageSeverity::Info, FText::FromString(TEXT("Thermals are Fair"))); break;
+						case NSProcessInfoThermalStateSerious:	OutMessages.Add(FCoreDelegates::EOnScreenMessageSeverity::Warning, FText::FromString(TEXT("Thermals are Serious"))); break;
+						case NSProcessInfoThermalStateCritical:	OutMessages.Add(FCoreDelegates::EOnScreenMessageSeverity::Error, FText::FromString(TEXT("Thermals are Critical"))); break;
+					}
 				}
 
 				// Uncomment to view the state of the AVAudioSession category, mode, and options.
 //#define VIEW_AVAUDIOSESSION_INFO
 #if defined(VIEW_AVAUDIOSESSION_INFO)
 				FString Message = FString::Printf(TEXT("Session Category: %s, Mode: %s, Options: %x"), UTF8_TO_TCHAR([[AVAudioSession sharedInstance].category UTF8String]), UTF8_TO_TCHAR([[AVAudioSession sharedInstance].mode UTF8String]),
-												  [AVAudioSession sharedInstance].categoryOptions);
+												[AVAudioSession sharedInstance].categoryOptions);
 				OutMessages.Add(FCoreDelegates::EOnScreenMessageSeverity::Info, FText::FromString(Message));
 #endif // defined(VIEW_AVAUDIOSESSION_INFO)
-
 			});
 	}
 
