Fix crash in headless client

diff --git a/Engine/Source/Runtime/ApplicationCore/Private/Linux/LinuxCursor.cpp b/Engine/Source/Runtime/ApplicationCore/Private/Linux/LinuxCursor.cpp
--- a/Engine/Source/Runtime/ApplicationCore/Private/Linux/LinuxCursor.cpp
+++ b/Engine/Source/Runtime/ApplicationCore/Private/Linux/LinuxCursor.cpp
@@ -280,6 +280,11 @@ void FLinuxCursor::Show( bool bShow )
 
 void FLinuxCursor::Lock( const RECT* const Bounds )
 {
+	if (!LinuxApplication)
+	{
+		return;
+	}
+
 	TSharedPtr< FLinuxWindow > CurrentFocusWindow = LinuxApplication->GetCurrentFocusWindow();
 	if(!CurrentFocusWindow.IsValid())
 	{
