UE-58081 - ARKit face AR - native orientation is Landscape, needs to match app

diff --git a/Engine/Source/Runtime/ApplicationCore/Private/IOS/IOSApplication.cpp b/Engine/Source/Runtime/ApplicationCore/Private/IOS/IOSApplication.cpp
--- a/Engine/Source/Runtime/ApplicationCore/Private/IOS/IOSApplication.cpp
+++ b/Engine/Source/Runtime/ApplicationCore/Private/IOS/IOSApplication.cpp
@@ -56,6 +56,9 @@ void FIOSApplication::AddExternalInputDevice(TSharedPtr<IInputDevice> InputDevic
 	}
 }
 
+UIDeviceOrientation CachedOrientation = UIDeviceOrientationPortrait;
+UIEdgeInsets CachedInsets;
+
 void FIOSApplication::PollGameDeviceState( const float TimeDelta )
 {
 	// initialize any externally-implemented input devices (we delay load initialize the array so any plugins have had time to load)
@@ -90,6 +93,7 @@ void FIOSApplication::PollGameDeviceState( const float TimeDelta )
 
 		GenericApplication::GetMessageHandler()->OnSizeChanged(Windows[0],WindowWidth,WindowHeight, false);
 		GenericApplication::GetMessageHandler()->OnResizingWindow(Windows[0]);
+		CacheDisplayMetrics();
 		FDisplayMetrics DisplayMetrics;
 		FDisplayMetrics::GetDisplayMetrics(DisplayMetrics);
 		BroadcastDisplayMetricsChanged(DisplayMetrics);
@@ -123,9 +127,6 @@ void FDisplayMetrics::GetDisplayMetrics(FDisplayMetrics& OutDisplayMetrics)
 	{
 		const float RequestedContentScaleFactor = [[IOSAppDelegate GetDelegate].IOSView contentScaleFactor];
 
-		UIEdgeInsets Insets = [[[[UIApplication sharedApplication] delegate] window] safeAreaInsets];
-		UIInterfaceOrientation Orientation = [[UIApplication sharedApplication] statusBarOrientation];
-		
 		//we need to set these according to the orientation
 		TAutoConsoleVariable<float>* CVar_Left = nullptr;
 		TAutoConsoleVariable<float>* CVar_Top = &CVarSafeZone_Landscape_Top;
@@ -133,22 +134,22 @@ void FDisplayMetrics::GetDisplayMetrics(FDisplayMetrics& OutDisplayMetrics)
 		TAutoConsoleVariable<float>* CVar_Bottom = &CVarSafeZone_Landscape_Bottom;
 
 		//making an assumption that the "normal" landscape mode is Landscape right
-		if (Orientation == UIInterfaceOrientationLandscapeLeft)
+		if (CachedOrientation == UIDeviceOrientationLandscapeLeft)
 		{
 			CVar_Left = &CVarSafeZone_Landscape_Left;
 			CVar_Right = &CVarSafeZone_Landscape_Right;
 		}
-		else if (Orientation == UIInterfaceOrientationLandscapeRight)
+		else if (CachedOrientation == UIDeviceOrientationLandscapeRight)
 		{
 			CVar_Left = &CVarSafeZone_Landscape_Right;
 			CVar_Right = &CVarSafeZone_Landscape_Left;
 		}
 
 		// of the CVars are set, use their values. If not, use what comes from iOS
-		const float Inset_Left = (!CVar_Left || CVar_Left->AsVariable()->GetFloat() < 0.0f) ? Insets.left : CVar_Left->AsVariable()->GetFloat();
-		const float Inset_Top = (!CVar_Top || CVar_Top->AsVariable()->GetFloat() < 0.0f) ? Insets.top : CVar_Top->AsVariable()->GetFloat();
-		const float Inset_Right = (!CVar_Right || CVar_Right->AsVariable()->GetFloat() < 0.0f) ? Insets.right : CVar_Right->AsVariable()->GetFloat();
-		const float Inset_Bottom = (!CVar_Bottom || CVar_Bottom->AsVariable()->GetFloat() < 0.0f) ? Insets.bottom : CVar_Bottom->AsVariable()->GetFloat();
+		const float Inset_Left = (!CVar_Left || CVar_Left->AsVariable()->GetFloat() < 0.0f) ? CachedInsets.left : CVar_Left->AsVariable()->GetFloat();
+		const float Inset_Top = (!CVar_Top || CVar_Top->AsVariable()->GetFloat() < 0.0f) ? CachedInsets.top : CVar_Top->AsVariable()->GetFloat();
+		const float Inset_Right = (!CVar_Right || CVar_Right->AsVariable()->GetFloat() < 0.0f) ? CachedInsets.right : CVar_Right->AsVariable()->GetFloat();
+		const float Inset_Bottom = (!CVar_Bottom || CVar_Bottom->AsVariable()->GetFloat() < 0.0f) ? CachedInsets.bottom : CVar_Bottom->AsVariable()->GetFloat();
 
 		//setup the asymmetrical padding
 		OutDisplayMetrics.TitleSafePaddingSize.X = Inset_Left;
@@ -168,6 +169,18 @@ void FDisplayMetrics::GetDisplayMetrics(FDisplayMetrics& OutDisplayMetrics)
 	}
 }
 
+void FIOSApplication::CacheDisplayMetrics()
+	{
+
+#if !PLATFORM_TVOS
+	if (@available(iOS 11, *))
+	{
+		CachedInsets = [[[[UIApplication sharedApplication] delegate] window] safeAreaInsets];
+		CachedOrientation = [[UIDevice currentDevice] orientation];
+	}
+#endif
+}
+
 TSharedRef< FGenericWindow > FIOSApplication::MakeWindow()
 {
 	return FIOSWindow::Make();
diff --git a/Engine/Source/Runtime/ApplicationCore/Public/IOS/IOSApplication.h b/Engine/Source/Runtime/ApplicationCore/Public/IOS/IOSApplication.h
--- a/Engine/Source/Runtime/ApplicationCore/Public/IOS/IOSApplication.h
+++ b/Engine/Source/Runtime/ApplicationCore/Public/IOS/IOSApplication.h
@@ -53,4 +53,6 @@ private:
 
 	static FCriticalSection CriticalSection;
 	static bool bOrientationChanged;
+
+	void CacheDisplayMetrics();
 };
