fix for TVOS build failure

diff --git a/Engine/Source/Runtime/ApplicationCore/Private/IOS/IOSApplication.cpp b/Engine/Source/Runtime/ApplicationCore/Private/IOS/IOSApplication.cpp
--- a/Engine/Source/Runtime/ApplicationCore/Private/IOS/IOSApplication.cpp
+++ b/Engine/Source/Runtime/ApplicationCore/Private/IOS/IOSApplication.cpp
@@ -56,9 +56,6 @@ void FIOSApplication::AddExternalInputDevice(TSharedPtr<IInputDevice> InputDevic
 	}
 }
 
-UIDeviceOrientation CachedOrientation = UIDeviceOrientationPortrait;
-UIEdgeInsets CachedInsets;
-
 void FIOSApplication::PollGameDeviceState( const float TimeDelta )
 {
 	// initialize any externally-implemented input devices (we delay load initialize the array so any plugins have had time to load)
@@ -112,6 +109,11 @@ static TAutoConsoleVariable<float> CVarSafeZone_Landscape_Top(TEXT("SafeZone.Lan
 static TAutoConsoleVariable<float> CVarSafeZone_Landscape_Right(TEXT("SafeZone.Landscape.Right"), -1.0f, TEXT("Safe Zone - Landscape - Right"));
 static TAutoConsoleVariable<float> CVarSafeZone_Landscape_Bottom(TEXT("SafeZone.Landscape.Bottom"), -1.0f, TEXT("Safe Zone - Landscape - Bottom"));
 
+#if !PLATFORM_TVOS
+UIDeviceOrientation CachedOrientation = UIDeviceOrientationPortrait;
+UIEdgeInsets CachedInsets;
+#endif
+
 void FDisplayMetrics::GetDisplayMetrics(FDisplayMetrics& OutDisplayMetrics)
 {
 	// Get screen rect
