Don't call AppKit functions on anything but the main-thread - as of 10.14 this is now a hard error.

diff --git a/Engine/Source/Runtime/ApplicationCore/Private/Mac/MacApplication.cpp b/Engine/Source/Runtime/ApplicationCore/Private/Mac/MacApplication.cpp
--- a/Engine/Source/Runtime/ApplicationCore/Private/Mac/MacApplication.cpp
+++ b/Engine/Source/Runtime/ApplicationCore/Private/Mac/MacApplication.cpp
@@ -735,8 +735,10 @@ void FMacApplication::ProcessMouseMovedEvent(const FDeferredMacEvent& Event, TSh
 		const EWindowZone::Type Zone = GetCurrentWindowZone(EventWindow.ToSharedRef());
 		bool IsMouseOverTitleBar = (Zone == EWindowZone::TitleBar);
 		const bool IsMovable = IsMouseOverTitleBar || IsEdgeZone(Zone);
-		[EventWindow->GetWindowHandle() setMovable:IsMovable];
-		[EventWindow->GetWindowHandle() setMovableByWindowBackground:IsMouseOverTitleBar];
+		MainThreadCall(^{
+			[EventWindow->GetWindowHandle() setMovable:IsMovable];
+			[EventWindow->GetWindowHandle() setMovableByWindowBackground:IsMouseOverTitleBar];
+		}, NSDefaultRunLoopMode, false);
 	}
 
 	FMacCursor* MacCursor = (FMacCursor*)Cursor.Get();
