remove the IPadPro9.7 and IPadPro10.5 from the IPadPro12.9 splash screen. THey use the IPad retina version
The previous fix was not needed

diff --git a/Engine/Source/Runtime/ApplicationCore/Private/IOS/IOSAppDelegate.cpp b/Engine/Source/Runtime/ApplicationCore/Private/IOS/IOSAppDelegate.cpp
--- a/Engine/Source/Runtime/ApplicationCore/Private/IOS/IOSAppDelegate.cpp
+++ b/Engine/Source/Runtime/ApplicationCore/Private/IOS/IOSAppDelegate.cpp
@@ -740,7 +740,7 @@ static FAutoConsoleVariableRef CVarGEnableThermalsReport(
 		// use IPhone6 image for now
 		[ImageString appendString : @"-IPhone6Plus-Landscape"];
 	}
-	else if (Device == FPlatformMisc::IOS_IPadPro_129 || Device == FPlatformMisc::IOS_IPadPro2_129 || Device == FPlatformMisc::IOS_IPadPro_97 || Device == FPlatformMisc::IOS_IPadPro_105)
+	else if (Device == FPlatformMisc::IOS_IPadPro_129 || Device == FPlatformMisc::IOS_IPadPro2_129)
 	{
 		if (!self.bDeviceInPortraitMode)
 		{
