Change the order of calls in FMacWindow::Show() to avoid deadlocks when showing certain modal dialog boxes.

diff --git a/Engine/Source/Runtime/ApplicationCore/Private/Mac/MacWindow.cpp b/Engine/Source/Runtime/ApplicationCore/Private/Mac/MacWindow.cpp
--- a/Engine/Source/Runtime/ApplicationCore/Private/Mac/MacWindow.cpp
+++ b/Engine/Source/Runtime/ApplicationCore/Private/Mac/MacWindow.cpp
@@ -332,11 +332,6 @@ void FMacWindow::Show()
 
 		bIsFirstTimeVisible = false;
 
-		MainThreadCall(^{
-			SCOPED_AUTORELEASE_POOL;
-			[WindowHandle orderFrontAndMakeMain:bShouldActivate andKey:bShouldActivate];
-		}, UE4ShowEventMode, true);
-
 		if (bShouldActivate)
 		{
 			// Tell MacApplication to send window deactivate and activate messages to Slate without waiting for Cocoa events.
@@ -347,6 +342,12 @@ void FMacWindow::Show()
 			MacApplication->OnWindowOrderedFront(SharedThis(this));
 		}
 
+		FCocoaWindow* WindowHandleCopy = WindowHandle;
+		MainThreadCall(^{
+			SCOPED_AUTORELEASE_POOL;
+			[WindowHandleCopy orderFrontAndMakeMain:bShouldActivate andKey:bShouldActivate];
+		}, UE4ShowEventMode, true);
+
 		bIsVisible = true;
 	}
 }
