Fixed a crash introduced with CL 4183733 - TSharedPtr doesn't work well with ObjC blocks

diff --git a/Engine/Source/Runtime/ApplicationCore/Private/Mac/MacApplication.cpp b/Engine/Source/Runtime/ApplicationCore/Private/Mac/MacApplication.cpp
--- a/Engine/Source/Runtime/ApplicationCore/Private/Mac/MacApplication.cpp
+++ b/Engine/Source/Runtime/ApplicationCore/Private/Mac/MacApplication.cpp
@@ -735,9 +735,10 @@ void FMacApplication::ProcessMouseMovedEvent(const FDeferredMacEvent& Event, TSh
 		const EWindowZone::Type Zone = GetCurrentWindowZone(EventWindow.ToSharedRef());
 		bool IsMouseOverTitleBar = (Zone == EWindowZone::TitleBar);
 		const bool IsMovable = IsMouseOverTitleBar || IsEdgeZone(Zone);
+		FCocoaWindow* WindowHandle = EventWindow->GetWindowHandle();
 		MainThreadCall(^{
-			[EventWindow->GetWindowHandle() setMovable:IsMovable];
-			[EventWindow->GetWindowHandle() setMovableByWindowBackground:IsMouseOverTitleBar];
+			[WindowHandle setMovable:IsMovable];
+			[WindowHandle setMovableByWindowBackground:IsMouseOverTitleBar];
 		}, NSDefaultRunLoopMode, false);
 	}
 
