Dont Confine the cursor if the w/h bounding box is less then zero

Same issue as UE-59111

diff --git a/Engine/Source/Runtime/ApplicationCore/Private/Linux/LinuxCursor.cpp b/Engine/Source/Runtime/ApplicationCore/Private/Linux/LinuxCursor.cpp
--- a/Engine/Source/Runtime/ApplicationCore/Private/Linux/LinuxCursor.cpp
+++ b/Engine/Source/Runtime/ApplicationCore/Private/Linux/LinuxCursor.cpp
@@ -311,7 +311,11 @@ void FLinuxCursor::Lock( const RECT* const Bounds )
 		CursorClipRect.w = FMath::TruncToInt(Bounds->right) - FMath::TruncToInt(Bounds->left) - 1 - ExceedingWindowPrevention;
 		CursorClipRect.h = FMath::TruncToInt(Bounds->bottom) - FMath::TruncToInt(Bounds->top) - 1 - ExceedingWindowPrevention;
 
-		SDL_ConfineCursor(CurrentFocusWindow->GetHWnd(), &CursorClipRect);
+		// We dont want to set a negative bounding region. If Top, Left, Bottom, Right are all 0
+		if (CursorClipRect.x >= 0 && CursorClipRect.y >= 0 && CursorClipRect.w > 0 && CursorClipRect.h > 0)
+		{
+			SDL_ConfineCursor(CurrentFocusWindow->GetHWnd(), &CursorClipRect);
+		}
 	}
 }
 
