Dont add colors if we are not outputing to a terminal

diff --git a/Engine/Source/Runtime/ApplicationCore/Private/Linux/LinuxConsoleOutputDevice.cpp b/Engine/Source/Runtime/ApplicationCore/Private/Linux/LinuxConsoleOutputDevice.cpp
--- a/Engine/Source/Runtime/ApplicationCore/Private/Linux/LinuxConsoleOutputDevice.cpp
+++ b/Engine/Source/Runtime/ApplicationCore/Private/Linux/LinuxConsoleOutputDevice.cpp
@@ -48,7 +48,9 @@ void FLinuxConsoleOutputDevice::Serialize(const TCHAR* Data, ELogVerbosity::Type
 		else
 		{
 			bool bNeedToResetColor = false;
-			if (!bOverrideColorSet)
+			bool bOutputtingToTerminal = isatty(STDOUT_FILENO);
+
+			if (bOutputtingToTerminal && !bOverrideColorSet)
 			{
 				if (Verbosity == ELogVerbosity::Error)
 				{
