Only check once if we are outputing to a terminal

diff --git a/Engine/Source/Runtime/ApplicationCore/Private/Linux/LinuxConsoleOutputDevice.cpp b/Engine/Source/Runtime/ApplicationCore/Private/Linux/LinuxConsoleOutputDevice.cpp
--- a/Engine/Source/Runtime/ApplicationCore/Private/Linux/LinuxConsoleOutputDevice.cpp
+++ b/Engine/Source/Runtime/ApplicationCore/Private/Linux/LinuxConsoleOutputDevice.cpp
@@ -19,7 +19,8 @@
 #define CONSOLE_NONE	"\x1b[0m"
 
 FLinuxConsoleOutputDevice::FLinuxConsoleOutputDevice()
-	: bOverrideColorSet(false)
+	: bOverrideColorSet(false),
+	  bOutputtingToTerminal(isatty(STDOUT_FILENO))
 {
 }
 
@@ -48,7 +49,6 @@ void FLinuxConsoleOutputDevice::Serialize(const TCHAR* Data, ELogVerbosity::Type
 		else
 		{
 			bool bNeedToResetColor = false;
-			bool bOutputtingToTerminal = isatty(STDOUT_FILENO);
 
 			if (bOutputtingToTerminal && !bOverrideColorSet)
 			{
diff --git a/Engine/Source/Runtime/ApplicationCore/Public/Linux/LinuxConsoleOutputDevice.h b/Engine/Source/Runtime/ApplicationCore/Public/Linux/LinuxConsoleOutputDevice.h
--- a/Engine/Source/Runtime/ApplicationCore/Public/Linux/LinuxConsoleOutputDevice.h
+++ b/Engine/Source/Runtime/ApplicationCore/Public/Linux/LinuxConsoleOutputDevice.h
@@ -16,6 +16,9 @@ private:
 	/** true if the color is currently set by the caller */
 	bool bOverrideColorSet;
 
+	/** true if we are outputting to a terminal */
+	bool bOutputtingToTerminal;
+
 public:
 
 	/**
