Linux: Implement minimized function for LinuxWindow

diff --git a/Engine/Source/Runtime/ApplicationCore/Private/Linux/LinuxWindow.cpp b/Engine/Source/Runtime/ApplicationCore/Private/Linux/LinuxWindow.cpp
--- a/Engine/Source/Runtime/ApplicationCore/Private/Linux/LinuxWindow.cpp
+++ b/Engine/Source/Runtime/ApplicationCore/Private/Linux/LinuxWindow.cpp
@@ -658,6 +658,12 @@ bool FLinuxWindow::IsVisible() const
 	return bIsVisible;
 }
 
+/** @return true if the native window is minimized, false otherwise */
+bool FLinuxWindow::IsMinimized() const
+{
+	return SDL_GetWindowFlags(HWnd) & SDL_WINDOW_MINIMIZED;
+}
+
 /** Returns the size and location of the window when it is restored */
 bool FLinuxWindow::GetRestoredDimensions(int32& X, int32& Y, int32& Width, int32& Height)
 {
diff --git a/Engine/Source/Runtime/ApplicationCore/Public/Linux/LinuxWindow.h b/Engine/Source/Runtime/ApplicationCore/Public/Linux/LinuxWindow.h
--- a/Engine/Source/Runtime/ApplicationCore/Public/Linux/LinuxWindow.h
+++ b/Engine/Source/Runtime/ApplicationCore/Public/Linux/LinuxWindow.h
@@ -141,6 +141,8 @@ public:
 
 	virtual bool IsVisible() const override;
 
+	virtual bool IsMinimized() const override;
+
 	virtual bool GetRestoredDimensions(int32& X, int32& Y, int32& Width, int32& Height) override;
 
 	virtual void SetWindowFocus() override;
