Linux: Clean up IsMaximized

diff --git a/Engine/Source/Runtime/ApplicationCore/Private/Linux/LinuxWindow.cpp b/Engine/Source/Runtime/ApplicationCore/Private/Linux/LinuxWindow.cpp
--- a/Engine/Source/Runtime/ApplicationCore/Private/Linux/LinuxWindow.cpp
+++ b/Engine/Source/Runtime/ApplicationCore/Private/Linux/LinuxWindow.cpp
@@ -640,16 +640,13 @@ bool FLinuxWindow::GetFullScreenInfo( int32& X, int32& Y, int32& Width, int32& H
 /** @return true if the native window is maximized, false otherwise */
 bool FLinuxWindow::IsMaximized() const
 {
-	uint32 flag = SDL_GetWindowFlags( HWnd );
+	return SDL_GetWindowFlags(HWnd) & SDL_WINDOW_MAXIMIZED;
+}
 
-	if ( flag & SDL_WINDOW_MAXIMIZED )
-	{
-		return true;
-	}
-	else
-	{
-		return false;
-	}
+/** @return true if the native window is minimized, false otherwise */
+bool FLinuxWindow::IsMinimized() const
+{
+	return SDL_GetWindowFlags(HWnd) & SDL_WINDOW_MINIMIZED;
 }
 
 /** @return true if the native window is visible, false otherwise */
@@ -658,12 +655,6 @@ bool FLinuxWindow::IsVisible() const
 	return bIsVisible;
 }
 
-/** @return true if the native window is minimized, false otherwise */
-bool FLinuxWindow::IsMinimized() const
-{
-	return SDL_GetWindowFlags(HWnd) & SDL_WINDOW_MINIMIZED;
-}
-
 /** Returns the size and location of the window when it is restored */
 bool FLinuxWindow::GetRestoredDimensions(int32& X, int32& Y, int32& Width, int32& Height)
 {
