Linux: Make popup menus BORDERLESS. Slate will give the menu events

This appears to fix a lot of our grabs causing compiz to do something issue.

diff --git a/Engine/Source/Runtime/ApplicationCore/Private/Linux/LinuxWindow.cpp b/Engine/Source/Runtime/ApplicationCore/Private/Linux/LinuxWindow.cpp
--- a/Engine/Source/Runtime/ApplicationCore/Private/Linux/LinuxWindow.cpp
+++ b/Engine/Source/Runtime/ApplicationCore/Private/Linux/LinuxWindow.cpp
@@ -171,7 +171,8 @@ void FLinuxWindow::Initialize( FLinuxApplication* const Application, const TShar
 		!Definition->IsModalWindow && !Definition->IsRegularWindow &&
 		bShouldActivate && !Definition->SizeWillChangeOften)
 	{
-		WindowStyle |= SDL_WINDOW_POPUP_MENU;
+		// Popup menus grab the mouse/keyboard which is undesired behaviour. Slate will give the window events.
+		WindowStyle |= SDL_WINDOW_BORDERLESS;
 		bIsPopupWindow = true;
 		UE_LOG(LogLinuxWindowType, Verbose, TEXT("*** New Window is a Popup Menu Window ***"));
 	}
