Linux: Allow windows that want to be resizable to be resizeable

diff --git a/Engine/Source/Runtime/ApplicationCore/Private/Linux/LinuxWindow.cpp b/Engine/Source/Runtime/ApplicationCore/Private/Linux/LinuxWindow.cpp
--- a/Engine/Source/Runtime/ApplicationCore/Private/Linux/LinuxWindow.cpp
+++ b/Engine/Source/Runtime/ApplicationCore/Private/Linux/LinuxWindow.cpp
@@ -122,11 +122,11 @@ void FLinuxWindow::Initialize( FLinuxApplication* const Application, const TShar
 		{
 			WindowStyle |= SDL_WINDOW_SKIP_TASKBAR;
 		}
+	}
 
-		if (Definition->IsRegularWindow && Definition->HasSizingFrame)
-		{
-			WindowStyle |= SDL_WINDOW_RESIZABLE;
-		}
+	if (Definition->IsRegularWindow && Definition->HasSizingFrame)
+	{
+		WindowStyle |= SDL_WINDOW_RESIZABLE;
 	}
 
 	const bool bShouldActivate = Definition->ActivationPolicy != EWindowActivationPolicy::Never;
