Linux: Update the Slate application what the window size will most likely be

As X11 takes a frame to send an Event that a window has had its size changed. This causes things such as the slate renderer to think the window size is different then it actually it. This causes streching of tooltips

diff --git a/Engine/Source/Runtime/ApplicationCore/Private/Linux/LinuxWindow.cpp b/Engine/Source/Runtime/ApplicationCore/Private/Linux/LinuxWindow.cpp
--- a/Engine/Source/Runtime/ApplicationCore/Private/Linux/LinuxWindow.cpp
+++ b/Engine/Source/Runtime/ApplicationCore/Private/Linux/LinuxWindow.cpp
@@ -503,6 +503,19 @@ static void _GetBestFullscreenResolution( SDL_HWindow hWnd, int32 *pWidth, int32
 
 void FLinuxWindow::ReshapeWindow( int32 NewX, int32 NewY, int32 NewWidth, int32 NewHeight )
 {
+	// X11 will take until the next frame to send a SizeChanged event. This means the X11 window
+	// will most likely have resized already by the time we render but the slate renderer will
+	// not have been updated leading to an incorrect frame.
+	//
+	// For now tell the owning application we are going to be this size. When the SizeChanged
+	// event comes through for X11 it'll confirm our size is the request one or resize to what
+	// the WM has forced as the size.
+	TSharedPtr< FLinuxWindow > LinuxWindow = OwningApplication->FindWindowBySDLWindow(HWnd);
+	if ( LinuxWindow )
+	{
+		OwningApplication->GetMessageHandler()->OnResizingWindow( LinuxWindow.ToSharedRef() );
+	}
+
 	switch( WindowMode )
 	{
 		// Fullscreen and WindowedFullscreen both use SDL_WINDOW_FULLSCREEN_DESKTOP now
@@ -541,6 +554,17 @@ void FLinuxWindow::ReshapeWindow( int32 NewX, int32 NewY, int32 NewWidth, int32
 	RegionHeight  = NewHeight;
 	VirtualWidth  = NewWidth;
 	VirtualHeight = NewHeight;
+
+	if ( LinuxWindow )
+	{
+		OwningApplication->GetMessageHandler()->OnSizeChanged(
+			LinuxWindow.ToSharedRef(),
+			VirtualWidth,
+			VirtualHeight,
+			//  bWasMinimized
+			false
+		);
+	}
 }
 
 /** Toggle native window between fullscreen and normal mode */
