Linux: Use Show so we preserve our bIsVisible bool and avoid sending SDL_ShowWindow twice (ie. if its already shown)

diff --git a/Engine/Source/Runtime/ApplicationCore/Private/Linux/LinuxWindow.cpp b/Engine/Source/Runtime/ApplicationCore/Private/Linux/LinuxWindow.cpp
--- a/Engine/Source/Runtime/ApplicationCore/Private/Linux/LinuxWindow.cpp
+++ b/Engine/Source/Runtime/ApplicationCore/Private/Linux/LinuxWindow.cpp
@@ -414,7 +414,7 @@ void FLinuxWindow::BringToFront( bool bForce )
 	}
 	else
 	{
-		SDL_ShowWindow(HWnd);
+		Show();
 	}
 }
 
