minor fixes to merge from Main

#rnx

diff --git a/Engine/Source/Runtime/ApplicationCore/Private/HTML5/HTML5InputInterface.cpp b/Engine/Source/Runtime/ApplicationCore/Private/HTML5/HTML5InputInterface.cpp
--- a/Engine/Source/Runtime/ApplicationCore/Private/HTML5/HTML5InputInterface.cpp
+++ b/Engine/Source/Runtime/ApplicationCore/Private/HTML5/HTML5InputInterface.cpp
@@ -100,9 +100,7 @@ void FHTML5InputInterface::Tick(float DeltaTime, const SDL_Event& Event,TSharedR
 				const SDL_Keycode KeyCode = KeyEvent.keysym.scancode;
 				const bool bIsRepeated = KeyEvent.repeat != 0;
 
-				// UE-58440 -- note: this will be removed when new toolchain comes in
-				// UE-54056 -- filtering out Windows/Super key
-				if ( KeyCode != 227 && KeyCode != 231 && KeyCode != SDL_SCANCODE_F5) 
+				if ( KeyCode != SDL_SCANCODE_F5 ) // UE-58440 -- note: this will be removed when new toolchain comes in
 				{
 					// First KeyDown, then KeyChar. This is important, as in-game console ignores first character otherwise
 					MessageHandler->OnKeyDown(KeyCode, KeyEvent.keysym.sym, bIsRepeated);
diff --git a/Engine/Source/Runtime/PakFile/Private/IPlatformFilePak.cpp b/Engine/Source/Runtime/PakFile/Private/IPlatformFilePak.cpp
--- a/Engine/Source/Runtime/PakFile/Private/IPlatformFilePak.cpp
+++ b/Engine/Source/Runtime/PakFile/Private/IPlatformFilePak.cpp
@@ -3905,10 +3905,14 @@ void FPakFile::Initialize(FArchive* Reader)
 
 	if (CachedTotalSize < Info.GetSerializedSize())
 	{
-		if (CachedTotalSize) // UEMOB-425: can be zero - only error when not zero
+		if (CachedTotalSize)
 		{
 			UE_LOG(LogPakFile, Fatal, TEXT("Corrupted pak file '%s' (too short). Verify your installation."), *PakFilename);
 		}
+		else
+		{
+			UE_LOG(LogPakFile, Warning, TEXT("Pak file '%s' is ZERO (expected %d) -- Prehaps Out-Of-Memory?"), *PakFilename, Info.GetSerializedSize());
+		}
 	}
 	else
 	{
