Add headers which can be included before and after Apple system headers, to avoid errors due to mach-o/loader.h redefining PLATFORM_IOS and PLATFORM_TVOS.

diff --git a/Engine/Source/Editor/UnrealEd/Private/SourceCodeNavigation.cpp b/Engine/Source/Editor/UnrealEd/Private/SourceCodeNavigation.cpp
--- a/Engine/Source/Editor/UnrealEd/Private/SourceCodeNavigation.cpp
+++ b/Engine/Source/Editor/UnrealEd/Private/SourceCodeNavigation.cpp
@@ -35,10 +35,12 @@
 	#include <psapi.h>
 #include "Windows/HideWindowsPlatformTypes.h"
 #elif PLATFORM_MAC
+#include "Apple/PreAppleSystemHeaders.h"
 #include <mach-o/dyld.h>
 #include <mach-o/nlist.h>
 #include <mach-o/stab.h>
 #include <cxxabi.h>
+#include "Apple/PostAppleSystemHeaders.h"
 #include "Apple/ApplePlatformSymbolication.h"
 #endif
 #include "Framework/Notifications/NotificationManager.h"
diff --git a/Engine/Source/Runtime/ApplicationCore/Private/Mac/MacPlatformApplicationMisc.cpp b/Engine/Source/Runtime/ApplicationCore/Private/Mac/MacPlatformApplicationMisc.cpp
--- a/Engine/Source/Runtime/ApplicationCore/Private/Mac/MacPlatformApplicationMisc.cpp
+++ b/Engine/Source/Runtime/ApplicationCore/Private/Mac/MacPlatformApplicationMisc.cpp
@@ -18,6 +18,7 @@
 #include "Mac/CocoaThread.h"
 #include "Mac/MacPlatformOutputDevices.h"
 
+#include "Apple/PreAppleSystemHeaders.h"
 #include <dlfcn.h>
 #include <IOKit/IOKitLib.h>
 #include <IOKit/kext/KextManager.h>
@@ -31,6 +32,7 @@
 #include <libproc.h>
 #include <notify.h>
 #include <uuid/uuid.h>
+#include "Apple/PostAppleSystemHeaders.h"
 
 extern FMacMallocCrashHandler* GCrashMalloc;
 
diff --git a/Engine/Source/Runtime/Core/Private/Apple/ApplePlatformStackWalk.cpp b/Engine/Source/Runtime/Core/Private/Apple/ApplePlatformStackWalk.cpp
--- a/Engine/Source/Runtime/Core/Private/Apple/ApplePlatformStackWalk.cpp
+++ b/Engine/Source/Runtime/Core/Private/Apple/ApplePlatformStackWalk.cpp
@@ -7,10 +7,12 @@
 #include "Apple/ApplePlatformStackWalk.h"
 #include "Apple/ApplePlatformSymbolication.h"
 #include "Containers/StringConv.h"
+#include "Apple/PreAppleSystemHeaders.h"
 #include <execinfo.h>
 #include <dlfcn.h>
 #include <mach-o/dyld.h>
 #include <cxxabi.h>
+#include "Apple/PostAppleSystemHeaders.h"
 #include "CoreGlobals.h"
 
 #if PLATFORM_MAC
diff --git a/Engine/Source/Runtime/Core/Private/IOS/IOSPlatformProcess.cpp b/Engine/Source/Runtime/Core/Private/IOS/IOSPlatformProcess.cpp
--- a/Engine/Source/Runtime/Core/Private/IOS/IOSPlatformProcess.cpp
+++ b/Engine/Source/Runtime/Core/Private/IOS/IOSPlatformProcess.cpp
@@ -8,7 +8,9 @@ IOSPlatformProcess.cpp: iOS implementations of Process functions
 #include "Apple/ApplePlatformRunnableThread.h"
 #include "IOS/IOSAppDelegate.h"
 #include "Misc/CoreDelegates.h"
+#include "Apple/PreAppleSystemHeaders.h"
 #include <mach-o/dyld.h>
+#include "Apple/PostAppleSystemHeaders.h"
 
 // numbers recommended by Apple
 #define GAME_THREAD_PRIORITY 47
diff --git a/Engine/Source/Runtime/Core/Private/Mac/MacCriticalSection.cpp b/Engine/Source/Runtime/Core/Private/Mac/MacCriticalSection.cpp
--- a/Engine/Source/Runtime/Core/Private/Mac/MacCriticalSection.cpp
+++ b/Engine/Source/Runtime/Core/Private/Mac/MacCriticalSection.cpp
@@ -7,8 +7,10 @@
 #include "Mac/MacCriticalSection.h"
 #include "Misc/App.h"
 #include "Misc/DateTime.h"
+#include "Apple/PreAppleSystemHeaders.h"
 #include <mach-o/dyld.h>
 #include <libproc.h>
+#include "Apple/PostAppleSystemHeaders.h"
 
 FMacSystemWideCriticalSection::FMacSystemWideCriticalSection(const FString& InName, FTimespan InTimeout)
 {
diff --git a/Engine/Source/Runtime/Core/Private/Mac/MacPlatformMisc.cpp b/Engine/Source/Runtime/Core/Private/Mac/MacPlatformMisc.cpp
--- a/Engine/Source/Runtime/Core/Private/Mac/MacPlatformMisc.cpp
+++ b/Engine/Source/Runtime/Core/Private/Mac/MacPlatformMisc.cpp
@@ -28,6 +28,7 @@
 #include "Internationalization/Culture.h"
 #include "Modules/ModuleManager.h"
 
+#include "Apple/PreAppleSystemHeaders.h"
 #include <dlfcn.h>
 #include <IOKit/IOKitLib.h>
 #include <IOKit/kext/KextManager.h>
@@ -41,6 +42,7 @@
 #include <libproc.h>
 #include <notify.h>
 #include <uuid/uuid.h>
+#include "Apple/PostAppleSystemHeaders.h"
 
 extern CORE_API bool GIsGPUCrashed;
 /*------------------------------------------------------------------------------
diff --git a/Engine/Source/Runtime/Core/Private/Mac/MacPlatformProcess.cpp b/Engine/Source/Runtime/Core/Private/Mac/MacPlatformProcess.cpp
--- a/Engine/Source/Runtime/Core/Private/Mac/MacPlatformProcess.cpp
+++ b/Engine/Source/Runtime/Core/Private/Mac/MacPlatformProcess.cpp
@@ -11,10 +11,12 @@
 #include "Misc/CoreDelegates.h"
 #include "Misc/Paths.h"
 #include "HAL/FileManager.h"
+#include "Apple/PreAppleSystemHeaders.h"
 #include <mach-o/dyld.h>
 #include <mach/thread_act.h>
 #include <mach/thread_policy.h>
 #include <libproc.h>
+#include "Apple/PostAppleSystemHeaders.h"
 
 void* FMacPlatformProcess::GetDllHandle( const TCHAR* Filename )
 {
diff --git a/Engine/Source/Runtime/Core/Public/Apple/PostAppleSystemHeaders.h b/Engine/Source/Runtime/Core/Public/Apple/PostAppleSystemHeaders.h
new file mode 100644
--- /dev/null
+++ b/Engine/Source/Runtime/Core/Public/Apple/PostAppleSystemHeaders.h
@@ -0,0 +1,5 @@
+// Copyright 1998-2019 Epic Games, Inc. All Rights Reserved.
+
+// Restore macros
+#pragma pop_macro("PLATFORM_IOS")
+#pragma pop_macro("PLATFORM_TVOS")
diff --git a/Engine/Source/Runtime/Core/Public/Apple/PreAppleSystemHeaders.h b/Engine/Source/Runtime/Core/Public/Apple/PreAppleSystemHeaders.h
new file mode 100644
--- /dev/null
+++ b/Engine/Source/Runtime/Core/Public/Apple/PreAppleSystemHeaders.h
@@ -0,0 +1,5 @@
+// Copyright 1998-2019 Epic Games, Inc. All Rights Reserved.
+
+// Save these macros; mach-o/loader.h redefines them
+#pragma push_macro("PLATFORM_IOS")
+#pragma push_macro("PLATFORM_TVOS")
