Linux: fix headless clients not updating level streaming (UE-70496).

- Zero-sized display resulted in zero-sized viewport, which caused UWorld::UpdateLevelStreaming() to not be called (see logic in FViewport::Draw()).

diff --git a/Engine/Source/Runtime/ApplicationCore/Private/Linux/LinuxApplication.cpp b/Engine/Source/Runtime/ApplicationCore/Private/Linux/LinuxApplication.cpp
--- a/Engine/Source/Runtime/ApplicationCore/Private/Linux/LinuxApplication.cpp
+++ b/Engine/Source/Runtime/ApplicationCore/Private/Linux/LinuxApplication.cpp
@@ -1443,15 +1443,42 @@ void FDisplayMetrics::RebuildDisplayMetrics(FDisplayMetrics& OutDisplayMetrics)
 
 	OutDisplayMetrics.MonitorInfo.Empty();
 
-	// exit early if no displays connected
 	if (NumDisplays <= 0)
 	{
-		OutDisplayMetrics.PrimaryDisplayWorkAreaRect = FPlatformRect(0, 0, 0, 0);
-		OutDisplayMetrics.VirtualDisplayRect = OutDisplayMetrics.PrimaryDisplayWorkAreaRect;
-		OutDisplayMetrics.PrimaryDisplayWidth = 0;
-		OutDisplayMetrics.PrimaryDisplayHeight = 0;
+		if (IsRunningDedicatedServer())
+		{
+			// dedicated servers has always been exiting early
+			OutDisplayMetrics.PrimaryDisplayWorkAreaRect = FPlatformRect(0, 0, 0, 0);
+			OutDisplayMetrics.VirtualDisplayRect = OutDisplayMetrics.PrimaryDisplayWorkAreaRect;
+			OutDisplayMetrics.PrimaryDisplayWidth = 0;
+			OutDisplayMetrics.PrimaryDisplayHeight = 0;
 
-		return;
+			return;
+		}
+		else
+		{
+			// headless clients need some plausible values because high level logic depends on viewport sizes not being 0 (see e.g. UnrealClient.cpp)
+			int32 Width = 1920;
+			int32 Height = 1080;
+
+			FMonitorInfo Display;
+			Display.bIsPrimary = true;
+			if (FPlatformApplicationMisc::IsHighDPIAwarenessEnabled())
+			{
+				Display.DPI = 96;
+			}
+			Display.ID = TEXT("fakedisplay");
+			Display.NativeWidth = Width;
+			Display.NativeHeight = Height;
+			Display.DisplayRect = FPlatformRect(0, 0, Width, Height);
+			Display.WorkArea = FPlatformRect(0, 0, Width, Height);
+
+			OutDisplayMetrics.PrimaryDisplayWorkAreaRect = FPlatformRectDisplay.WorkArea;
+			OutDisplayMetrics.VirtualDisplayRect = OutDisplayMetrics.PrimaryDisplayWorkAreaRect;
+			OutDisplayMetrics.PrimaryDisplayWidth = Display.NativeWidth;
+			OutDisplayMetrics.PrimaryDisplayHeight = Display.NativeHeight;
+			OutDisplayMetrics.MonitorInfo.Add(Display);
+		}
 	}
 
 	for (int32 DisplayIdx = 0; DisplayIdx < NumDisplays; ++DisplayIdx)
