Fix Linux build (last minute copy-paste went wrong) (UE-70519).

diff --git a/Engine/Source/Runtime/ApplicationCore/Private/Linux/LinuxApplication.cpp b/Engine/Source/Runtime/ApplicationCore/Private/Linux/LinuxApplication.cpp
--- a/Engine/Source/Runtime/ApplicationCore/Private/Linux/LinuxApplication.cpp
+++ b/Engine/Source/Runtime/ApplicationCore/Private/Linux/LinuxApplication.cpp
@@ -1473,7 +1473,7 @@ void FDisplayMetrics::RebuildDisplayMetrics(FDisplayMetrics& OutDisplayMetrics)
 			Display.DisplayRect = FPlatformRect(0, 0, Width, Height);
 			Display.WorkArea = FPlatformRect(0, 0, Width, Height);
 
-			OutDisplayMetrics.PrimaryDisplayWorkAreaRect = FPlatformRectDisplay.WorkArea;
+			OutDisplayMetrics.PrimaryDisplayWorkAreaRect = Display.WorkArea;
 			OutDisplayMetrics.VirtualDisplayRect = OutDisplayMetrics.PrimaryDisplayWorkAreaRect;
 			OutDisplayMetrics.PrimaryDisplayWidth = Display.NativeWidth;
 			OutDisplayMetrics.PrimaryDisplayHeight = Display.NativeHeight;
