Avoid infinite recursive calls in FMacApplication::CloseQueuedWindows()

diff --git a/Engine/Source/Runtime/ApplicationCore/Private/Mac/MacApplication.cpp b/Engine/Source/Runtime/ApplicationCore/Private/Mac/MacApplication.cpp
--- a/Engine/Source/Runtime/ApplicationCore/Private/Mac/MacApplication.cpp
+++ b/Engine/Source/Runtime/ApplicationCore/Private/Mac/MacApplication.cpp
@@ -1951,16 +1951,20 @@ TCHAR FMacApplication::TranslateCharCode(TCHAR CharCode, uint32 KeyCode) const
 
 void FMacApplication::CloseQueuedWindows()
 {
+	// OnWindowClose may call PumpMessages, which would reenter this function, so make a local copy of SlateWindowsToClose array to avoid infinite recursive calls
+	TArray<TSharedRef<FMacWindow>> LocalWindowsToClose;
+
 	{
 		FScopeLock Lock(&WindowsToCloseMutex);
+		LocalWindowsToClose = SlateWindowsToClose;
+		SlateWindowsToClose.Empty();
+	}
 
-		if (SlateWindowsToClose.Num() > 0)
+	if (LocalWindowsToClose.Num() > 0)
+	{
+		for (TSharedRef<FMacWindow> Window : LocalWindowsToClose)
 		{
-			for (TSharedRef<FMacWindow> Window : SlateWindowsToClose)
-			{
-				MessageHandler->OnWindowClose(Window);
-			}
-			SlateWindowsToClose.Empty();
+			MessageHandler->OnWindowClose(Window);
 		}
 	}
 
