Use FMath::CeilToInt when calculating the resized Mac window width and height to make sure the window size matches its viewport size after changes to SlateRHIRenderer in CL 4969188

diff --git a/Engine/Source/Runtime/ApplicationCore/Private/Mac/MacWindow.cpp b/Engine/Source/Runtime/ApplicationCore/Private/Mac/MacWindow.cpp
--- a/Engine/Source/Runtime/ApplicationCore/Private/Mac/MacWindow.cpp
+++ b/Engine/Source/Runtime/ApplicationCore/Private/Mac/MacWindow.cpp
@@ -623,8 +623,8 @@ void FMacWindow::ApplySizeAndModeChanges(int32 X, int32 Y, int32 Width, int32 He
 			WindowHandle.TargetWindowMode = WindowMode;
 
 			const float DPIScaleFactor = FPlatformApplicationMisc::GetDPIScaleFactorAtPoint(X, Y);
-			Width /= DPIScaleFactor;
-			Height /= DPIScaleFactor;
+			Width = FMath::CeilToInt(Width / DPIScaleFactor);
+			Height = FMath::CeilToInt(Height / DPIScaleFactor);
 
 			const FVector2D CocoaPosition = FMacApplication::ConvertSlatePositionToCocoa(X, Y);
 			NSRect Rect = NSMakeRect(CocoaPosition.X, CocoaPosition.Y - Height + 1, FMath::Max(Width, 0), FMath::Max(Height, 0));
