Manually merging CL#5261568 using Main_To_Release-4.22 due to conflict in RoboMerge

Don't wait indefinitely for a native window in Daydream if the application is exiting.
Fixes a bug that leaves a defunct process when dismissing the dialog to put the phone in the daydream headset preventing the application from being launched without killing it first.

diff --git a/Engine/Source/Runtime/ApplicationCore/ApplicationCore.Build.cs b/Engine/Source/Runtime/ApplicationCore/ApplicationCore.Build.cs
--- a/Engine/Source/Runtime/ApplicationCore/ApplicationCore.Build.cs
+++ b/Engine/Source/Runtime/ApplicationCore/ApplicationCore.Build.cs
@@ -76,7 +76,16 @@ public class ApplicationCore : ModuleRules
 			// export ApplicationCore symbols for embedded Dlls
 			ModuleSymbolVisibility = ModuleRules.SymbolVisibility.VisibileForDll;
 		}
-		if (!Target.bCompileAgainstApplicationCore)
+        else if (Target.Platform == UnrealTargetPlatform.Android)
+        {
+            PrivateDependencyModuleNames.AddRange(
+				new string[] {
+					"Launch"
+				}
+			);
+        }
+
+        if (!Target.bCompileAgainstApplicationCore)
         {
 			throw new System.Exception("ApplicationCore cannot be used when Target.bCompileAgainstApplicationCore = false.");
         }
diff --git a/Engine/Source/Runtime/ApplicationCore/Private/Android/AndroidWindow.cpp b/Engine/Source/Runtime/ApplicationCore/Private/Android/AndroidWindow.cpp
--- a/Engine/Source/Runtime/ApplicationCore/Private/Android/AndroidWindow.cpp
+++ b/Engine/Source/Runtime/ApplicationCore/Private/Android/AndroidWindow.cpp
@@ -2,6 +2,8 @@
 
 #include "Android/AndroidWindow.h"
 #include "Android/AndroidWindowUtils.h"
+#include "Android/AndroidEventManager.h"
+
 #if USE_ANDROID_JNI
 #include <android/native_window.h>
 #include <android/native_window_jni.h>
@@ -120,6 +122,37 @@ void* FAndroidWindow::GetHardwareWindow()
 	return NativeWindow;
 }
 
+void* FAndroidWindow::WaitForHardwareWindow()
+{
+	// Sleep if the hardware window isn't currently available.
+	// The Window may not exist if the activity is pausing/resuming, in which case we make this thread wait
+	// This case will come up frequently as a result of the DON flow in Gvr.
+	// Until the app is fully resumed. It would be nicer if this code respected the lifecycle events
+	// of an android app instead, but all of those events are handled on a separate thread and it would require
+	// significant re-architecturing to do.
+
+	// Before sleeping, we peek into the event manager queue to see if it contains an ON_DESTROY event, 
+	// in which case, we exit the loop to allow the application to exit before a window has been created.
+	// For instance when the user aborts the "Place your phone into thr Daydream headset." screen.
+	// It is not sufficient to check the GIsRequestingExit global variable, as the handler reacting to the APP_EVENT_STATE_ON_DESTROY
+	// may be running in the same thread as this method and therefore lead to a deadlock.
+
+	void* Window = GetHardwareWindow();
+	while (Window == nullptr)
+	{
+#if USE_ANDROID_EVENTS
+		if (GIsRequestingExit || FAppEventManager::GetInstance()->WaitForEventInQueue(EAppEventState::APP_EVENT_STATE_ON_DESTROY, 0.0f))
+		{
+			// Application is shutting down soon, abort the wait and return nullptr
+			return nullptr;
+		}
+#endif
+		FPlatformProcess::Sleep(0.001f);
+		Window = GetHardwareWindow();
+	}
+	return Window;
+}
+
 #if USE_ANDROID_JNI
 extern bool AndroidThunkCpp_IsGearVRApplication();
 #endif
@@ -198,17 +231,8 @@ FPlatformRect FAndroidWindow::GetScreenRect()
 	if (bIsDaydreamApp && Window == NULL)
 	{
 		// Sleep if the hardware window isn't currently available.
-		// The Window may not exist if the activity is pausing/resuming, in which case we make this thread wait
-		// This case will come up frequently as a result of the DON flow in Gvr.
-		// Until the app is fully resumed. It would be nicer if this code respected the lifecycle events
-		// of an android app instead, but all of those events are handled on a separate thread and it would require
-		// significant re-architecturing to do.
 		FPlatformMisc::LowLevelOutputDebugString(TEXT("Waiting for Native window in FAndroidWindow::GetScreenRect"));
-		while (Window == NULL)
-		{
-			FPlatformProcess::Sleep(0.001f);
-			Window = (ANativeWindow*)FAndroidWindow::GetHardwareWindow();
-		}
+		Window = (ANativeWindow*)FAndroidWindow::WaitForHardwareWindow();
 	}
 	//	check(Window != NULL);
 	if (Window == NULL)
diff --git a/Engine/Source/Runtime/ApplicationCore/Public/Android/AndroidWindow.h b/Engine/Source/Runtime/ApplicationCore/Public/Android/AndroidWindow.h
--- a/Engine/Source/Runtime/ApplicationCore/Public/Android/AndroidWindow.h
+++ b/Engine/Source/Runtime/ApplicationCore/Public/Android/AndroidWindow.h
@@ -51,6 +51,11 @@ public:
 	static void* GetHardwareWindow();
 	static void SetHardwareWindow(void* InWindow);
 
+	/** Waits on the current thread for a hardware window and returns it. 
+	 *  May return nullptr if the application is shutting down.
+	 */
+	static void* WaitForHardwareWindow();
+
 protected:
 	/** @return true if the native window is currently in fullscreen mode, false otherwise */
 	virtual EWindowMode::Type GetWindowMode() const override { return EWindowMode::Fullscreen; }
diff --git a/Engine/Source/Runtime/OpenGLDrv/Private/Android/AndroidEGL.cpp b/Engine/Source/Runtime/OpenGLDrv/Private/Android/AndroidEGL.cpp
--- a/Engine/Source/Runtime/OpenGLDrv/Private/Android/AndroidEGL.cpp
+++ b/Engine/Source/Runtime/OpenGLDrv/Private/Android/AndroidEGL.cpp
@@ -668,15 +668,13 @@ void AndroidEGL::InitSurface(bool bUseSmallSurface, bool bCreateWndSurface)
 	{
 		// Sleep if the hardware window isn't currently available.
 		// The Window may not exist if the activity is pausing/resuming, in which case we make this thread wait
-		// This case will come up frequently as a result of the DON flow in Gvr.
-		// Until the app is fully resumed. It would be nicer if this code respected the lifecycle events
-		// of an android app instead, but all of those events are handled on a separate thread and it would require
-		// significant re-architecturing to do.
 		FPlatformMisc::LowLevelOutputDebugString(TEXT("Waiting for Native window in AndroidEGL::InitSurface"));
-		while (window == NULL)
+		window = (ANativeWindow*)FAndroidWindow::WaitForHardwareWindow();
+	
+		if (window == NULL)
 		{
-			FPlatformProcess::Sleep(0.001f);
-			window = (ANativeWindow*)FAndroidWindow::GetHardwareWindow();
+			FPlatformMisc::LowLevelOutputDebugString(TEXT("Aborting AndroidEGL::InitSurface, FAndroidWindow::WaitForHardwareWindow() returned null"));
+			return;
 		}
 	}
 
diff --git a/Engine/Source/Runtime/VulkanRHI/Private/Android/VulkanAndroidPlatform.cpp b/Engine/Source/Runtime/VulkanRHI/Private/Android/VulkanAndroidPlatform.cpp
--- a/Engine/Source/Runtime/VulkanRHI/Private/Android/VulkanAndroidPlatform.cpp
+++ b/Engine/Source/Runtime/VulkanRHI/Private/Android/VulkanAndroidPlatform.cpp
@@ -117,13 +117,16 @@ void FVulkanAndroidPlatform::CreateSurface(void* WindowHandle, VkInstance Instan
 	WindowHandle = FAndroidWindow::GetHardwareWindow();
 	if (WindowHandle == NULL)
 	{
+
 		// Sleep if the hardware window isn't currently available.
 		// The Window may not exist if the activity is pausing/resuming, in which case we make this thread wait
 		FPlatformMisc::LowLevelOutputDebugString(TEXT("Waiting for Native window in FVulkanAndroidPlatform::CreateSurface"));
-		while (WindowHandle == NULL)
+		WindowHandle = FAndroidWindow::WaitForHardwareWindow();
+
+		if (WindowHandle == NULL)
 		{
-			FPlatformProcess::Sleep(0.001f);
-			WindowHandle = FAndroidWindow::GetHardwareWindow();
+			FPlatformMisc::LowLevelOutputDebugString(TEXT("Aborting FVulkanAndroidPlatform::CreateSurface, FAndroidWindow::WaitForHardwareWindow() returned null"));
+			return;
 		}
 	}
 
