Fix compilation on Lumin due to missing AndroidEventManager header.

diff --git a/Engine/Source/Runtime/ApplicationCore/ApplicationCore.Build.cs b/Engine/Source/Runtime/ApplicationCore/ApplicationCore.Build.cs
--- a/Engine/Source/Runtime/ApplicationCore/ApplicationCore.Build.cs
+++ b/Engine/Source/Runtime/ApplicationCore/ApplicationCore.Build.cs
@@ -76,7 +76,7 @@ public class ApplicationCore : ModuleRules
 			// export ApplicationCore symbols for embedded Dlls
 			ModuleSymbolVisibility = ModuleRules.SymbolVisibility.VisibileForDll;
 		}
-        else if (Target.Platform == UnrealTargetPlatform.Android)
+        else if (Target.Platform == UnrealTargetPlatform.Android || Target.Platform == UnrealTargetPlatform.Lumin)
         {
             PrivateDependencyModuleNames.AddRange(
 				new string[] {
