Speculative fix for UE-70172. When closing the log window, stop processing additional log lines as soon DestroyConsole() is called

diff --git a/Engine/Source/Runtime/ApplicationCore/Private/Mac/MacConsoleOutputDevice.cpp b/Engine/Source/Runtime/ApplicationCore/Private/Mac/MacConsoleOutputDevice.cpp
--- a/Engine/Source/Runtime/ApplicationCore/Private/Mac/MacConsoleOutputDevice.cpp
+++ b/Engine/Source/Runtime/ApplicationCore/Private/Mac/MacConsoleOutputDevice.cpp
@@ -115,20 +115,22 @@ void FMacConsoleOutputDevice::DestroyConsole()
 {
 	if (ConsoleHandle)
 	{
+		SaveToINI();
+
+		FMacConsoleWindow* ConsoleWindow = ConsoleHandle;
+		ConsoleHandle = nullptr; // Stop further serialization as soon as possible
+
 		do
 		{
 			FMacPlatformApplicationMisc::PumpMessages( true );
 		} while(OutstandingTasks);
 
-		SaveToINI();
-		
 		MainThreadCall(^{
 			SCOPED_AUTORELEASE_POOL;
 			if( TextViewTextColor )
 				[TextViewTextColor release];
 			
-			[ConsoleHandle close];
-			ConsoleHandle = NULL;
+			[ConsoleWindow close];
 			TextViewTextColor = NULL;
 		}, UE4NilEventMode, true);
 	}
