You've got to retain & release the window under the mouse cursor to avoid a heap-use-after-free bug - the alternative is to hook window deallocation events and clear the pointer but then you'll need a mutex.

diff --git a/Engine/Source/Runtime/ApplicationCore/Private/Mac/MacApplication.cpp b/Engine/Source/Runtime/ApplicationCore/Private/Mac/MacApplication.cpp
--- a/Engine/Source/Runtime/ApplicationCore/Private/Mac/MacApplication.cpp
+++ b/Engine/Source/Runtime/ApplicationCore/Private/Mac/MacApplication.cpp
@@ -104,7 +104,7 @@ FMacApplication::FMacApplication()
 		
 		CacheKeyboardInputSource();
 
-		WindowUnderCursor = FindSlateWindowUnderCursor();
+		WindowUnderCursor = [FindSlateWindowUnderCursor() retain];
 	}, NSDefaultRunLoopMode, true);
 
 #if WITH_EDITOR
@@ -170,6 +170,8 @@ FMacApplication::~FMacApplication()
 			[KeyBoardLayoutData release];
 			KeyBoardLayoutData = nil;
 		}
+		
+		[WindowUnderCursor release];
 	}, NSDefaultRunLoopMode, true);
 
 	if (TextInputMethodSystem.IsValid())
@@ -366,7 +368,12 @@ void FMacApplication::DeferEvent(NSObject* Object)
 {
 	FDeferredMacEvent DeferredEvent;
 
-	WindowUnderCursor = FindSlateWindowUnderCursor();
+	FCocoaWindow* Window = FindSlateWindowUnderCursor();
+	if (WindowUnderCursor != Window)
+	{
+		[WindowUnderCursor release];
+		WindowUnderCursor = [Window retain];
+	}
 
 	if (Object && [Object isKindOfClass:[NSEvent class]])
 	{
