Fix a shadow warning from my 5448877 changes.

diff --git a/Engine/Source/Runtime/ApplicationCore/Private/Mac/MacApplication.cpp b/Engine/Source/Runtime/ApplicationCore/Private/Mac/MacApplication.cpp
--- a/Engine/Source/Runtime/ApplicationCore/Private/Mac/MacApplication.cpp
+++ b/Engine/Source/Runtime/ApplicationCore/Private/Mac/MacApplication.cpp
@@ -368,11 +368,11 @@ void FMacApplication::DeferEvent(NSObject* Object)
 {
 	FDeferredMacEvent DeferredEvent;
 
-	FCocoaWindow* Window = FindSlateWindowUnderCursor();
-	if (WindowUnderCursor != Window)
+	FCocoaWindow* CursorWindow = FindSlateWindowUnderCursor();
+	if (WindowUnderCursor != CursorWindow)
 	{
 		[WindowUnderCursor release];
-		WindowUnderCursor = [Window retain];
+		WindowUnderCursor = [CursorWindow retain];
 	}
 
 	if (Object && [Object isKindOfClass:[NSEvent class]])
