Option to disable motion data polling on iOS

diff --git a/Engine/Config/BaseEngine.ini b/Engine/Config/BaseEngine.ini
--- a/Engine/Config/BaseEngine.ini
+++ b/Engine/Config/BaseEngine.ini
@@ -1870,6 +1870,7 @@ UseFastIntrinsics=False
 ForceFloats=False
 EnableMathOptimisations=True
 MaxShaderLanguageVersion=2
+bDisableMotionData=False
 bEnableAdvertisingIdentifier=True
 
 ; These are the defaults for Android settings, and they need to be in the .ini since UBT reads the .ini settings, without instantiating the class
diff --git a/Engine/Source/Runtime/ApplicationCore/Private/IOS/IOSInputInterface.cpp b/Engine/Source/Runtime/ApplicationCore/Private/IOS/IOSInputInterface.cpp
--- a/Engine/Source/Runtime/ApplicationCore/Private/IOS/IOSInputInterface.cpp
+++ b/Engine/Source/Runtime/ApplicationCore/Private/IOS/IOSInputInterface.cpp
@@ -18,7 +18,6 @@ static TAutoConsoleVariable<float> CVarHapticsKickMedium(TEXT("ios.VibrationHapt
 static TAutoConsoleVariable<float> CVarHapticsKickLight(TEXT("ios.VibrationHapticsKickLightValue"), 0.3f, TEXT("Vibation values higher than this will kick a haptics light Impact"));
 static TAutoConsoleVariable<float> CVarHapticsRest(TEXT("ios.VibrationHapticsRestValue"), 0.2f, TEXT("Vibation values lower than this will allow haptics to Kick again when going over ios.VibrationHapticsKickValue"));
 
-
 //@interface FControllerHelper : NSObject
 //{
 //	FIOSInputInterface
@@ -49,9 +48,10 @@ FIOSInputInterface::FIOSInputInterface( const TSharedRef< FGenericApplicationMes
 #if !PLATFORM_TVOS
 	MotionManager = nil;
 	ReferenceAttitude = nil;
-	bPauseMotion = false;
 #endif
-
+	bPauseMotion = false;
+	GConfig->GetBool(TEXT("/Script/IOSRuntimeSettings.IOSRuntimeSettings"), TEXT("bDisableMotionData"), bPauseMotion, GEngineIni);
+	
 	GConfig->GetBool(TEXT("/Script/IOSRuntimeSettings.IOSRuntimeSettings"), TEXT("bTreatRemoteAsSeparateController"), bTreatRemoteAsSeparateController, GEngineIni);
 	GConfig->GetBool(TEXT("/Script/IOSRuntimeSettings.IOSRuntimeSettings"), TEXT("bAllowRemoteRotation"), bAllowRemoteRotation, GEngineIni);
 	GConfig->GetBool(TEXT("/Script/IOSRuntimeSettings.IOSRuntimeSettings"), TEXT("bUseRemoteAsVirtualJoystick"), bUseRemoteAsVirtualJoystick, GEngineIni);
@@ -315,29 +315,30 @@ void FIOSInputInterface::SendControllerEvents()
 
 	
 #if !PLATFORM_TVOS // @todo tvos: This needs to come from the Microcontroller rotation
-	// Update motion controls.
-	FVector Attitude;
-	FVector RotationRate;
-	FVector Gravity;
-	FVector Acceleration;
+	if (!bPauseMotion)
+	{
+		// Update motion controls.
+		FVector Attitude;
+		FVector RotationRate;
+		FVector Gravity;
+		FVector Acceleration;
 
-	GetMovementData(Attitude, RotationRate, Gravity, Acceleration);
+		GetMovementData(Attitude, RotationRate, Gravity, Acceleration);
 
-	// Fix-up yaw to match directions
-	Attitude.Y = -Attitude.Y;
-	RotationRate.Y = -RotationRate.Y;
+		// Fix-up yaw to match directions
+		Attitude.Y = -Attitude.Y;
+		RotationRate.Y = -RotationRate.Y;
 
-	// munge the vectors based on the orientation
-	ModifyVectorByOrientation(Attitude, true);
-	ModifyVectorByOrientation(RotationRate, true);
-	ModifyVectorByOrientation(Gravity, false);
-	ModifyVectorByOrientation(Acceleration, false);
+		// munge the vectors based on the orientation
+		ModifyVectorByOrientation(Attitude, true);
+		ModifyVectorByOrientation(RotationRate, true);
+		ModifyVectorByOrientation(Gravity, false);
+		ModifyVectorByOrientation(Acceleration, false);
 
-	MessageHandler->OnMotionDetected(Attitude, RotationRate, Gravity, Acceleration, 0);
+		MessageHandler->OnMotionDetected(Attitude, RotationRate, Gravity, Acceleration, 0);
+	}
 #endif
-	
-	
-	
+		
 	for (GCController* Cont in [GCController controllers])
  	{
 		GCGamepad* Gamepad = Cont.gamepad;
@@ -577,100 +578,118 @@ void FIOSInputInterface::QueueKeyInput(int32 Key, int32 Char)
 	FIOSInputInterface::KeyInputStack.Add(Char);
 }
 
-
-void FIOSInputInterface::GetMovementData(FVector& Attitude, FVector& RotationRate, FVector& Gravity, FVector& Acceleration)
+void FIOSInputInterface::EnableMotionData(bool bEnable)
 {
+	bPauseMotion = !bEnable;
+
 #if !PLATFORM_TVOS
-	if (!bPauseMotion)
+	if (bPauseMotion && MotionManager != nil)
 	{
-		// initialize on first use
-		if (MotionManager == nil)
-		{
-			// Look to see if we can create the motion manager
-			MotionManager = [[CMMotionManager alloc] init];
+		[ReferenceAttitude release];
+		ReferenceAttitude = nil;
+		
+		[MotionManager release];
+		MotionManager = nil;
+	}
+	// When enabled MotionManager will be initialized on first use
+#endif
+}
 
-			// Check to see if the device supports full motion (gyro + accelerometer)
-			if (MotionManager.deviceMotionAvailable)
-			{
-				MotionManager.deviceMotionUpdateInterval = 0.02;
+bool FIOSInputInterface::IsMotionDataEnabled() const
+{
+	return !bPauseMotion;
+}
 
-				// Start the Device updating motion
-				[MotionManager startDeviceMotionUpdates];
-			}
-			else
-			{
-				[MotionManager startAccelerometerUpdates];
-				CenterPitch = CenterPitch = 0;
-				bIsCalibrationRequested = false;
-			}
-		}
+void FIOSInputInterface::GetMovementData(FVector& Attitude, FVector& RotationRate, FVector& Gravity, FVector& Acceleration)
+{
+#if !PLATFORM_TVOS
+	// initialize on first use
+	if (MotionManager == nil)
+	{
+		// Look to see if we can create the motion manager
+		MotionManager = [[CMMotionManager alloc] init];
 
-		// do we have full motion data?
-		if (MotionManager.deviceMotionActive)
+		// Check to see if the device supports full motion (gyro + accelerometer)
+		if (MotionManager.deviceMotionAvailable)
 		{
-			// Grab the values
-			CMAttitude* CurrentAttitude = MotionManager.deviceMotion.attitude;
-			CMRotationRate CurrentRotationRate = MotionManager.deviceMotion.rotationRate;
-			CMAcceleration CurrentGravity = MotionManager.deviceMotion.gravity;
-			CMAcceleration CurrentUserAcceleration = MotionManager.deviceMotion.userAcceleration;
-
-			// apply a reference attitude if we have been calibrated away from default
-			if (ReferenceAttitude)
-			{
-				[CurrentAttitude multiplyByInverseOfAttitude : ReferenceAttitude];
-			}
+			MotionManager.deviceMotionUpdateInterval = 0.02;
 
-			// convert to UE3
-			Attitude = FVector(float(CurrentAttitude.pitch), float(CurrentAttitude.yaw), float(CurrentAttitude.roll));
-			RotationRate = FVector(float(CurrentRotationRate.x), float(CurrentRotationRate.y), float(CurrentRotationRate.z));
-			Gravity = FVector(float(CurrentGravity.x), float(CurrentGravity.y), float(CurrentGravity.z));
-			Acceleration = FVector(float(CurrentUserAcceleration.x), float(CurrentUserAcceleration.y), float(CurrentUserAcceleration.z));
+			// Start the Device updating motion
+			[MotionManager startDeviceMotionUpdates];
 		}
 		else
 		{
-			// get the plain accleration
-			CMAcceleration RawAcceleration = [MotionManager accelerometerData].acceleration;
-			FVector NewAcceleration(RawAcceleration.x, RawAcceleration.y, RawAcceleration.z);
+			[MotionManager startAccelerometerUpdates];
+			CenterPitch = CenterPitch = 0;
+			bIsCalibrationRequested = false;
+		}
+	}
+
+	// do we have full motion data?
+	if (MotionManager.deviceMotionActive)
+	{
+		// Grab the values
+		CMAttitude* CurrentAttitude = MotionManager.deviceMotion.attitude;
+		CMRotationRate CurrentRotationRate = MotionManager.deviceMotion.rotationRate;
+		CMAcceleration CurrentGravity = MotionManager.deviceMotion.gravity;
+		CMAcceleration CurrentUserAcceleration = MotionManager.deviceMotion.userAcceleration;
+
+		// apply a reference attitude if we have been calibrated away from default
+		if (ReferenceAttitude)
+		{
+			[CurrentAttitude multiplyByInverseOfAttitude : ReferenceAttitude];
+		}
 
-			// storage for keeping the accelerometer values over time (for filtering)
-			static bool bFirstAccel = true;
+		// convert to UE3
+		Attitude = FVector(float(CurrentAttitude.pitch), float(CurrentAttitude.yaw), float(CurrentAttitude.roll));
+		RotationRate = FVector(float(CurrentRotationRate.x), float(CurrentRotationRate.y), float(CurrentRotationRate.z));
+		Gravity = FVector(float(CurrentGravity.x), float(CurrentGravity.y), float(CurrentGravity.z));
+		Acceleration = FVector(float(CurrentUserAcceleration.x), float(CurrentUserAcceleration.y), float(CurrentUserAcceleration.z));
+	}
+	else
+	{
+		// get the plain accleration
+		CMAcceleration RawAcceleration = [MotionManager accelerometerData].acceleration;
+		FVector NewAcceleration(RawAcceleration.x, RawAcceleration.y, RawAcceleration.z);
 
-			// how much of the previous frame's acceleration to keep
-			const float VectorFilter = bFirstAccel ? 0.0f : 0.85f;
-			bFirstAccel = false;
+		// storage for keeping the accelerometer values over time (for filtering)
+		static bool bFirstAccel = true;
 
-			// apply new accelerometer values to last frames
-			FilteredAccelerometer = FilteredAccelerometer * VectorFilter + (1.0f - VectorFilter) * NewAcceleration;
+		// how much of the previous frame's acceleration to keep
+		const float VectorFilter = bFirstAccel ? 0.0f : 0.85f;
+		bFirstAccel = false;
 
-			// create an normalized acceleration vector
-			FVector FinalAcceleration = -FilteredAccelerometer.GetSafeNormal();
+		// apply new accelerometer values to last frames
+		FilteredAccelerometer = FilteredAccelerometer * VectorFilter + (1.0f - VectorFilter) * NewAcceleration;
 
-			// calculate Roll/Pitch
-			float CurrentPitch = FMath::Atan2(FinalAcceleration.Y, FinalAcceleration.Z);
-			float CurrentRoll = -FMath::Atan2(FinalAcceleration.X, FinalAcceleration.Z);
+		// create an normalized acceleration vector
+		FVector FinalAcceleration = -FilteredAccelerometer.GetSafeNormal();
 
-			// if we want to calibrate, use the current values as center
-			if (bIsCalibrationRequested)
-			{
-				CenterPitch = CurrentPitch;
-				CenterRoll = CurrentRoll;
-				bIsCalibrationRequested = false;
-			}
+		// calculate Roll/Pitch
+		float CurrentPitch = FMath::Atan2(FinalAcceleration.Y, FinalAcceleration.Z);
+		float CurrentRoll = -FMath::Atan2(FinalAcceleration.X, FinalAcceleration.Z);
+
+		// if we want to calibrate, use the current values as center
+		if (bIsCalibrationRequested)
+		{
+			CenterPitch = CurrentPitch;
+			CenterRoll = CurrentRoll;
+			bIsCalibrationRequested = false;
+		}
 
-			CurrentPitch -= CenterPitch;
-			CurrentRoll -= CenterRoll;
+		CurrentPitch -= CenterPitch;
+		CurrentRoll -= CenterRoll;
 
-			Attitude = FVector(CurrentPitch, 0, CurrentRoll);
-			RotationRate = FVector(LastPitch - CurrentPitch, 0, LastRoll - CurrentRoll);
-			Gravity = FVector(0, 0, 0);
+		Attitude = FVector(CurrentPitch, 0, CurrentRoll);
+		RotationRate = FVector(LastPitch - CurrentPitch, 0, LastRoll - CurrentRoll);
+		Gravity = FVector(0, 0, 0);
 
-			// use the raw acceleration for acceleration
-			Acceleration = NewAcceleration;
+		// use the raw acceleration for acceleration
+		Acceleration = NewAcceleration;
 
-			// remember for next time (for rotation rate)
-			LastPitch = CurrentPitch;
-			LastRoll = CurrentRoll;
-		}
+		// remember for next time (for rotation rate)
+		LastPitch = CurrentPitch;
+		LastRoll = CurrentRoll;
 	}
 #endif
 }
@@ -680,7 +699,7 @@ void FIOSInputInterface::CalibrateMotion(uint32 PlayerIndex)
 #if !PLATFORM_TVOS
 	// If we are using the motion manager, grab a reference frame.  Note, once you set the Attitude Reference frame
 	// all additional reference information will come from it
-	if (MotionManager.deviceMotionActive)
+	if (MotionManager && MotionManager.deviceMotionActive)
 	{
 		ReferenceAttitude = [MotionManager.deviceMotion.attitude retain];
 	}
diff --git a/Engine/Source/Runtime/ApplicationCore/Private/IOS/IOSPlatformApplicationMisc.cpp b/Engine/Source/Runtime/ApplicationCore/Private/IOS/IOSPlatformApplicationMisc.cpp
--- a/Engine/Source/Runtime/ApplicationCore/Private/IOS/IOSPlatformApplicationMisc.cpp
+++ b/Engine/Source/Runtime/ApplicationCore/Private/IOS/IOSPlatformApplicationMisc.cpp
@@ -208,6 +208,18 @@ bool FIOSPlatformApplicationMisc::IsControllerAssignedToGamepad(int32 Controller
 	return InputInterface->IsControllerAssignedToGamepad(ControllerId);
 }
 
+void FIOSPlatformApplicationMisc::EnableMotionData(bool bEnable)
+{
+	FIOSInputInterface* InputInterface = (FIOSInputInterface*)CachedApplication->GetInputInterface();
+	return InputInterface->EnableMotionData(bEnable);
+}
+
+bool FIOSPlatformApplicationMisc::IsMotionDataEnabled()
+{
+	const FIOSInputInterface* InputInterface = (const FIOSInputInterface*)CachedApplication->GetInputInterface();
+	return InputInterface->IsMotionDataEnabled();
+}
+
 void FIOSPlatformApplicationMisc::ClipboardCopy(const TCHAR* Str)
 {
 #if !PLATFORM_TVOS
diff --git a/Engine/Source/Runtime/ApplicationCore/Public/GenericPlatform/GenericPlatformApplicationMisc.h b/Engine/Source/Runtime/ApplicationCore/Public/GenericPlatform/GenericPlatformApplicationMisc.h
--- a/Engine/Source/Runtime/ApplicationCore/Public/GenericPlatform/GenericPlatformApplicationMisc.h
+++ b/Engine/Source/Runtime/ApplicationCore/Public/GenericPlatform/GenericPlatformApplicationMisc.h
@@ -201,6 +201,21 @@ struct APPLICATIONCORE_API FGenericPlatformApplicationMisc
 		return FString(TEXT("None"));
 	}
 
+	/*
+	* Whether to enable controller motion data polling (by default motion data is enabled)
+	* Some platforms may want to disable it to reduce battery drain
+	*/
+	static void EnableMotionData(bool bEnable)
+	{}
+
+	/*
+	* Whether controller motion data polling is enabled (by default motion data is enabled)
+	*/
+	static bool IsMotionDataEnabled()
+	{
+		return true;
+	}
+				
 	/** Copies text to the operating system clipboard. */
 	static void ClipboardCopy(const TCHAR* Str);
 
diff --git a/Engine/Source/Runtime/ApplicationCore/Public/IOS/IOSInputInterface.h b/Engine/Source/Runtime/ApplicationCore/Public/IOS/IOSInputInterface.h
--- a/Engine/Source/Runtime/ApplicationCore/Public/IOS/IOSInputInterface.h
+++ b/Engine/Source/Runtime/ApplicationCore/Public/IOS/IOSInputInterface.h
@@ -82,6 +82,9 @@ public:
 	bool IsControllerAssignedToGamepad(int32 ControllerId) const;
 	bool IsGamepadAttached() const;
 
+	void EnableMotionData(bool bEnable);
+	bool IsMotionDataEnabled() const;
+	
 private:
 
 	FIOSInputInterface( const TSharedRef< FGenericApplicationMessageHandler >& InMessageHandler );
@@ -151,17 +154,16 @@ private:
 	// should we allow controllers to send input
 	bool bAllowControllers;
 	
+	/** Is motion paused or not? */
+	bool bPauseMotion;
+
 #if !PLATFORM_TVOS
 	/** Access to the ios devices motion */
 	CMMotionManager* MotionManager;
-	
-	/** Is motion paused or not? */
-	bool bPauseMotion;
 
 	/** Access to the ios devices tilt information */
 	CMAttitude* ReferenceAttitude;
 #endif
-	
 
 	/** Last frames roll, for calculating rate */
 	float LastRoll;
diff --git a/Engine/Source/Runtime/ApplicationCore/Public/IOS/IOSPlatformApplicationMisc.h b/Engine/Source/Runtime/ApplicationCore/Public/IOS/IOSPlatformApplicationMisc.h
--- a/Engine/Source/Runtime/ApplicationCore/Public/IOS/IOSPlatformApplicationMisc.h
+++ b/Engine/Source/Runtime/ApplicationCore/Public/IOS/IOSPlatformApplicationMisc.h
@@ -19,6 +19,9 @@ struct APPLICATIONCORE_API FIOSPlatformApplicationMisc : public FGenericPlatform
 	static void ResetGamepadAssignmentToController(int32 ControllerId);
 	static bool IsControllerAssignedToGamepad(int32 ControllerId);
 
+	static void EnableMotionData(bool bEnable);
+	static bool IsMotionDataEnabled();
+
 	static void ClipboardCopy(const TCHAR* Str);
 	static void ClipboardPaste(class FString& Dest);
 
diff --git a/Engine/Source/Runtime/IOS/IOSRuntimeSettings/Classes/IOSRuntimeSettings.h b/Engine/Source/Runtime/IOS/IOSRuntimeSettings/Classes/IOSRuntimeSettings.h
--- a/Engine/Source/Runtime/IOS/IOSRuntimeSettings/Classes/IOSRuntimeSettings.h
+++ b/Engine/Source/Runtime/IOS/IOSRuntimeSettings/Classes/IOSRuntimeSettings.h
@@ -347,6 +347,10 @@ public:
 	// If checked, Bluetooth connected controllers will send input
 	UPROPERTY(GlobalConfig, EditAnywhere, Category = Input, meta = (DisplayName = "Allow MFi (Bluetooth) controllers"))
 	bool bAllowControllers;
+
+	// Disables usage of device motion data. If application does not use motion data disabling it will improve battery life
+	UPROPERTY(GlobalConfig, EditAnywhere, Category = Input, meta = (DisplayName = "Disable Motion Controls"))
+	bool bDisableMotionData;
 	
 	// Supports default portrait orientation. Landscape will not be supported.
 	UPROPERTY(GlobalConfig, EditAnywhere, Category = DeviceOrientations)
diff --git a/Engine/Source/Runtime/IOS/IOSRuntimeSettings/Private/IOSRuntimeSettings.cpp b/Engine/Source/Runtime/IOS/IOSRuntimeSettings/Private/IOSRuntimeSettings.cpp
--- a/Engine/Source/Runtime/IOS/IOSRuntimeSettings/Private/IOSRuntimeSettings.cpp
+++ b/Engine/Source/Runtime/IOS/IOSRuntimeSettings/Private/IOSRuntimeSettings.cpp
@@ -41,6 +41,7 @@ UIOSRuntimeSettings::UIOSRuntimeSettings(const FObjectInitializer& ObjectInitial
 	bAllowRemoteRotation = true;
 	bUseRemoteAsVirtualJoystick = true;
 	bUseRemoteAbsoluteDpadValues = false;
+	bDisableMotionData = false;
     bEnableRemoteNotificationsSupport = false;
     bEnableBackgroundFetch = false;
 	bSupportsOpenGLES2 = false;
