revert back to old audio toggle for the time being

diff --git a/Engine/Source/Runtime/ApplicationCore/Private/IOS/IOSAppDelegate.cpp b/Engine/Source/Runtime/ApplicationCore/Private/IOS/IOSAppDelegate.cpp
--- a/Engine/Source/Runtime/ApplicationCore/Private/IOS/IOSAppDelegate.cpp
+++ b/Engine/Source/Runtime/ApplicationCore/Private/IOS/IOSAppDelegate.cpp
@@ -497,7 +497,7 @@ static IOSAppDelegate* CachedDelegate = nil;
 - (void)ToggleAudioSession:(bool)bActive force:(bool)bForce
 {
 	// @todo kairos: resolve old vs new before we go to main
-	if (true)
+	if (false)
 	{
 		// we can actually override bActive based on backgrounding behavior, as that's the only time we actually deactivate the session
 		// @todo kairos: is this a valid check?
