Avoid calling SDL_Destroy when the HWnd is NULL

diff --git a/Engine/Source/Runtime/ApplicationCore/Private/Linux/LinuxWindow.cpp b/Engine/Source/Runtime/ApplicationCore/Private/Linux/LinuxWindow.cpp
--- a/Engine/Source/Runtime/ApplicationCore/Private/Linux/LinuxWindow.cpp
+++ b/Engine/Source/Runtime/ApplicationCore/Private/Linux/LinuxWindow.cpp
@@ -420,14 +420,17 @@ void FLinuxWindow::BringToFront( bool bForce )
 /** Native windows should implement this function by asking the OS to destroy OS-specific resource associated with the window (e.g. Win32 window handle) */
 void FLinuxWindow::Destroy()
 {
-	OwningApplication->RemoveRevertFocusWindow( HWnd );
-	OwningApplication->RemoveEventWindow( HWnd );
-	OwningApplication->RemoveNotificationWindow( HWnd );
+	if (HWnd)
+	{
+		OwningApplication->RemoveRevertFocusWindow( HWnd );
+		OwningApplication->RemoveEventWindow( HWnd );
+		OwningApplication->RemoveNotificationWindow( HWnd );
 
-	UE_LOG(LogLinuxWindow, Verbose, TEXT("Destroying SDL Window '%p'\n"), HWnd);
+		UE_LOG(LogLinuxWindow, Verbose, TEXT("Destroying SDL Window '%p'\n"), HWnd);
 
-	SDL_DestroyWindow( HWnd );
-	HWnd = nullptr;
+		SDL_DestroyWindow( HWnd );
+		HWnd = nullptr;
+	}
 }
 
 /** Native window should implement this function by performing the equivalent of the Win32 minimize-to-taskbar operation */
