UE-71937 //UE4/Main - Compile UE4Game TVOS - error: 'UIInterfaceOrientationUnknown' is unavailable: not available on tvOS

#4.23
#tvOS

Compile errors on tvOS

diff --git a/Engine/Source/Runtime/ApplicationCore/Private/IOS/IOSAppDelegate.cpp b/Engine/Source/Runtime/ApplicationCore/Private/IOS/IOSAppDelegate.cpp
--- a/Engine/Source/Runtime/ApplicationCore/Private/IOS/IOSAppDelegate.cpp
+++ b/Engine/Source/Runtime/ApplicationCore/Private/IOS/IOSAppDelegate.cpp
@@ -610,10 +610,10 @@ static IOSAppDelegate* CachedDelegate = nil;
 					else
 					{
 						AVAudioSessionCategoryOptions opts =
-							AVAudioSessionCategoryOptionAllowBluetooth |
 							AVAudioSessionCategoryOptionAllowBluetoothA2DP |
 #if !PLATFORM_TVOS
-    						AVAudioSessionCategoryOptionDefaultToSpeaker |
+							AVAudioSessionCategoryOptionAllowBluetooth |
+							AVAudioSessionCategoryOptionDefaultToSpeaker |
 #endif
 							AVAudioSessionCategoryOptionMixWithOthers;
 						
@@ -684,9 +684,9 @@ static IOSAppDelegate* CachedDelegate = nil;
 				else
 				{
 					AVAudioSessionCategoryOptions opts =
-						AVAudioSessionCategoryOptionAllowBluetooth |
 						AVAudioSessionCategoryOptionAllowBluetoothA2DP |
 #if !PLATFORM_TVOS
+						AVAudioSessionCategoryOptionAllowBluetooth |
 						AVAudioSessionCategoryOptionDefaultToSpeaker |
 #endif
 						AVAudioSessionCategoryOptionMixWithOthers;
@@ -714,9 +714,9 @@ static IOSAppDelegate* CachedDelegate = nil;
         if (self.bVoiceChatEnabled)
 		{
 			AVAudioSessionCategoryOptions opts =
-				AVAudioSessionCategoryOptionAllowBluetooth |
 				AVAudioSessionCategoryOptionAllowBluetoothA2DP |
 #if !PLATFORM_TVOS
+				AVAudioSessionCategoryOptionAllowBluetooth |
 				AVAudioSessionCategoryOptionDefaultToSpeaker |
 #endif
 				AVAudioSessionCategoryOptionMixWithOthers;
@@ -881,7 +881,7 @@ bool GIsSuspended = 0;
 		// Don't deadlock here because a msg box may appear super early blocking the game thread and then the app may go into the background
 		double	startTime = FPlatformTime::Seconds();
 
-		// don't wait for FDefaultGameMoviePlayer::WaitForMovieToFinish(), crash with 0x8badf00d if “Wait for Movies to Complete” is checked
+		// don't wait for FDefaultGameMoviePlayer::WaitForMovieToFinish(), crash with 0x8badf00d if "Wait for Movies to Complete" is checked
 		while(!self.bHasSuspended && !FAppEntry::IsStartupMoviePlaying() &&  (FPlatformTime::Seconds() - startTime) < cMaxThreadWaitTime)
 		{
             FIOSPlatformRHIFramePacer::Suspend();
diff --git a/Engine/Source/Runtime/ApplicationCore/Private/IOS/IOSView.cpp b/Engine/Source/Runtime/ApplicationCore/Private/IOS/IOSView.cpp
--- a/Engine/Source/Runtime/ApplicationCore/Private/IOS/IOSView.cpp
+++ b/Engine/Source/Runtime/ApplicationCore/Private/IOS/IOSView.cpp
@@ -232,16 +232,11 @@ id<MTLDevice> GMetalDevice = nil;
 	// Initialize some variables
 	SwapCount = 0;
 
-#if PLATFORM_TVOS
-	// @todo tvos: This may need to be exposed to the game so that when you click Menu it will background the app
-	// this is basically the same way Android handles the Back button (maybe we should pass Menu button as back... maybe)
-	// @todo embedded: this will not work, because this view hasn't been added to a hierarchy yet, so the IOSController is guaranteed to be null. DELAY THIS
-	AppDelegate.IOSController.controllerUserInteractionEnabled = NO;
-#endif
-
 //	self.userInteractionEnabled = YES;
 //	self.clearsContextBeforeDrawing = NO;
+#if !PLATFORM_TVOS
 	self.multipleTouchEnabled = YES;
+#endif
 
 	FMemory::Memzero(AllTouches, sizeof(AllTouches));
 	[self setAutoresizingMask: UIViewAutoresizingFlexibleWidth | UIViewAutoresizingFlexibleHeight];
diff --git a/Engine/Source/Runtime/Core/Private/IOS/IOSPlatformMisc.cpp b/Engine/Source/Runtime/Core/Private/IOS/IOSPlatformMisc.cpp
--- a/Engine/Source/Runtime/Core/Private/IOS/IOSPlatformMisc.cpp
+++ b/Engine/Source/Runtime/Core/Private/IOS/IOSPlatformMisc.cpp
@@ -312,7 +312,9 @@ EDeviceScreenOrientation ConvertFromUIInterfaceOrientation(UIInterfaceOrientatio
 }
 #endif
 
+#if !PLATFORM_TVOS
 UIInterfaceOrientation GInterfaceOrientation = UIInterfaceOrientationUnknown;
+#endif
 
 EDeviceScreenOrientation FIOSPlatformMisc::GetDeviceOrientation()
 {
diff --git a/Engine/Source/Runtime/Launch/Private/IOS/LaunchIOS.cpp b/Engine/Source/Runtime/Launch/Private/IOS/LaunchIOS.cpp
--- a/Engine/Source/Runtime/Launch/Private/IOS/LaunchIOS.cpp
+++ b/Engine/Source/Runtime/Launch/Private/IOS/LaunchIOS.cpp
@@ -180,7 +180,7 @@ void FAppEntry::Resume(bool bIsInterrupt)
 void FAppEntry::PreInit(IOSAppDelegate* AppDelegate, UIApplication* Application)
 {
 	// make a controller object
-	UIViewController* IOSController = [[IOSViewController alloc] init];
+	IOSViewController* IOSController = [[IOSViewController alloc] init];
 	
 #if PLATFORM_TVOS
 	// @todo tvos: This may need to be exposed to the game so that when you click Menu it will background the app
