[iOS] fix old audio toggle not working for recording

diff --git a/Engine/Source/Runtime/ApplicationCore/Private/IOS/IOSAppDelegate.cpp b/Engine/Source/Runtime/ApplicationCore/Private/IOS/IOSAppDelegate.cpp
--- a/Engine/Source/Runtime/ApplicationCore/Private/IOS/IOSAppDelegate.cpp
+++ b/Engine/Source/Runtime/ApplicationCore/Private/IOS/IOSAppDelegate.cpp
@@ -758,17 +758,22 @@ static IOSAppDelegate* CachedDelegate = nil;
 
 - (void)EnableVoiceChat:(bool)bEnable
 {
+	self.bVoiceChatEnabled = false;
+	
     // mobile will prompt for microphone access
     if (FApp::IsUnattended())
 	{
 		return;
 	}
-	[self SetFeature:EAudioFeature::VoiceChat Active:bEnable];
+	self.bVoiceChatEnabled = bEnable;
+	[self ToggleAudioSession:self.bAudioActive force:true];
+	//[self SetFeature:EAudioFeature::VoiceChat Active:bEnable];
 }
 
 - (bool)IsVoiceChatEnabled
 {
-	return [self IsFeatureActive:EAudioFeature::VoiceChat];
+	return self.bVoiceChatEnabled;
+	//return [self IsFeatureActive:EAudioFeature::VoiceChat];
 }
 
 
