redo the exit crash fix for IOS

diff --git a/Engine/Source/Runtime/ApplicationCore/Private/IOS/IOSAppDelegate.cpp b/Engine/Source/Runtime/ApplicationCore/Private/IOS/IOSAppDelegate.cpp
--- a/Engine/Source/Runtime/ApplicationCore/Private/IOS/IOSAppDelegate.cpp
+++ b/Engine/Source/Runtime/ApplicationCore/Private/IOS/IOSAppDelegate.cpp
@@ -1464,21 +1464,22 @@ extern double GCStartTime;
 	FCoreDelegates::ApplicationWillTerminateDelegate.Broadcast();
     
     // note that we are shutting down
-    GIsRequestingExit = true;
+    // TODO: fix the reason why we are hanging when asked to shutdown
+/*    GIsRequestingExit = true;
     
-    if (!bEngineInit)
+    if (!bEngineInit)*/
     {
         // we haven't yet made it to the point where the engine is initialized, so just exit the app
         _Exit(0);
     }
-    else
+/*    else
     {
         // wait for the game thread to shut down
         while (self.bHasStarted == true)
         {
             usleep(3);
         }
-    }
+    }*/
 }
 
 - (void)applicationDidReceiveMemoryWarning:(UIApplication *)application
