Remove spamming USER EVENTS when window has lost focus

When we got a FOCUS OUT event we would start spamming upwards of ~2000 USER EVENTS checking if we need to deactivate the application. This was taking upwards of 150ms
Now per event pump we check if we have been focused out for to then we deactivate.

diff --git a/Engine/Source/Runtime/ApplicationCore/Private/Linux/LinuxApplication.cpp b/Engine/Source/Runtime/ApplicationCore/Private/Linux/LinuxApplication.cpp
--- a/Engine/Source/Runtime/ApplicationCore/Private/Linux/LinuxApplication.cpp
+++ b/Engine/Source/Runtime/ApplicationCore/Private/Linux/LinuxApplication.cpp
@@ -838,6 +838,9 @@ void FLinuxApplication::ProcessDeferredMessage( SDL_Event Event )
 						UE_LOG(LogLinuxWindowEvent, Verbose, TEXT("WM_SETFOCUS                                 : %d"), CurrentEventWindow->GetID());
 
 						CurrentFocusWindow = CurrentEventWindow;
+
+						// We have gained focus, we can stop trying to check if we need to deactivate
+						FocusOutDeactivationTime = 0.0;
 					}
 					break;
 
@@ -855,11 +858,6 @@ void FLinuxApplication::ProcessDeferredMessage( SDL_Event Event )
 							if(bActivateApp)
 							{
 								FocusOutDeactivationTime = FPlatformTime::Seconds() + DeactivationThreadshold;
-
-								SDL_Event event;
-								event.type = SDL_USEREVENT;
-								event.user.code = CheckForDeactivation;
-								SDL_PushEvent(&event);
 							}
 						}
 						CurrentFocusWindow = nullptr;
@@ -921,28 +919,6 @@ void FLinuxApplication::ProcessDeferredMessage( SDL_Event Event )
 		}
 		break;
 
-	case SDL_USEREVENT:
-		{
-			if(Event.user.code == CheckForDeactivation)
-			{
-				// We still havent hit our timeout limit, try again
-				if (FocusOutDeactivationTime > FPlatformTime::Seconds())
-				{
-					SDL_Event event;
-					event.type = SDL_USEREVENT;
-					event.user.code = CheckForDeactivation;
-					SDL_PushEvent(&event);
-				}
-				// If we don't use bIsDragWindowButtonPressed the draged window will be destroyed because we
-				// deactivate the whole appliacton. TODO Is that a bug? Do we have to do something?
-				else if (!CurrentFocusWindow.IsValid() && !bIsDragWindowButtonPressed)
-				{
-					DeactivateApplication();
-				}
-			}
-		}
-		break;
-
 	case SDL_FINGERDOWN:
 		{
 			UE_LOG(LogLinuxWindow, Verbose, TEXT("Finger %llu is down at (%f, %f)"), Event.tfinger.fingerId, Event.tfinger.x, Event.tfinger.y);
@@ -1744,6 +1720,33 @@ void FLinuxApplication::SaveWindowPropertiesForEventLoop(void)
 void FLinuxApplication::ClearWindowPropertiesAfterEventLoop(void)
 {
 	SavedWindowPropertiesForEventLoop.Empty();
+
+	// This is a hack for 4.22.1 to avoid changing a header. We will be calling in from here at
+	// LinuxPlatformApplicationMisc.cpp in FLinuxPlatformApplicationMisc::PumpMessages
+	//
+	// In 4.23 this will be in a function: void FLinuxApplication::CheckIfApplicatioNeedsDeactivation()
+
+	// If FocusOutDeactivationTime is set we have had a focus out event and are waiting to see if we need to Deactivate the Application
+	if (!FMath::IsNearlyZero(FocusOutDeactivationTime))
+	{
+		// We still havent hit our timeout limit, keep waiting
+		if (FocusOutDeactivationTime > FPlatformTime::Seconds())
+		{
+			return;
+		}
+		// If we don't use bIsDragWindowButtonPressed the draged window will be destroyed because we
+		// deactivate the whole appliacton. TODO Is that a bug? Do we have to do something?
+		else if (!CurrentFocusWindow.IsValid() && !bIsDragWindowButtonPressed)
+		{
+			DeactivateApplication();
+
+			FocusOutDeactivationTime = 0.0;
+		}
+		else
+		{
+			FocusOutDeactivationTime = 0.0;
+		}
+	}
 }
 
 void FLinuxApplication::GetWindowPropertiesInEventLoop(SDL_HWindow NativeWindow, FWindowProperties& Properties)
