more scope counters to try to track down spike in poll game device state on IOS

diff --git a/Engine/Source/Runtime/ApplicationCore/Private/IOS/IOSApplication.cpp b/Engine/Source/Runtime/ApplicationCore/Private/IOS/IOSApplication.cpp
--- a/Engine/Source/Runtime/ApplicationCore/Private/IOS/IOSApplication.cpp
+++ b/Engine/Source/Runtime/ApplicationCore/Private/IOS/IOSApplication.cpp
@@ -12,6 +12,7 @@
 #include "Misc/ScopeLock.h"
 #include "HAL/IConsoleManager.h"
 #include "IOS/IOSAsyncTask.h"
+#include "Stats/Stats.h"
 
 FIOSApplication* FIOSApplication::CreateIOSApplication()
 {
@@ -72,14 +73,23 @@ void FIOSApplication::PollGameDeviceState( const float TimeDelta )
 	}
 
 	// Poll game device state and send new events
-	InputInterface->Tick(TimeDelta);
-	InputInterface->SendControllerEvents();
+	{
+		QUICK_SCOPE_CYCLE_COUNTER(STAT_IOSApplication_InputInterface_Tick);
+		InputInterface->Tick(TimeDelta);
+	}
+	{
+		QUICK_SCOPE_CYCLE_COUNTER(STAT_IOSApplication_InputInterface_SendControllerEvents);
+		InputInterface->SendControllerEvents();
+	}
 
 	// Poll externally-implemented devices
-	for (auto DeviceIt = ExternalInputDevices.CreateIterator(); DeviceIt; ++DeviceIt)
 	{
-		(*DeviceIt)->Tick(TimeDelta);
-		(*DeviceIt)->SendControllerEvents();
+		QUICK_SCOPE_CYCLE_COUNTER(STAT_IOSApplication_ExternalInputDevice);
+		for (auto DeviceIt = ExternalInputDevices.CreateIterator(); DeviceIt; ++DeviceIt)
+		{
+			(*DeviceIt)->Tick(TimeDelta);
+			(*DeviceIt)->SendControllerEvents();
+		}
 	}
 }
 
