Call _Exit when leaving main loop on iOS and in unattended mode
#rnx

diff --git a/Engine/Source/Runtime/ApplicationCore/Private/IOS/IOSAppDelegate.cpp b/Engine/Source/Runtime/ApplicationCore/Private/IOS/IOSAppDelegate.cpp
--- a/Engine/Source/Runtime/ApplicationCore/Private/IOS/IOSAppDelegate.cpp
+++ b/Engine/Source/Runtime/ApplicationCore/Private/IOS/IOSAppDelegate.cpp
@@ -336,7 +336,7 @@ static IOSAppDelegate* CachedDelegate = nil;
     
     self.bHasStarted = false;
     
-    if(bForceExit)
+    if(bForceExit || FApp::IsUnattended())
     {
         _Exit(0);
         //exit(0);  // As far as I can tell we run into a lot of trouble trying to run static destructors, so this is a no go :(
