Fix compile issue for Android x86 with media player

#ue4
#android

diff --git a/Engine/Source/Runtime/ApplicationCore/Private/Android/AndroidJavaMediaPlayer.cpp b/Engine/Source/Runtime/ApplicationCore/Private/Android/AndroidJavaMediaPlayer.cpp
--- a/Engine/Source/Runtime/ApplicationCore/Private/Android/AndroidJavaMediaPlayer.cpp
+++ b/Engine/Source/Runtime/ApplicationCore/Private/Android/AndroidJavaMediaPlayer.cpp
@@ -28,7 +28,7 @@ static jfieldID FindField(JNIEnv* JEnv, jclass Class, const ANSICHAR* FieldName,
 *****************************************************************************/
 
 FCriticalSection FJavaAndroidMediaPlayer::MediaDataSourcesCS;
-TMap<int64, TWeakPtr<FJavaAndroidMediaDataSource>> FJavaAndroidMediaPlayer::AllMediaDataSources;
+TMap<int64, TWeakPtr<FJavaAndroidMediaDataSource, ESPMode::ThreadSafe>> FJavaAndroidMediaPlayer::AllMediaDataSources;
 
 TSharedPtr<FJavaAndroidMediaDataSource, ESPMode::ThreadSafe> FJavaAndroidMediaPlayer::GetMediaDataSourcePtr(int64 Identifier)
 {
diff --git a/Engine/Source/Runtime/ApplicationCore/Public/Android/AndroidJavaMediaPlayer.h b/Engine/Source/Runtime/ApplicationCore/Public/Android/AndroidJavaMediaPlayer.h
--- a/Engine/Source/Runtime/ApplicationCore/Public/Android/AndroidJavaMediaPlayer.h
+++ b/Engine/Source/Runtime/ApplicationCore/Public/Android/AndroidJavaMediaPlayer.h
@@ -217,10 +217,10 @@ public:
 
 protected:
 	static FCriticalSection MediaDataSourcesCS;
-	static TMap<int64, TWeakPtr<FJavaAndroidMediaDataSource>> AllMediaDataSources;
+	static TMap<int64, TWeakPtr<FJavaAndroidMediaDataSource, ESPMode::ThreadSafe>> AllMediaDataSources;
 
 public:
-	static TSharedPtr<FJavaAndroidMediaDataSource> GetMediaDataSourcePtr(int64 Identifier);
+	static TSharedPtr<FJavaAndroidMediaDataSource, ESPMode::ThreadSafe> GetMediaDataSourcePtr(int64 Identifier);
 	static void AddMediaDataSourcePtr(int64 Identifier, TSharedPtr<FJavaAndroidMediaDataSource, ESPMode::ThreadSafe> MediaDataSource);
 };
 
