Only allow A2DP bluetooth - fixes poor quality audio with bluetooth headsets on iOS

diff --git a/Engine/Source/Runtime/ApplicationCore/Private/IOS/IOSAppDelegate.cpp b/Engine/Source/Runtime/ApplicationCore/Private/IOS/IOSAppDelegate.cpp
--- a/Engine/Source/Runtime/ApplicationCore/Private/IOS/IOSAppDelegate.cpp
+++ b/Engine/Source/Runtime/ApplicationCore/Private/IOS/IOSAppDelegate.cpp
@@ -590,7 +590,6 @@ static IOSAppDelegate* CachedDelegate = nil;
 						AVAudioSessionCategoryOptions opts =
 							AVAudioSessionCategoryOptionAllowBluetoothA2DP |
 #if !PLATFORM_TVOS
-							AVAudioSessionCategoryOptionAllowBluetooth |
 							AVAudioSessionCategoryOptionDefaultToSpeaker |
 #endif
 							AVAudioSessionCategoryOptionMixWithOthers;
@@ -664,7 +663,6 @@ static IOSAppDelegate* CachedDelegate = nil;
 					AVAudioSessionCategoryOptions opts =
 						AVAudioSessionCategoryOptionAllowBluetoothA2DP |
 #if !PLATFORM_TVOS
-						AVAudioSessionCategoryOptionAllowBluetooth |
 						AVAudioSessionCategoryOptionDefaultToSpeaker |
 #endif
 						AVAudioSessionCategoryOptionMixWithOthers;
@@ -694,7 +692,6 @@ static IOSAppDelegate* CachedDelegate = nil;
 			AVAudioSessionCategoryOptions opts =
 				AVAudioSessionCategoryOptionAllowBluetoothA2DP |
 #if !PLATFORM_TVOS
-				AVAudioSessionCategoryOptionAllowBluetooth |
 				AVAudioSessionCategoryOptionDefaultToSpeaker |
 #endif
 				AVAudioSessionCategoryOptionMixWithOthers;
