Append HasAppFocus to engine data in crash context under "Platform" prefix.

diff --git a/Engine/Source/Runtime/ApplicationCore/Private/Windows/WindowsPlatformApplicationMisc.cpp b/Engine/Source/Runtime/ApplicationCore/Private/Windows/WindowsPlatformApplicationMisc.cpp
--- a/Engine/Source/Runtime/ApplicationCore/Private/Windows/WindowsPlatformApplicationMisc.cpp
+++ b/Engine/Source/Runtime/ApplicationCore/Private/Windows/WindowsPlatformApplicationMisc.cpp
@@ -12,6 +12,7 @@
 #include "Modules/ModuleManager.h"
 #include "Misc/CoreDelegates.h"
 #include "Windows/WindowsPlatformOutputDevices.h"
+#include "GenericPlatform/GenericPlatformCrashContext.h"
 
 // Resource includes.
 #include "Runtime/Launch/Resources/Windows/Resource.h"
@@ -123,11 +124,12 @@ void FWindowsPlatformApplicationMisc::PumpMessages(bool bFromMainLoop)
 
 	// Determine if application has focus
 	bool HasFocus = FApp::UseVRFocus() ? FApp::HasVRFocus() : FWindowsPlatformApplicationMisc::IsThisApplicationForeground();
+	static bool HadFocus = false;
 
+#if WITH_EDITOR
 	// If editor thread doesn't have the focus, don't suck up too much CPU time.
 	if( GIsEditor )
 	{
-		static bool HadFocus=1;
 		if( HadFocus && !HasFocus )
 		{
 			// Drop our priority to speed up whatever is in the foreground.
@@ -145,6 +147,17 @@ void FWindowsPlatformApplicationMisc::PumpMessages(bool bFromMainLoop)
 		}
 		HadFocus = HasFocus;
 	}
+#endif
+
+#if !UE_SERVER
+	// For non-editor clients, record if the active window is in focus
+	if( HadFocus != HasFocus )
+	{
+		FGenericCrashContext::SetEngineData(TEXT("Platform.AppHasFocus"), HasFocus ? TEXT("true") : TEXT("false"));
+	}
+#endif
+
+	HadFocus = HasFocus;
 
 	// if its our window, allow sound, otherwise apply multiplier
 	FApp::SetVolumeMultiplier( HasFocus ? 1.0f : FApp::GetUnfocusedVolumeMultiplier() );
