- Added "engineisrunning" call from EmbeddedToNative on IOS to tell native app when the while(1) loop is going
- Added "devicesleep" native call to control mobile "screensaver"

diff --git a/Engine/Source/Runtime/ApplicationCore/Private/IOS/IOSAppDelegate.cpp b/Engine/Source/Runtime/ApplicationCore/Private/IOS/IOSAppDelegate.cpp
--- a/Engine/Source/Runtime/ApplicationCore/Private/IOS/IOSAppDelegate.cpp
+++ b/Engine/Source/Runtime/ApplicationCore/Private/IOS/IOSAppDelegate.cpp
@@ -24,6 +24,7 @@
 #include "Misc/App.h"
 #include "Algo/AllOf.h"
 #include "Misc/App.h"
+#include "Misc/EmbeddedCommunication.h"
 #if USE_MUTE_SWITCH_DETECTION
 #include "SharkfoodMuteSwitchDetector.h"
 #include "SharkfoodMuteSwitchDetector.m"
@@ -259,6 +260,13 @@ static IOSAppDelegate* CachedDelegate = nil;
 	}
 	self.savedOpenUrlParameters = nil; // clear after saved openurl delegate running
 
+#if BUILD_EMBEDDED_APP
+	// tell the embedded app that the while 1 loop is going
+	FEmbeddedCallParamsHelper Helper;
+	Helper.Command = TEXT("engineisrunning");
+	FEmbeddedDelegates::GetEmbeddedToNativeParamsDelegateForSubsystem(TEXT("native")).Broadcast(Helper);
+#endif
+
     while( !GIsRequestingExit )
 	{
         if (self.bIsSuspended)
diff --git a/Engine/Source/Runtime/Engine/Private/GameEngine.cpp b/Engine/Source/Runtime/Engine/Private/GameEngine.cpp
--- a/Engine/Source/Runtime/Engine/Private/GameEngine.cpp
+++ b/Engine/Source/Runtime/Engine/Private/GameEngine.cpp
@@ -774,6 +774,14 @@ public:
 				FEmbeddedCommunication::AllowSleep(Requester);
 				Message.OnCompleteDelegate({}, TEXT(""));
 			}
+			// execute any console commands
+			else if (Message.Command == TEXT("devicesleep"))
+			{
+				bool bEnabled = Message.Parameters.FindRef(TEXT("enabled")).ToBool();
+				FPlatformApplicationMisc::ControlScreensaver(bEnabled ? FGenericPlatformApplicationMisc::Enable : FGenericPlatformApplicationMisc::Disable);
+				Message.OnCompleteDelegate({}, TEXT(""));
+			}
+
 // disable things that are too scary to expose to a shipping external app/webview
 #if !UE_BUILD_SHIPPING
 			// execute any console commands
