addition of iPod Touch 7 to engine

diff --git a/Engine/Config/BaseDeviceProfiles.ini b/Engine/Config/BaseDeviceProfiles.ini
--- a/Engine/Config/BaseDeviceProfiles.ini
+++ b/Engine/Config/BaseDeviceProfiles.ini
@@ -18,6 +18,7 @@
 +DeviceProfileNameAndTypes=iPadMini4,IOS
 +DeviceProfileNameAndTypes=iPodTouch5,IOS
 +DeviceProfileNameAndTypes=iPodTouch6,IOS
++DeviceProfileNameAndTypes=iPodTouch7,IOS
 +DeviceProfileNameAndTypes=iPhone4,IOS
 +DeviceProfileNameAndTypes=iPhone4S,IOS
 +DeviceProfileNameAndTypes=iPhone5,IOS
@@ -241,6 +242,11 @@ DeviceType=IOS
 BaseProfileName=iPhone6
 +CVars=r.MobileContentScaleFactor=2
 
+[iPodTouch7 DeviceProfile]
+DeviceType=IOS
+BaseProfileName=iPhone7
++CVars=r.MobileContentScaleFactor=2
+
 [iPhone6 DeviceProfile]
 DeviceType=IOS
 BaseProfileName=IOS
diff --git a/Engine/Source/Runtime/ApplicationCore/Private/IOS/IOSAppDelegate.cpp b/Engine/Source/Runtime/ApplicationCore/Private/IOS/IOSAppDelegate.cpp
--- a/Engine/Source/Runtime/ApplicationCore/Private/IOS/IOSAppDelegate.cpp
+++ b/Engine/Source/Runtime/ApplicationCore/Private/IOS/IOSAppDelegate.cpp
@@ -1069,7 +1069,7 @@ static FAutoConsoleVariableRef CVarGEnableThermalsReport(
 		{
 			orient = UIImageOrientationRight;
 		}
-        else if (MainFrame.size.height == 568 || Device == FPlatformMisc::IOS_IPodTouch6)
+		else if (MainFrame.size.height == 568 || Device == FPlatformMisc::IOS_IPodTouch6 || Device == FPlatformMisc::IOS_IPodTouch7)
 		{
 			[ImageString appendString:@"-568h"];
 		}
diff --git a/Engine/Source/Runtime/ApplicationCore/Private/IOS/IOSPlatformApplicationMisc.cpp b/Engine/Source/Runtime/ApplicationCore/Private/IOS/IOSPlatformApplicationMisc.cpp
--- a/Engine/Source/Runtime/ApplicationCore/Private/IOS/IOSPlatformApplicationMisc.cpp
+++ b/Engine/Source/Runtime/ApplicationCore/Private/IOS/IOSPlatformApplicationMisc.cpp
@@ -251,7 +251,7 @@ void FIOSPlatformApplicationMisc::ClipboardPaste(class FString& Result)
 EScreenPhysicalAccuracy FIOSPlatformApplicationMisc::ComputePhysicalScreenDensity(int32& ScreenDensity)
 {
 	FPlatformMisc::EIOSDevice Device = FPlatformMisc::GetIOSDeviceType();
-	static_assert( FPlatformMisc::EIOSDevice::IOS_Unknown == 40, "Every device needs to be handled here." );
+	static_assert( FPlatformMisc::EIOSDevice::IOS_Unknown == 41, "Every device needs to be handled here." );
 
 	ScreenDensity = 0;
 	EScreenPhysicalAccuracy Accuracy = EScreenPhysicalAccuracy::Unknown;
@@ -268,6 +268,7 @@ EScreenPhysicalAccuracy FIOSPlatformApplicationMisc::ComputePhysicalScreenDensit
 	case FPlatformMisc::IOS_IPhone5S:
 	case FPlatformMisc::IOS_IPodTouch5:
 	case FPlatformMisc::IOS_IPodTouch6:
+	case FPlatformMisc::IOS_IPodTouch7:
 	case FPlatformMisc::IOS_IPhone6:
 	case FPlatformMisc::IOS_IPhone6S:
 	case FPlatformMisc::IOS_IPhone7:
diff --git a/Engine/Source/Runtime/Core/Private/IOS/IOSPlatformMisc.cpp b/Engine/Source/Runtime/Core/Private/IOS/IOSPlatformMisc.cpp
--- a/Engine/Source/Runtime/Core/Private/IOS/IOSPlatformMisc.cpp
+++ b/Engine/Source/Runtime/Core/Private/IOS/IOSPlatformMisc.cpp
@@ -396,10 +396,14 @@ FIOSPlatformMisc::EIOSDevice FIOSPlatformMisc::GetIOSDeviceType()
 		{
 			DeviceType = IOS_IPodTouch5;
 		}
-		else if (Major >= 7)
+		else if (Major == 7)
 		{
 			DeviceType = IOS_IPodTouch6;
 		}
+		else if (Major >= 9)
+		{
+			DeviceType = IOS_IPodTouch7;
+		}
 	}
 	// iPads
 	else if (DeviceIDString.StartsWith(TEXT("iPad")))
diff --git a/Engine/Source/Runtime/Core/Public/IOS/IOSPlatformMisc.h b/Engine/Source/Runtime/Core/Public/IOS/IOSPlatformMisc.h
--- a/Engine/Source/Runtime/Core/Public/IOS/IOSPlatformMisc.h
+++ b/Engine/Source/Runtime/Core/Public/IOS/IOSPlatformMisc.h
@@ -143,6 +143,7 @@ struct CORE_API FIOSPlatformMisc : public FApplePlatformMisc
 		IOS_IPadPro3_129,
         IOS_IPadAir3,
         IOS_IPadMini5,
+		IOS_IPodTouch7,
 		IOS_Unknown,
 	};
 
@@ -192,6 +193,7 @@ struct CORE_API FIOSPlatformMisc : public FApplePlatformMisc
 			TEXT("IPadPro3_129"),
             TEXT("IPadAir3"),
             TEXT("IPadMini5"),
+			TEXT("IPodTouch7"),
 			TEXT("Unknown"),
 		};
 		static_assert((sizeof(IOSDeviceNames) / sizeof(IOSDeviceNames[0])) == ((int32)IOS_Unknown + 1), "Mismatched IOSDeviceNames and EIOSDevice.");
