- Enabled going in and out of high quality voice chat mode (echo cancellation, etc, but makes game audio quiet, so it's an optional thing)

diff --git a/Engine/Source/Runtime/ApplicationCore/Private/IOS/IOSAppDelegate.cpp b/Engine/Source/Runtime/ApplicationCore/Private/IOS/IOSAppDelegate.cpp
--- a/Engine/Source/Runtime/ApplicationCore/Private/IOS/IOSAppDelegate.cpp
+++ b/Engine/Source/Runtime/ApplicationCore/Private/IOS/IOSAppDelegate.cpp
@@ -524,7 +524,7 @@ static IOSAppDelegate* CachedDelegate = nil;
 					
 					if ([self IsFeatureActive:EAudioFeature::VoiceChat])
 					{
-						Mode = AVAudioSessionModeVoiceChat;
+						Mode = self.bHighQualityVoiceChatEnabled ? AVAudioSessionModeVoiceChat : AVAudioSessionModeDefault;
 #if !PLATFORM_TVOS
 						// allow bluetooth for chatting and such
 						Options |= AVAudioSessionCategoryOptionDefaultToSpeaker | AVAudioSessionCategoryOptionAllowBluetooth | AVAudioSessionCategoryOptionAllowBluetoothA2DP;
@@ -604,7 +604,8 @@ static IOSAppDelegate* CachedDelegate = nil;
 						
 						if (@available(iOS 10, *))
 						{
-							[[AVAudioSession sharedInstance] setCategory:AVAudioSessionCategoryPlayAndRecord mode:AVAudioSessionModeDefault options:opts error:&ActiveError];
+							AVAudioSessionMode VoiceChatMode = self.bHighQualityVoiceChatEnabled ? AVAudioSessionModeVoiceChat : AVAudioSessionModeDefault;
+							[[AVAudioSession sharedInstance] setCategory:AVAudioSessionCategoryPlayAndRecord mode:VoiceChatMode options:opts error:&ActiveError];
 						}
 						else
 						{
@@ -677,7 +678,8 @@ static IOSAppDelegate* CachedDelegate = nil;
 					
 					if (@available(iOS 10, *))
 					{
-						[[AVAudioSession sharedInstance] setCategory:AVAudioSessionCategoryPlayAndRecord mode:AVAudioSessionModeDefault options:opts error:&ActiveError];
+						AVAudioSessionMode VoiceChatMode = self.bHighQualityVoiceChatEnabled ? AVAudioSessionModeVoiceChat : AVAudioSessionModeDefault;
+						[[AVAudioSession sharedInstance] setCategory:AVAudioSessionCategoryPlayAndRecord mode:VoiceChatMode options:opts error:&ActiveError];
 					}
 					else
 					{
@@ -707,7 +709,8 @@ static IOSAppDelegate* CachedDelegate = nil;
 			// Necessary for voice chat if audio is not active
 			if (@available(iOS 10, *))
 			{
-				[[AVAudioSession sharedInstance] setCategory:AVAudioSessionCategoryPlayAndRecord mode:AVAudioSessionModeDefault options:opts error:&ActiveError];
+				AVAudioSessionMode VoiceChatMode = self.bHighQualityVoiceChatEnabled ? AVAudioSessionModeVoiceChat : AVAudioSessionModeDefault;
+				[[AVAudioSession sharedInstance] setCategory:AVAudioSessionCategoryPlayAndRecord mode:VoiceChatMode options:opts error:&ActiveError];
 			}
 			else
 			{
@@ -739,6 +742,11 @@ static IOSAppDelegate* CachedDelegate = nil;
 	return Session.otherAudioPlaying;
 }
 
+-(void)EnableHighQualityVoiceChat:(bool)bEnable
+{
+	self.bHighQualityVoiceChatEnabled = bEnable;
+}
+
 - (void)EnableVoiceChat:(bool)bEnable
 {
 	self.bVoiceChatEnabled = false;
diff --git a/Engine/Source/Runtime/ApplicationCore/Public/IOS/IOSAppDelegate.h b/Engine/Source/Runtime/ApplicationCore/Public/IOS/IOSAppDelegate.h
--- a/Engine/Source/Runtime/ApplicationCore/Public/IOS/IOSAppDelegate.h
+++ b/Engine/Source/Runtime/ApplicationCore/Public/IOS/IOSAppDelegate.h
@@ -192,6 +192,7 @@ APPLICATIONCORE_API
 - (void)ToggleAudioSession:(bool)bActive force:(bool)bForce;
 - (bool)IsBackgroundAudioPlaying;
 - (void)EnableVoiceChat:(bool)bEnable;
+- (void)EnableHighQualityVoiceChat:(bool)bEnable;
 - (bool)IsVoiceChatEnabled;
 
 - (void)SetFeature:(EAudioFeature)Feature Active:(bool)bIsActive;
@@ -199,6 +200,7 @@ APPLICATIONCORE_API
 
 @property (atomic) bool bAudioActive;
 @property (atomic) bool bVoiceChatEnabled;
+@property (atomic) bool bHighQualityVoiceChatEnabled;
 
 @property (atomic) bool bIsSuspended;
 @property (atomic) bool bHasSuspended;
