- Fixed compile error on builder

diff --git a/Engine/Source/Runtime/ApplicationCore/Private/IOS/IOSAppDelegate.cpp b/Engine/Source/Runtime/ApplicationCore/Private/IOS/IOSAppDelegate.cpp
--- a/Engine/Source/Runtime/ApplicationCore/Private/IOS/IOSAppDelegate.cpp
+++ b/Engine/Source/Runtime/ApplicationCore/Private/IOS/IOSAppDelegate.cpp
@@ -604,7 +604,7 @@ static IOSAppDelegate* CachedDelegate = nil;
 						
 						if (@available(iOS 10, *))
 						{
-							AVAudioSessionMode VoiceChatMode = self.bHighQualityVoiceChatEnabled ? AVAudioSessionModeVoiceChat : AVAudioSessionModeDefault;
+							NSString* VoiceChatMode = self.bHighQualityVoiceChatEnabled ? AVAudioSessionModeVoiceChat : AVAudioSessionModeDefault;
 							[[AVAudioSession sharedInstance] setCategory:AVAudioSessionCategoryPlayAndRecord mode:VoiceChatMode options:opts error:&ActiveError];
 						}
 						else
@@ -678,7 +678,7 @@ static IOSAppDelegate* CachedDelegate = nil;
 					
 					if (@available(iOS 10, *))
 					{
-						AVAudioSessionMode VoiceChatMode = self.bHighQualityVoiceChatEnabled ? AVAudioSessionModeVoiceChat : AVAudioSessionModeDefault;
+						NSString* VoiceChatMode = self.bHighQualityVoiceChatEnabled ? AVAudioSessionModeVoiceChat : AVAudioSessionModeDefault;
 						[[AVAudioSession sharedInstance] setCategory:AVAudioSessionCategoryPlayAndRecord mode:VoiceChatMode options:opts error:&ActiveError];
 					}
 					else
@@ -709,7 +709,7 @@ static IOSAppDelegate* CachedDelegate = nil;
 			// Necessary for voice chat if audio is not active
 			if (@available(iOS 10, *))
 			{
-				AVAudioSessionMode VoiceChatMode = self.bHighQualityVoiceChatEnabled ? AVAudioSessionModeVoiceChat : AVAudioSessionModeDefault;
+				NSString* VoiceChatMode = self.bHighQualityVoiceChatEnabled ? AVAudioSessionModeVoiceChat : AVAudioSessionModeDefault;
 				[[AVAudioSession sharedInstance] setCategory:AVAudioSessionCategoryPlayAndRecord mode:VoiceChatMode options:opts error:&ActiveError];
 			}
 			else
