[iOS] Add method to query if we have the recording permission

diff --git a/Engine/Source/Runtime/ApplicationCore/Private/IOS/IOSAppDelegate.cpp b/Engine/Source/Runtime/ApplicationCore/Private/IOS/IOSAppDelegate.cpp
--- a/Engine/Source/Runtime/ApplicationCore/Private/IOS/IOSAppDelegate.cpp
+++ b/Engine/Source/Runtime/ApplicationCore/Private/IOS/IOSAppDelegate.cpp
@@ -752,6 +752,11 @@ static IOSAppDelegate* CachedDelegate = nil;
 	return Session.otherAudioPlaying;
 }
 
+-(bool)HasRecordPermission
+{
+	return [[AVAudioSession sharedInstance] recordPermission] == AVAudioSessionRecordPermissionGranted;
+}
+
 -(void)EnableHighQualityVoiceChat:(bool)bEnable
 {
 	self.bHighQualityVoiceChatEnabled = bEnable;
diff --git a/Engine/Source/Runtime/ApplicationCore/Public/IOS/IOSAppDelegate.h b/Engine/Source/Runtime/ApplicationCore/Public/IOS/IOSAppDelegate.h
--- a/Engine/Source/Runtime/ApplicationCore/Public/IOS/IOSAppDelegate.h
+++ b/Engine/Source/Runtime/ApplicationCore/Public/IOS/IOSAppDelegate.h
@@ -195,6 +195,7 @@ APPLICATIONCORE_API
 - (void)InitializeAudioSession;
 - (void)ToggleAudioSession:(bool)bActive force:(bool)bForce;
 - (bool)IsBackgroundAudioPlaying;
+- (bool)HasRecordPermission;
 - (void)EnableVoiceChat:(bool)bEnable;
 - (void)EnableHighQualityVoiceChat:(bool)bEnable;
 - (bool)IsVoiceChatEnabled;
