[TVOS] Always return false when querying recording permission.

diff --git a/Engine/Source/Runtime/ApplicationCore/Private/IOS/IOSAppDelegate.cpp b/Engine/Source/Runtime/ApplicationCore/Private/IOS/IOSAppDelegate.cpp
--- a/Engine/Source/Runtime/ApplicationCore/Private/IOS/IOSAppDelegate.cpp
+++ b/Engine/Source/Runtime/ApplicationCore/Private/IOS/IOSAppDelegate.cpp
@@ -754,7 +754,12 @@ static IOSAppDelegate* CachedDelegate = nil;
 
 -(bool)HasRecordPermission
 {
+#if PLATFORM_TVOS
+	// TVOS does not have sound recording capabilities.
+	return false;
+#else
 	return [[AVAudioSession sharedInstance] recordPermission] == AVAudioSessionRecordPermissionGranted;
+#endif
 }
 
 -(void)EnableHighQualityVoiceChat:(bool)bEnable
