Set Platform.AppHasFocus property for the crash reporter on Mac

diff --git a/Engine/Source/Runtime/ApplicationCore/Private/Mac/MacApplication.cpp b/Engine/Source/Runtime/ApplicationCore/Private/Mac/MacApplication.cpp
--- a/Engine/Source/Runtime/ApplicationCore/Private/Mac/MacApplication.cpp
+++ b/Engine/Source/Runtime/ApplicationCore/Private/Mac/MacApplication.cpp
@@ -18,6 +18,7 @@
 #include "Mac/MacPlatformApplicationMisc.h"
 #include "HAL/ThreadHeartBeat.h"
 #include "IHapticDevice.h"
+#include "Apple/ApplePlatformCrashContext.h"
 
 #include <IOKit/IOKitLib.h>
 #include <IOKit/graphics/IOGraphicsLib.h>
@@ -1305,6 +1306,11 @@ void FMacApplication::OnApplicationDidBecomeActive()
 				}
 			}
 		}
+
+#if !UE_SERVER
+		// For non-editor clients, record if the active window is in focus
+		FGenericCrashContext::SetEngineData(TEXT("Platform.AppHasFocus"), TEXT("true"));
+#endif
 	}, @[ NSDefaultRunLoopMode ], false);
 }
 
@@ -1356,6 +1362,11 @@ void FMacApplication::OnApplicationWillResignActive()
 		{
 			MessageHandler->OnApplicationActivationChanged(false);
 		}
+
+#if !UE_SERVER
+		// For non-editor clients, record if the active window is in focus
+		FGenericCrashContext::SetEngineData(TEXT("Platform.AppHasFocus"), TEXT("false"));
+#endif
 	}, @[ NSDefaultRunLoopMode ], false);
 }
 
