init bHighQualityVoiceChatEnabled

diff --git a/Engine/Source/Runtime/ApplicationCore/Private/IOS/IOSAppDelegate.cpp b/Engine/Source/Runtime/ApplicationCore/Private/IOS/IOSAppDelegate.cpp
--- a/Engine/Source/Runtime/ApplicationCore/Private/IOS/IOSAppDelegate.cpp
+++ b/Engine/Source/Runtime/ApplicationCore/Private/IOS/IOSAppDelegate.cpp
@@ -210,6 +210,7 @@ static IOSAppDelegate* CachedDelegate = nil;
 	CachedDelegate = self;
 	// default to old style
 	memset(GEnabledAudioFeatures, 0, sizeof(GEnabledAudioFeatures));
+	bHighQualityVoiceChatEnabled = false;
 	return self;
 }
 
