Less contention on iOS input stack. Added additional stat to iOS input processing

diff --git a/Engine/Source/Runtime/ApplicationCore/Private/IOS/IOSInputInterface.cpp b/Engine/Source/Runtime/ApplicationCore/Private/IOS/IOSInputInterface.cpp
--- a/Engine/Source/Runtime/ApplicationCore/Private/IOS/IOSInputInterface.cpp
+++ b/Engine/Source/Runtime/ApplicationCore/Private/IOS/IOSInputInterface.cpp
@@ -25,9 +25,10 @@ static TAutoConsoleVariable<float> CVarHapticsRest(TEXT("ios.VibrationHapticsRes
 //
 //@end
 
-TArray<TouchInput> FIOSInputInterface::TouchInputStack = TArray<TouchInput>();
-TArray<int32> FIOSInputInterface::KeyInputStack;
-FCriticalSection FIOSInputInterface::CriticalSection;
+// protects the input stack used on 2 threads
+static FCriticalSection CriticalSection;
+static TArray<TouchInput> TouchInputStack;
+static TArray<int32> KeyInputStack;
 
 TSharedRef< FIOSInputInterface > FIOSInputInterface::Create(  const TSharedRef< FGenericApplicationMessageHandler >& InMessageHandler )
 {
@@ -258,11 +259,13 @@ void ModifyVectorByOrientation(FVector& Vec, bool bIsRotation)
 }
 #endif
 
-void FIOSInputInterface::ProcessTouchesAndKeys(uint32 ControllerId)
+void FIOSInputInterface::ProcessTouchesAndKeys(uint32 ControllerId, const TArray<TouchInput>& InTouchInputStack, const TArray<int32>& InKeyInputStack)
 {
-	for(int i = 0; i < TouchInputStack.Num(); ++i)
+	QUICK_SCOPE_CYCLE_COUNTER(STAT_IOSInputInterface_ProcessTouchesAndKeys);
+	
+	for(int i = 0; i < InTouchInputStack.Num(); ++i)
 	{
-		const TouchInput& Touch = TouchInputStack[i];
+		const TouchInput& Touch = InTouchInputStack[i];
 		
 		// send input to handler
 		if (Touch.Type == TouchBegan)
@@ -287,30 +290,32 @@ void FIOSInputInterface::ProcessTouchesAndKeys(uint32 ControllerId)
 		}
 	}
 	
-	TouchInputStack.Empty(0);
-	
-	
 	// these come in pairs
-	for(int32 KeyIndex = 0; KeyIndex < KeyInputStack.Num(); KeyIndex+=2)
+	for(int32 KeyIndex = 0; KeyIndex < InKeyInputStack.Num(); KeyIndex+=2)
 	{
-		int32 KeyCode = KeyInputStack[KeyIndex];
-		int32 CharCode = KeyInputStack[KeyIndex + 1];
+		int32 KeyCode = InKeyInputStack[KeyIndex];
+		int32 CharCode = InKeyInputStack[KeyIndex + 1];
 		MessageHandler->OnKeyDown(KeyCode, CharCode, false);
 		MessageHandler->OnKeyChar(CharCode,  false);
 		MessageHandler->OnKeyUp  (KeyCode, CharCode, false);
 	}
-	KeyInputStack.Empty(0);
 }
 
 void FIOSInputInterface::SendControllerEvents()
 {
-	FScopeLock Lock(&CriticalSection);
-
+	TArray<TouchInput> LocalTouchInputStack;
+	TArray<int32> LocalKeyInputStack;
+	{
+		FScopeLock Lock(&CriticalSection);
+		Exchange(LocalTouchInputStack, TouchInputStack);
+		Exchange(LocalKeyInputStack, KeyInputStack);
+	}
+	
 	int32 ControllerIndex = -1;
 	
 #if !PLATFORM_TVOS
 	// on ios, touches always go go player 0
-	ProcessTouchesAndKeys(0);
+	ProcessTouchesAndKeys(0, LocalTouchInputStack, LocalKeyInputStack);
 #endif
 
 	
@@ -491,13 +496,11 @@ if ((Controller.Previous##Gamepad != nil && Gamepad.GCAxis.value != Controller.P
 				HANDLE_BUTTON(MicroGamepad, dpad.down,	FGamepadKeyNames::LeftStickDown);
 				HANDLE_BUTTON(MicroGamepad, dpad.right,	FGamepadKeyNames::LeftStickRight);
 				HANDLE_BUTTON(MicroGamepad, dpad.left,	FGamepadKeyNames::LeftStickLeft);
-
-				TouchInputStack.Empty(0);
 			}
 			// otherwise, process touches like ios for the remote's index
 			else
 			{
-				ProcessTouchesAndKeys(Cont.playerIndex);
+				ProcessTouchesAndKeys(Cont.playerIndex, LocalTouchInputStack, LocalKeyInputStack);
 			}
 			
 			HANDLE_BUTTON(MicroGamepad, buttonA,	FGamepadKeyNames::FaceButtonBottom);
@@ -566,7 +569,7 @@ void FIOSInputInterface::QueueTouchInput(const TArray<TouchInput>& InTouchEvents
 {
 	FScopeLock Lock(&CriticalSection);
 
-	FIOSInputInterface::TouchInputStack.Append(InTouchEvents);
+	TouchInputStack.Append(InTouchEvents);
 }
 
 void FIOSInputInterface::QueueKeyInput(int32 Key, int32 Char)
@@ -574,8 +577,8 @@ void FIOSInputInterface::QueueKeyInput(int32 Key, int32 Char)
 	FScopeLock Lock(&CriticalSection);
 
 	// put the key and char into the array
-	FIOSInputInterface::KeyInputStack.Add(Key);
-	FIOSInputInterface::KeyInputStack.Add(Char);
+	KeyInputStack.Add(Key);
+	KeyInputStack.Add(Char);
 }
 
 void FIOSInputInterface::EnableMotionData(bool bEnable)
diff --git a/Engine/Source/Runtime/ApplicationCore/Public/IOS/IOSInputInterface.h b/Engine/Source/Runtime/ApplicationCore/Public/IOS/IOSInputInterface.h
--- a/Engine/Source/Runtime/ApplicationCore/Public/IOS/IOSInputInterface.h
+++ b/Engine/Source/Runtime/ApplicationCore/Public/IOS/IOSInputInterface.h
@@ -109,15 +109,8 @@ private:
 	void CalibrateMotion(uint32 PlayerIndex);
 
 private:
-	void ProcessTouchesAndKeys(uint32 ControllerId);
-
-
-	static TArray<TouchInput> TouchInputStack;
-	static TArray<int32> KeyInputStack;
-
-	// protects the input stack used on 2 threads
-	static FCriticalSection CriticalSection;
-
+	void ProcessTouchesAndKeys(uint32 ControllerId, const TArray<TouchInput>& InTouchInputStack, const TArray<int32>& InKeyInputStack);
+	   
 	TSharedRef< FGenericApplicationMessageHandler > MessageHandler;
 
 
