Android event handling refactor.
Solve issues during start-up where the HW window can be created/lost before the game can use it or OS displaying android ANR message if GT is too busy to respond to system events.

Window access is guarded with mutex that guarantees a valid window HW exists when locked.
Rendering thread is blocked when the window is lost (to prevent RHIs accessing surfaces belonging to the lost window)
Window dimensions are cached so that most window dimension queries just read from the cache.

diff --git a/Engine/Source/Runtime/ApplicationCore/Private/Android/AndroidApplication.cpp b/Engine/Source/Runtime/ApplicationCore/Private/Android/AndroidApplication.cpp
--- a/Engine/Source/Runtime/ApplicationCore/Private/Android/AndroidApplication.cpp
+++ b/Engine/Source/Runtime/ApplicationCore/Private/Android/AndroidApplication.cpp
@@ -70,9 +70,9 @@ void FAndroidApplication::PollGameDeviceState( const float TimeDelta )
 	InputInterface->Tick( TimeDelta );
 	InputInterface->SendControllerEvents();
 	
+	// Note: We're using window size values cached on GT. We dont wait for a valid HW window.
 	if (bWindowSizeChanged && 
-		Windows.Num() > 0 && 
-		FAndroidWindow::GetHardwareWindow() != nullptr)
+		Windows.Num() > 0)
 	{
 		int32 WindowX,WindowY, WindowWidth,WindowHeight;
 		Windows[0]->GetFullScreenInfo(WindowX, WindowY, WindowWidth, WindowHeight);
diff --git a/Engine/Source/Runtime/ApplicationCore/Private/Android/AndroidCursor.cpp b/Engine/Source/Runtime/ApplicationCore/Private/Android/AndroidCursor.cpp
--- a/Engine/Source/Runtime/ApplicationCore/Private/Android/AndroidCursor.cpp
+++ b/Engine/Source/Runtime/ApplicationCore/Private/Android/AndroidCursor.cpp
@@ -17,7 +17,7 @@ void FAndroidCursor::ComputeUIScaleFactor()
 {
 	FPlatformRect ScreenRect = FAndroidWindow::GetScreenRect();
 	int32_t ScreenWidth, ScreenHeight;
-	FAndroidWindow::CalculateSurfaceSize(FAndroidWindow::GetHardwareWindow(), ScreenWidth, ScreenHeight);
+	FAndroidWindow::CalculateSurfaceSize(ScreenWidth, ScreenHeight);
 
 	if (ScreenWidth > ScreenHeight)
 	{
diff --git a/Engine/Source/Runtime/ApplicationCore/Private/Android/AndroidPlatformApplicationMisc.cpp b/Engine/Source/Runtime/ApplicationCore/Private/Android/AndroidPlatformApplicationMisc.cpp
--- a/Engine/Source/Runtime/ApplicationCore/Private/Android/AndroidPlatformApplicationMisc.cpp
+++ b/Engine/Source/Runtime/ApplicationCore/Private/Android/AndroidPlatformApplicationMisc.cpp
@@ -172,7 +172,7 @@ static float GetWindowUpscaleFactor()
 	// and return that scalar.
 
 	int32_t SurfaceWidth, SurfaceHeight;
-	FAndroidWindow::CalculateSurfaceSize(FAndroidWindow::GetHardwareWindow(), SurfaceWidth, SurfaceHeight);
+	FAndroidWindow::CalculateSurfaceSize(SurfaceWidth, SurfaceHeight);
 
 	FPlatformRect ScreenRect = FAndroidWindow::GetScreenRect();
 
diff --git a/Engine/Source/Runtime/ApplicationCore/Private/Android/AndroidPlatformSurvey.cpp b/Engine/Source/Runtime/ApplicationCore/Private/Android/AndroidPlatformSurvey.cpp
--- a/Engine/Source/Runtime/ApplicationCore/Private/Android/AndroidPlatformSurvey.cpp
+++ b/Engine/Source/Runtime/ApplicationCore/Private/Android/AndroidPlatformSurvey.cpp
@@ -24,7 +24,7 @@ bool FAndroidPlatformSurvey::GetSurveyResults(FHardwareSurveyResults& OutResults
 
 	// display 0 is max size
 	int32 ScreenWidth, ScreenHeight;
-	FAndroidWindow::CalculateSurfaceSize(FAndroidWindow::GetHardwareWindow(), ScreenWidth, ScreenHeight);
+	FAndroidWindow::CalculateSurfaceSize(ScreenWidth, ScreenHeight);
 
 	OutResults.Displays[0].CurrentModeWidth = ScreenWidth;
 	OutResults.Displays[0].CurrentModeHeight = ScreenHeight;
diff --git a/Engine/Source/Runtime/ApplicationCore/Private/Android/AndroidWindow.cpp b/Engine/Source/Runtime/ApplicationCore/Private/Android/AndroidWindow.cpp
--- a/Engine/Source/Runtime/ApplicationCore/Private/Android/AndroidWindow.cpp
+++ b/Engine/Source/Runtime/ApplicationCore/Private/Android/AndroidWindow.cpp
@@ -14,15 +14,33 @@
 #include "Misc/CommandLine.h"
 #include "HAL/PlatformStackWalk.h"
 
+struct FCachedWindowRect
+{
+	FCachedWindowRect() : WindowWidth(-1), WindowHeight(-1), WindowInit(false), ContentScaleFactor(-1.0f), bLastMosaicState(false), Window_EventThread(nullptr)
+	{
+	}
+
+	int32 WindowWidth;
+	int32 WindowHeight;
+	bool WindowInit;
+	float ContentScaleFactor;
+	bool bLastMosaicState;
+	ANativeWindow* Window_EventThread;
+};
+
+
 // Cached calculated screen resolution
-static int32 WindowWidth = -1;
-static int32 WindowHeight = -1;
+static FCachedWindowRect CachedWindowRect;
+static FCachedWindowRect CachedWindowRect_EventThread;
+
+static void ClearCachedWindowRects()
+{
+	CachedWindowRect.WindowInit = false;
+	CachedWindowRect_EventThread.WindowInit = false;
+}
+
 static int32 GSurfaceViewWidth = -1;
 static int32 GSurfaceViewHeight = -1;
-static bool WindowInit = false;
-static float ContentScaleFactor = -1.0f;
-static ANativeWindow* LastWindow = NULL;
-static bool bLastMosaicState = false;
 
 void* FAndroidWindow::NativeWindow = NULL;
 
@@ -36,17 +54,14 @@ TSharedRef<FAndroidWindow> FAndroidWindow::Make()
 	return MakeShareable( new FAndroidWindow() );
 }
 
-FAndroidWindow::FAndroidWindow() :Window(NULL)
+FAndroidWindow::FAndroidWindow()
 {
 }
 
 void FAndroidWindow::Initialize( class FAndroidApplication* const Application, const TSharedRef< FGenericWindowDefinition >& InDefinition, const TSharedPtr< FAndroidWindow >& InParent, const bool bShowImmediately )
 {
-	//set window here.
-
 	OwningApplication = Application;
 	Definition = InDefinition;
-	Window = static_cast<ANativeWindow*>(FAndroidWindow::GetHardwareWindow());
 }
 
 bool FAndroidWindow::GetFullScreenInfo( int32& X, int32& Y, int32& Width, int32& Height ) const
@@ -64,7 +79,8 @@ bool FAndroidWindow::GetFullScreenInfo( int32& X, int32& Y, int32& Width, int32&
 
 void FAndroidWindow::SetOSWindowHandle(void* InWindow)
 {
-	Window = static_cast<ANativeWindow*>(InWindow);
+	// not expecting a window to be supplied on android.
+	check(InWindow == nullptr);
 }
 
 
@@ -76,7 +92,7 @@ static FVector4 GAndroidLandscapeSafezone = FVector4(-1.0f, -1.0f, -1.0f, -1.0f)
 #if USE_ANDROID_JNI
 JNI_METHOD void Java_com_epicgames_ue4_GameActivity_nativeSetWindowInfo(JNIEnv* jenv, jobject thiz, jboolean bIsPortrait, jint DepthBufferPreference)
 {
-	WindowInit = false;
+	ClearCachedWindowRects();
 	GAndroidIsPortrait = bIsPortrait == JNI_TRUE;
 	GAndroidDepthBufferPreference = DepthBufferPreference;
 	FPlatformMisc::LowLevelOutputDebugStringf(TEXT("App is running in %s\n"), GAndroidIsPortrait ? TEXT("Portrait") : TEXT("Landscape"));
@@ -108,6 +124,10 @@ JNI_METHOD void Java_com_epicgames_ue4_GameActivity_nativeSetSafezoneInfo(JNIEnv
 }
 #endif
 
+bool FAndroidWindow::bAreCachedNativeDimensionsValid = false;
+int32 FAndroidWindow::CachedNativeWindowWidth = -1;
+int32 FAndroidWindow::CachedNativeWindowHeight = -1;
+
 bool FAndroidWindow::IsPortraitOrientation()
 {
 	return GAndroidIsPortrait;
@@ -125,7 +145,7 @@ int32 FAndroidWindow::GetDepthBufferPreference()
 
 void FAndroidWindow::InvalidateCachedScreenRect()
 {
-	WindowInit = false;
+	ClearCachedWindowRects();
 }
 
 void FAndroidWindow::AcquireWindowRef(ANativeWindow* InWindow)
@@ -142,16 +162,74 @@ void FAndroidWindow::ReleaseWindowRef(ANativeWindow* InWindow)
 #endif
 }
 
-void FAndroidWindow::SetHardwareWindow(void* InWindow)
+ void FAndroidWindow::SetHardwareWindow_EventThread(void* InWindow)
 {
+	check(IsInAndroidEventThread());
 	NativeWindow = InWindow; //using raw native window handle for now. Could be changed to use AndroidWindow later if needed
 }
 
-void* FAndroidWindow::GetHardwareWindow()
+void* FAndroidWindow::GetHardwareWindow_EventThread()
 {
 	return NativeWindow;
 }
 
+bool FAndroidWindow::WaitForWindowDimensions()
+{
+	while (!bAreCachedNativeDimensionsValid)
+	{
+		if (GIsRequestingExit
+#if USE_ANDROID_EVENTS
+		|| FAppEventManager::GetInstance()->WaitForEventInQueue(EAppEventState::APP_EVENT_STATE_ON_DESTROY, 0.0f)
+#endif
+			)
+		{
+			// Application is shutting down
+			return false;
+		}
+		FPlatformProcess::Sleep(0.001f);
+	}
+	return true;
+}
+
+// To be called during initialization from the event thread.
+// once set the dimensions are 'valid' and further changes are updated via FAppEventManager::Tick 
+void FAndroidWindow::SetWindowDimensions_EventThread(ANativeWindow* DimensionWindow)
+{
+	if(bAreCachedNativeDimensionsValid == false)
+	{
+#if USE_ANDROID_JNI
+		CachedNativeWindowWidth = ANativeWindow_getWidth(DimensionWindow);
+		CachedNativeWindowHeight = ANativeWindow_getHeight(DimensionWindow);
+#else // Lumin case.
+		int32 ResWidth, ResHeight;
+		if (FPlatformMisc::GetOverrideResolution(ResWidth, ResHeight))
+		{
+			CachedNativeWindowWidth = ResWidth;
+			CachedNativeWindowHeight = ResHeight;
+		}
+#endif
+		FPlatformMisc::MemoryBarrier();
+		bAreCachedNativeDimensionsValid = true;
+	}
+}
+
+// Update the dimensions from FAppEventManager::Tick based on messages posted from the event thread.
+void FAndroidWindow::EventManagerUpdateWindowDimensions(int32 Width, int32 Height)
+{
+	check(bAreCachedNativeDimensionsValid);
+	check(Width >= 0 && Height >= 0);
+
+	bool bChanged = CachedNativeWindowWidth != Width || CachedNativeWindowHeight != Height;
+
+	CachedNativeWindowWidth = Width;
+	CachedNativeWindowHeight = Height;
+
+	if (bChanged)
+	{
+		InvalidateCachedScreenRect();
+	}
+}
+
 void* FAndroidWindow::WaitForHardwareWindow()
 {
 	// Sleep if the hardware window isn't currently available.
@@ -167,8 +245,8 @@ void* FAndroidWindow::WaitForHardwareWindow()
 	// It is not sufficient to check the GIsRequestingExit global variable, as the handler reacting to the APP_EVENT_STATE_ON_DESTROY
 	// may be running in the same thread as this method and therefore lead to a deadlock.
 
-	void* Window = GetHardwareWindow();
-	while (Window == nullptr)
+	void* WindowEventThread = GetHardwareWindow_EventThread();
+	while (WindowEventThread == nullptr)
 	{
 #if USE_ANDROID_EVENTS
 		if (GIsRequestingExit || FAppEventManager::GetInstance()->WaitForEventInQueue(EAppEventState::APP_EVENT_STATE_ON_DESTROY, 0.0f))
@@ -178,62 +256,71 @@ void* FAndroidWindow::WaitForHardwareWindow()
 		}
 #endif
 		FPlatformProcess::Sleep(0.001f);
-		Window = GetHardwareWindow();
+		WindowEventThread = GetHardwareWindow_EventThread();
 	}
-	return Window;
+	return WindowEventThread;
 }
 
 #if USE_ANDROID_JNI
 extern bool AndroidThunkCpp_IsOculusMobileApplication();
 #endif
 
-bool FAndroidWindow::IsCachedRectValid(const bool bMosaicEnabled, const float RequestedContentScaleFactor, ANativeWindow* Window)
+bool FAndroidWindow::IsCachedRectValid(bool bUseEventThreadWindow, const bool bMosaicEnabled, const float RequestedContentScaleFactor, ANativeWindow* Window)
 {
-	if (!WindowInit)
+	// window must be valid when bUseEventThreadWindow and null when !bUseEventThreadWindow.
+	check((Window != nullptr) == bUseEventThreadWindow);
+
+	const FCachedWindowRect& CachedRect = bUseEventThreadWindow ? CachedWindowRect_EventThread : CachedWindowRect;
+
+	if (!CachedRect.WindowInit)
 	{
 		return false;
 	}
 
 	bool bValidCache = true;
 
-	if (bLastMosaicState != bMosaicEnabled)
+	if (CachedRect.bLastMosaicState != bMosaicEnabled)
 	{
-		FPlatformMisc::LowLevelOutputDebugStringf(TEXT("***** Mosaic State change (to %s), not using res cache"), bMosaicEnabled ? TEXT("enabled") : TEXT("disabled"));
+		FPlatformMisc::LowLevelOutputDebugStringf(TEXT("***** Mosaic State change (to %s), not using res cache (%d)"), bMosaicEnabled ? TEXT("enabled") : TEXT("disabled"), (int32)bUseEventThreadWindow);
 		bValidCache = false;
 	}
 
-	if (RequestedContentScaleFactor != ContentScaleFactor)
+	if (CachedRect.ContentScaleFactor != RequestedContentScaleFactor )
 	{
-		FPlatformMisc::LowLevelOutputDebugStringf(TEXT("***** RequestedContentScaleFactor different %f != %f, not using res cache"), RequestedContentScaleFactor, ContentScaleFactor);
+		FPlatformMisc::LowLevelOutputDebugStringf(TEXT("***** RequestedContentScaleFactor different %f != %f, not using res cache (%d)"), RequestedContentScaleFactor, CachedWindowRect.ContentScaleFactor, (int32)bUseEventThreadWindow);
 		bValidCache = false;
 	}
 
-	if (Window != LastWindow)
+	if (CachedRect.Window_EventThread != Window)
 	{
-		FPlatformMisc::LowLevelOutputDebugString(TEXT("***** Window different, not using res cache"));
+		FPlatformMisc::LowLevelOutputDebugStringf(TEXT("***** Window different, not using res cache (%d)"), (int32)bUseEventThreadWindow);
 		bValidCache = false;
 	}
 
-	if (WindowWidth <= 8)
+	if (CachedRect.WindowWidth <= 8)
 	{
-		FPlatformMisc::LowLevelOutputDebugStringf(TEXT("***** WindowWidth is %d, not using res cache"), WindowWidth);
+		FPlatformMisc::LowLevelOutputDebugStringf(TEXT("***** WindowWidth is %d, not using res cache (%d)"), CachedRect.WindowWidth, (int32)bUseEventThreadWindow);
 		bValidCache = false;
 	}
 
 	return bValidCache;
 }
 
-void FAndroidWindow::CacheRect(ANativeWindow* Window, const int32 Width, const int32 Height, const float RequestedContentScaleFactor, const bool bMosaicEnabled)
+void FAndroidWindow::CacheRect(bool bUseEventThreadWindow, const int32 Width, const int32 Height, const float RequestedContentScaleFactor, const bool bMosaicEnabled, ANativeWindow* Window)
 {
-	WindowWidth = Width;
-	WindowHeight = Height;
-	WindowInit = true;
-	ContentScaleFactor = RequestedContentScaleFactor;
-	LastWindow = Window;
-	bLastMosaicState = bMosaicEnabled;
+	check(Window != nullptr || !bUseEventThreadWindow);
+
+	FCachedWindowRect& CachedRect = bUseEventThreadWindow ? CachedWindowRect_EventThread : CachedWindowRect;
+
+	CachedRect.WindowWidth = Width;
+	CachedRect.WindowHeight = Height;
+	CachedRect.WindowInit = true;
+	CachedRect.ContentScaleFactor = RequestedContentScaleFactor;
+	CachedRect.bLastMosaicState = bMosaicEnabled;
+	CachedRect.Window_EventThread = Window;
 }
 
-FPlatformRect FAndroidWindow::GetScreenRect()
+FPlatformRect FAndroidWindow::GetScreenRect(bool bUseEventThreadWindow)
 {
 	int32 OverrideResX, OverrideResY;
 	// allow a subplatform to dictate resolution - we can't easily subclass FAndroidWindow the way its used
@@ -255,27 +342,33 @@ FPlatformRect FAndroidWindow::GetScreenRect()
 #else
 
 	static const bool bIsOculusMobileApp = AndroidThunkCpp_IsOculusMobileApplication();
-
-	ANativeWindow* Window = (ANativeWindow*)FAndroidWindow::GetHardwareWindow();
 	static const bool bIsDaydreamApp = FAndroidMisc::IsDaydreamApplication();
-	if (bIsDaydreamApp && Window == NULL)
-	{
-		// Sleep if the hardware window isn't currently available.
-		FPlatformMisc::LowLevelOutputDebugString(TEXT("Waiting for Native window in FAndroidWindow::GetScreenRect"));
-		Window = (ANativeWindow*)FAndroidWindow::WaitForHardwareWindow();
-	}
-	//	check(Window != NULL);
-	if (Window == NULL)
+	if (bIsDaydreamApp)
 	{
-		FPlatformRect ScreenRect;
-		ScreenRect.Left = 0;
-		ScreenRect.Top = 0;
-		ScreenRect.Right = GAndroidIsPortrait ? 720 : 1280;
-		ScreenRect.Bottom = GAndroidIsPortrait ? 1280 : 720;
+		// TODO: confirm that this is no longer required.
+		ANativeWindow* Window = (ANativeWindow*)FAndroidWindow::GetHardwareWindow_EventThread();
+		if (Window == NULL)
+		{
+			// Sleep if the hardware window isn't currently available.
+			FPlatformMisc::LowLevelOutputDebugString(TEXT("Waiting for Native window in FAndroidWindow::GetScreenRect"));
+			Window = (ANativeWindow*)FAndroidWindow::WaitForHardwareWindow();
+		}
 
-		UE_LOG(LogAndroid, Log, TEXT("FAndroidWindow::GetScreenRect: Window was NULL, returned default resolution: %d x %d"), ScreenRect.Right, ScreenRect.Bottom);
+		if (Window == NULL)
+		{
+			FPlatformRect ScreenRect;
+			ScreenRect.Left = 0;
+			ScreenRect.Top = 0;
+			ScreenRect.Right = GAndroidIsPortrait ? 720 : 1280;
+			ScreenRect.Bottom = GAndroidIsPortrait ? 1280 : 720;
+
+			UE_LOG(LogAndroid, Log, TEXT("FAndroidWindow::GetScreenRect: Window was NULL, returned default resolution: %d x %d"), ScreenRect.Right, ScreenRect.Bottom);
+
+			return ScreenRect;
+		}
 
-		return ScreenRect;
+		// dont cache rect.
+		bUseEventThreadWindow = true;
 	}
 
 	// determine mosaic requirements:
@@ -292,15 +385,16 @@ FPlatformRect FAndroidWindow::GetScreenRect()
 		RequestedContentScaleFactor = FCString::Atof(*CmdLineCSF);
 	}
 
-	// since orientation won't change on Android, use cached results if still valid
-	bool bComputeRect = !IsCachedRectValid(bMosaicEnabled, RequestedContentScaleFactor, Window);
+	// since orientation won't change on Android, use cached results if still valid. Different cache is maintained for event_thread flavor.
+	ANativeWindow* Window = bUseEventThreadWindow ? (ANativeWindow*)FAndroidWindow::GetHardwareWindow_EventThread() : nullptr;
+	bool bComputeRect = !IsCachedRectValid(bUseEventThreadWindow, bMosaicEnabled, RequestedContentScaleFactor, Window);
 	if (bComputeRect)
 	{
 		// currently hardcoding resolution
 
 		// get the aspect ratio of the physical screen
 		int32 ScreenWidth, ScreenHeight;
-		CalculateSurfaceSize(Window, ScreenWidth, ScreenHeight);
+		CalculateSurfaceSize(ScreenWidth, ScreenHeight, bUseEventThreadWindow);
 
 		static auto* MobileHDRCvar = IConsoleManager::Get().FindTConsoleVariableDataInt(TEXT("r.MobileHDR"));
 		const bool bMobileHDR = (MobileHDRCvar && MobileHDRCvar->GetValueOnAnyThread() == 1);
@@ -318,21 +412,23 @@ FPlatformRect FAndroidWindow::GetScreenRect()
 		}
 
 		// save for future calls
-		CacheRect(Window, ScreenWidth, ScreenHeight, RequestedContentScaleFactor, bMosaicEnabled);
+		CacheRect(bUseEventThreadWindow, ScreenWidth, ScreenHeight, RequestedContentScaleFactor, bMosaicEnabled, Window);
 	}
 
+	const FCachedWindowRect& CachedRect = bUseEventThreadWindow ? CachedWindowRect_EventThread : CachedWindowRect;
+
 	// create rect and return
 	FPlatformRect ScreenRect;
 	ScreenRect.Left = 0;
 	ScreenRect.Top = 0;
-	ScreenRect.Right = WindowWidth;
-	ScreenRect.Bottom = WindowHeight;
+	ScreenRect.Right = CachedRect.WindowWidth;
+	ScreenRect.Bottom = CachedRect.WindowHeight;
 
 	return ScreenRect;
 #endif
 }
 
-void FAndroidWindow::CalculateSurfaceSize(void* InWindow, int32_t& SurfaceWidth, int32_t& SurfaceHeight)
+void FAndroidWindow::CalculateSurfaceSize(int32_t& SurfaceWidth, int32_t& SurfaceHeight, bool bUseEventThreadWindow)
 {
 	// allow a subplatform to dictate resolution - we can't easily subclass FAndroidWindow the way its used
 	if (FPlatformMisc::GetOverrideResolution(SurfaceWidth, SurfaceHeight))
@@ -347,56 +443,76 @@ void FAndroidWindow::CalculateSurfaceSize(void* InWindow, int32_t& SurfaceWidth,
 
 #else
 
-	static const bool bIsMobileVRApp = AndroidThunkCpp_IsOculusMobileApplication() || FAndroidMisc::IsDaydreamApplication();
-
-	ANativeWindow* Window = (ANativeWindow*)InWindow;
-
-	if (InWindow == nullptr)
+	// daydream case still using hardware window direct from event thread, this might not be required now.
+	if (FAndroidMisc::IsDaydreamApplication())
 	{
-		// log the issue and callstack for backtracking the issue
-		// dump the stack HERE
+		ANativeWindow* Window = (ANativeWindow*)GetHardwareWindow_EventThread();
+		if (Window == nullptr)
 		{
-			const SIZE_T StackTraceSize = 65535;
-			ANSICHAR* StackTrace = (ANSICHAR*)FMemory::Malloc(StackTraceSize);
-			StackTrace[0] = 0;
+			// log the issue and callstack for backtracking the issue
+			// dump the stack HERE
+			{
+				const SIZE_T StackTraceSize = 65535;
+				ANSICHAR* StackTrace = (ANSICHAR*)FMemory::Malloc(StackTraceSize);
+				StackTrace[0] = 0;
 
-			// Walk the stack and dump it to the allocated memory.
-			FPlatformStackWalk::StackWalkAndDump(StackTrace, StackTraceSize, 0, NULL);
+				// Walk the stack and dump it to the allocated memory.
+				FPlatformStackWalk::StackWalkAndDump(StackTrace, StackTraceSize, 0, NULL);
 
-			FPlatformMisc::LowLevelOutputDebugString(TEXT("== WARNNG: CalculateSurfaceSize called with NULL window:"));
+				FPlatformMisc::LowLevelOutputDebugString(TEXT("== WARNNG: CalculateSurfaceSize called with NULL hardware window:"));
 
-			ANSICHAR* Start = StackTrace;
-			ANSICHAR* Next = StackTrace;
-			FPlatformMisc::LowLevelOutputDebugString(TEXT("==> STACK TRACE"));
-			while (*Next)
-			{
+				ANSICHAR* Start = StackTrace;
+				ANSICHAR* Next = StackTrace;
+				FPlatformMisc::LowLevelOutputDebugString(TEXT("==> STACK TRACE"));
 				while (*Next)
 				{
-					if (*Next == 10 || *Next == 13)
+					while (*Next)
 					{
-						while (*Next == 10 || *Next == 13)
+						if (*Next == 10 || *Next == 13)
 						{
-							*Next++ = 0;
+							while (*Next == 10 || *Next == 13)
+							{
+								*Next++ = 0;
+							}
+							break;
 						}
-						break;
+						++Next;
 					}
-					++Next;
+					FPlatformMisc::LowLevelOutputDebugStringf(TEXT("==> %s"), ANSI_TO_TCHAR(Start));
+					Start = Next;
 				}
-				FPlatformMisc::LowLevelOutputDebugStringf(TEXT("==> %s"), ANSI_TO_TCHAR(Start));
-				Start = Next;
+				FPlatformMisc::LowLevelOutputDebugString(TEXT("<== STACK TRACE"));
+
+				FMemory::Free(StackTrace);
 			}
-			FPlatformMisc::LowLevelOutputDebugString(TEXT("<== STACK TRACE"));
 
-			FMemory::Free(StackTrace);
+			SurfaceWidth = (GSurfaceViewWidth > 0) ? GSurfaceViewWidth : 1280;
+			SurfaceHeight = (GSurfaceViewHeight > 0) ? GSurfaceViewHeight : 720;
+		}
+		else
+		{
+			SurfaceWidth = (GSurfaceViewWidth > 0) ? GSurfaceViewWidth : ANativeWindow_getWidth(Window);
+			SurfaceHeight = (GSurfaceViewHeight > 0) ? GSurfaceViewHeight : ANativeWindow_getHeight(Window);
 		}
-
-		SurfaceWidth = (GSurfaceViewWidth > 0) ? GSurfaceViewWidth : 1280;
-		SurfaceHeight = (GSurfaceViewHeight > 0) ? GSurfaceViewHeight : 720;
 	}
 	else
 	{
-		SurfaceWidth = (GSurfaceViewWidth > 0) ? GSurfaceViewWidth : ANativeWindow_getWidth(Window);
-		SurfaceHeight = (GSurfaceViewHeight > 0) ? GSurfaceViewHeight : ANativeWindow_getHeight(Window);
+		if (bUseEventThreadWindow)
+		{
+			check(IsInAndroidEventThread());
+			ANativeWindow* WindowEventThread = (ANativeWindow*)GetHardwareWindow_EventThread();
+			check(WindowEventThread);
+
+			SurfaceWidth = (GSurfaceViewWidth > 0) ? GSurfaceViewWidth : ANativeWindow_getWidth(WindowEventThread);
+			SurfaceHeight = (GSurfaceViewHeight > 0) ? GSurfaceViewHeight : ANativeWindow_getHeight(WindowEventThread);
+		}
+		else
+		{
+			FAndroidWindow::WaitForWindowDimensions();
+
+			SurfaceWidth = (GSurfaceViewWidth > 0) ? GSurfaceViewWidth : CachedNativeWindowWidth;
+			SurfaceHeight = (GSurfaceViewHeight > 0) ? GSurfaceViewHeight : CachedNativeWindowHeight;
+		}
 	}
 
 	// some phones gave it the other way (so, if swap if the app is landscape, but width < height)
@@ -409,6 +525,7 @@ void FAndroidWindow::CalculateSurfaceSize(void* InWindow, int32_t& SurfaceWidth,
 	// ensure the size is divisible by a specified amount
 	// do not convert to a surface size that is larger than native resolution
 	// Mobile VR doesn't need buffer quantization as UE4 never renders directly to the buffer in VR mode. 
+	static const bool bIsMobileVRApp = AndroidThunkCpp_IsOculusMobileApplication() || FAndroidMisc::IsDaydreamApplication();
 	const int DividableBy = bIsMobileVRApp ? 1 : 8;
 	SurfaceWidth = (SurfaceWidth / DividableBy) * DividableBy;
 	SurfaceHeight = (SurfaceHeight / DividableBy) * DividableBy;
diff --git a/Engine/Source/Runtime/ApplicationCore/Public/Android/AndroidApplication.h b/Engine/Source/Runtime/ApplicationCore/Public/Android/AndroidApplication.h
--- a/Engine/Source/Runtime/ApplicationCore/Public/Android/AndroidApplication.h
+++ b/Engine/Source/Runtime/ApplicationCore/Public/Android/AndroidApplication.h
@@ -15,7 +15,6 @@ namespace FAndroidAppEntry
 	// if the native window handle has changed then the new handle is required.
 	void ReInitWindow(void* NewNativeWindowHandle = nullptr);
 
-	void DestroyWindow();
 	void ReleaseEGL();
 	void OnPauseEvent();
 }
diff --git a/Engine/Source/Runtime/ApplicationCore/Public/Android/AndroidWindow.h b/Engine/Source/Runtime/ApplicationCore/Public/Android/AndroidWindow.h
--- a/Engine/Source/Runtime/ApplicationCore/Public/Android/AndroidWindow.h
+++ b/Engine/Source/Runtime/ApplicationCore/Public/Android/AndroidWindow.h
@@ -28,7 +28,7 @@ public:
 	static TSharedRef<FAndroidWindow> Make();
 
 	
-	virtual void* GetOSWindowHandle() const override { return Window; } //can be null.
+	virtual void* GetOSWindowHandle() const override { return nullptr; }
 
 	void Initialize( class FAndroidApplication* const Application, const TSharedRef< FGenericWindowDefinition >& InDefinition, const TSharedPtr< FAndroidWindow >& InParent, const bool bShowImmediately );
 
@@ -37,10 +37,12 @@ public:
 
 	virtual void SetOSWindowHandle(void*);
 
-	static FPlatformRect GetScreenRect();
+	static FPlatformRect GetScreenRect(bool bUseEventThreadWindow = false);
 	static void InvalidateCachedScreenRect();
 
-	static void CalculateSurfaceSize(void* InWindow, int32_t& SurfaceWidth, int32_t& SurfaceHeight);
+	// When bUseEventThreadWindow == false this uses dimensions cached when the game thread processes android events.
+	// When bUseEventThreadWindow == true this uses dimensions directly from the android event thread, unless called from event thread this requires acquiring GAndroidWindowLock to use.
+	static void CalculateSurfaceSize(int32_t& SurfaceWidth, int32_t& SurfaceHeight, bool bUseEventThreadWindow = false);
 	static bool OnWindowOrientationChanged(bool bIsPortrait);
 
 	static int32 GetDepthBufferPreference();
@@ -48,8 +50,9 @@ public:
 	static void AcquireWindowRef(ANativeWindow* InWindow);
 	static void ReleaseWindowRef(ANativeWindow* InWindow);
 
-	static void* GetHardwareWindow();
-	static void SetHardwareWindow(void* InWindow);
+	// This returns the current hardware window as set from the event thread.
+	static void* GetHardwareWindow_EventThread();
+	static void SetHardwareWindow_EventThread(void* InWindow);
 
 	/** Waits on the current thread for a hardware window and returns it. 
 	 *  May return nullptr if the application is shutting down.
@@ -59,6 +62,12 @@ public:
 	static bool IsPortraitOrientation();
 	static FVector4 GetSafezone(bool bPortrait);
 
+	// called by the Android event thread to initially set the current window dimensions.
+	static void SetWindowDimensions_EventThread(ANativeWindow* DimensionWindow);
+
+	// Called by the event manager to update the cached window dimensions to match the event it is processing.
+	static void EventManagerUpdateWindowDimensions(int32 Width, int32 Height);
+
 protected:
 	/** @return true if the native window is currently in fullscreen mode, false otherwise */
 	virtual EWindowMode::Type GetWindowMode() const override { return EWindowMode::Fullscreen; }
@@ -66,9 +75,9 @@ protected:
 private:
 	/** called from GetScreenRect function */
 	/** test cached values from the latest computations stored by CacheRect to decide their validity with the provided arguments */
-	static bool IsCachedRectValid(const bool bMosaicEnabled, const float RequestedContentScaleFactor, ANativeWindow* Window);
+	static bool IsCachedRectValid(bool bUseEventThreadWindow, const bool bMosaicEnabled, const float RequestedContentScaleFactor, ANativeWindow* Window);
 	/** caches some values used to compute the size of the window by GetScreenRect function */
-	static void CacheRect(ANativeWindow* Window, const int32 Width, const int32 Height, const float RequestedContentScaleFactor, const bool bMosaicEnabled);
+	static void CacheRect(bool bUseEventThreadWindow, const int32 Width, const int32 Height, const float RequestedContentScaleFactor, const bool bMosaicEnabled, ANativeWindow* Window);
 
 	/**
 	 * Protect the constructor; only TSharedRefs of this class can be made.
@@ -77,12 +86,16 @@ private:
 
 	FAndroidApplication* OwningApplication;
 
-	/** iOS window handle, typically, only one should ever exist */
-	ANativeWindow *Window;
-
 	/** Store the window region size for querying whether a point lies within the window */
 	int32 RegionX;
 	int32 RegionY;
 
 	static void* NativeWindow;
+
+	// Waits for the event thread to report an initial window size.
+	static bool WaitForWindowDimensions();
+
+	static bool bAreCachedNativeDimensionsValid;
+	static int32 CachedNativeWindowWidth;
+	static int32 CachedNativeWindowHeight;
 };
diff --git a/Engine/Source/Runtime/Core/Public/Android/AndroidPlatformMisc.h b/Engine/Source/Runtime/Core/Public/Android/AndroidPlatformMisc.h
--- a/Engine/Source/Runtime/Core/Public/Android/AndroidPlatformMisc.h
+++ b/Engine/Source/Runtime/Core/Public/Android/AndroidPlatformMisc.h
@@ -258,7 +258,10 @@ public:
 	static uint32 GetCoreFrequency(int32 CoreIndex, ECoreFrequencyProperty CoreFrequencyProperty);
     
     static void SetDeviceOrientation(EDeviceScreenOrientation NewDeviceOrentation) { DeviceOrientation = NewDeviceOrentation; }
-    
+
+	// Window access is locked by the game thread before preinit and unlocked here after RHIInit (PlatformCreateDynamicRHI). 
+	static void UnlockAndroidWindow();
+
 private:
     static EDeviceScreenOrientation DeviceOrientation;
 };
diff --git a/Engine/Source/Runtime/Launch/Private/Android/AndroidEventManager.cpp b/Engine/Source/Runtime/Launch/Private/Android/AndroidEventManager.cpp
--- a/Engine/Source/Runtime/Launch/Private/Android/AndroidEventManager.cpp
+++ b/Engine/Source/Runtime/Launch/Private/Android/AndroidEventManager.cpp
@@ -43,7 +43,10 @@ static const TCHAR* GetAppEventName(EAppEventState State)
 		TEXT("APP_EVENT_STATE_ON_START"),
 		TEXT("APP_EVENT_STATE_WINDOW_LOST_FOCUS"),
 		TEXT("APP_EVENT_STATE_WINDOW_GAINED_FOCUS"),
-		TEXT("APP_EVENT_STATE_SAVE_STATE")
+		TEXT("APP_EVENT_STATE_SAVE_STATE"),
+		TEXT("APP_EVENT_STATE_APP_SUSPENDED"),
+		TEXT("APP_EVENT_STATE_APP_ACTIVATED"),
+		TEXT("APP_EVENT_RUN_CALLBACK"),
 		};
 
 
@@ -51,7 +54,7 @@ static const TCHAR* GetAppEventName(EAppEventState State)
 	{
 		return TEXT("APP_EVENT_STATE_INVALID");
 	}
-	else if (State > APP_EVENT_STATE_SAVE_STATE || State < 0)
+	else if (State > APP_EVENT_RUN_CALLBACK || State < 0)
 	{
 		return TEXT("UnknownEAppEventStateValue");
 	}
@@ -64,104 +67,47 @@ static const TCHAR* GetAppEventName(EAppEventState State)
 
 void FAppEventManager::Tick()
 {
-	static const bool bIsDaydreamApp = FAndroidMisc::IsDaydreamApplication();
-	bool bWindowCreatedThisTick = false;
-	
+	check(IsInGameThread());
 	while (!Queue.IsEmpty())
 	{
-		bool bDestroyWindow = false;
-		bool bShuttingDown = false;
-
-		FAppEventData Event = DequeueAppEvent();
+		FAppEventPacket Event = DequeueAppEvent();
+		FPlatformMisc::LowLevelOutputDebugStringf(TEXT("FAppEventManager::Tick processing, %d"), int(Event.State));
 
 		switch (Event.State)
 		{
 		case APP_EVENT_STATE_WINDOW_CREATED:
-			// if we have a "destroy window" event pending, the data has been invalidated
-			if (!bDestroyWindowPending)
-			{
-				bCreateWindow = true;
-				PendingWindow = (ANativeWindow*)Event.Data;
-				FPlatformMisc::LowLevelOutputDebugStringf(TEXT("APP_EVENT_STATE_WINDOW_CREATED %d, %d, %d, %d"), int(bDestroyWindowPending), int(bRunning), int(bHaveWindow), int(bHaveGame));
-
-			}
-			else
-			{
-				// If we skip a create window because it is destroyed before we created the window *something* gets out of sync
-				// resulting in either one buffer still in the wrong orienation or just a black screen.
-				// So reset everything here. When the new window is created we will recover successfully.
-
-				FAndroidAppEntry::DestroyWindow();
-				FAndroidWindow::SetHardwareWindow(NULL);
-
-				PauseRendering();
-				PauseAudio();
-				FPlatformMisc::LowLevelOutputDebugStringf(TEXT("APP_EVENT_STATE_WINDOW_CREATED window creation skipped because a destroy is pending %d, %d, %d, %d"), int(bDestroyWindowPending), int(bRunning), int(bHaveWindow), int(bHaveGame));
-
-			}
+			FAndroidWindow::EventManagerUpdateWindowDimensions(Event.Data.WindowWidth, Event.Data.WindowHeight);
+			bCreateWindow = true;
 			break;
-		
 		case APP_EVENT_STATE_WINDOW_RESIZED:
+			// Cache the new window's dimensions for the game thread.
+			FAndroidWindow::EventManagerUpdateWindowDimensions(Event.Data.WindowWidth, Event.Data.WindowHeight);
+			ExecWindowResized();
+			break;
 		case APP_EVENT_STATE_WINDOW_CHANGED:
 			// React on device orientation/windowSize changes only when application has window
 			// In case window was created this tick it should already has correct size
-			if (bHaveWindow && !bWindowCreatedThisTick)
-			{
-				ExecWindowResized();
-			}
-			break;
+			// see 'Java_com_epicgames_ue4_GameActivity_nativeOnConfigurationChanged' for event thread/game thread mismatches.
+			ExecWindowResized();
+		break;
 		case APP_EVENT_STATE_SAVE_STATE:
 			bSaveState = true; //todo android: handle save state.
 			break;
 		case APP_EVENT_STATE_WINDOW_DESTROYED:
-			// only if precedeed by a a successfull "create window" event 
-			if (bHaveWindow)
-			{
-				if (bIsDaydreamApp)
-				{
-					bCreateWindow = false;
-				}
-				else
-				{
-					if (GEngine != nullptr && GEngine->IsInitialized() && GEngine->XRSystem.IsValid() && GEngine->XRSystem->GetHMDDevice() && GEngine->XRSystem->GetHMDDevice()->IsHMDConnected())
-					{
-						// delay the destruction until after the renderer teardown on Oculus Mobile
-						bDestroyWindow = true;
-					}
-					else
-					{
-						//FAndroidAppEntry::DestroyWindow();
-						//FAndroidWindow::SetHardwareWindow(NULL);
-						bDestroyWindow = true;
-					}
-				}
-			}
-
 			bHaveWindow = false;
-
-			// allow further "create window" events to be processed 
-			bDestroyWindowPending = false;
 			FPlatformMisc::LowLevelOutputDebugStringf(TEXT("APP_EVENT_STATE_WINDOW_DESTROYED, %d, %d, %d"), int(bRunning), int(bHaveWindow), int(bHaveGame));
 			break;
 		case APP_EVENT_STATE_ON_START:
 			//doing nothing here
 			break;
 		case APP_EVENT_STATE_ON_DESTROY:
-			if (FTaskGraphInterface::IsRunning())
-			{
-				FGraphEventRef WillTerminateTask = FFunctionGraphTask::CreateAndDispatchWhenReady([&]()
-				{
-					FCoreDelegates::ApplicationWillTerminateDelegate.Broadcast();
-				}, TStatId(), NULL, ENamedThreads::GameThread);
-				FTaskGraphInterface::Get().WaitUntilTaskCompletes(WillTerminateTask);
-			}
-			GIsRequestingExit = true; //destroy immediately. Game will shutdown.
-			FirstInitialized = false;
+			check(bHaveWindow == false);
+			check(GIsRequestingExit); //destroy immediately. Game will shutdown.
 			FPlatformMisc::LowLevelOutputDebugStringf(TEXT("APP_EVENT_STATE_ON_DESTROY"));
 			break;
 		case APP_EVENT_STATE_ON_STOP:
-			bShuttingDown = true;
 			bHaveGame = false;
+			ReleaseMicrophone(true);
 			break;
 		case APP_EVENT_STATE_ON_PAUSE:
 			FAndroidAppEntry::OnPauseEvent();
@@ -172,14 +118,26 @@ void FAppEventManager::Tick()
 			break;
 
 		// window focus events that follow their own hierarchy, and might or might not respect App main events hierarchy
-
 		case APP_EVENT_STATE_WINDOW_GAINED_FOCUS: 
 			bWindowInFocus = true;
 			break;
 		case APP_EVENT_STATE_WINDOW_LOST_FOCUS:
 			bWindowInFocus = false;
 			break;
-
+		case APP_EVENT_RUN_CALLBACK:
+		{
+			UE_LOG(LogAndroidEvents, Display, TEXT("Event thread callback running."));
+			Event.Data.CallbackFunc();
+			break;
+		}
+		case APP_EVENT_STATE_APP_ACTIVATED:
+			bRunning = true;
+			FPlatformMisc::LowLevelOutputDebugStringf(TEXT("Execution will be resumed!"));
+			break;
+		case APP_EVENT_STATE_APP_SUSPENDED:
+			bRunning = false;
+			FPlatformMisc::LowLevelOutputDebugStringf(TEXT("Execution will be paused..."));
+			break;
 		default:
 			UE_LOG(LogAndroidEvents, Display, TEXT("Application Event : %u  not handled. "), Event.State);
 		}
@@ -192,89 +150,21 @@ void FAppEventManager::Tick()
 				ExecWindowCreated();
 				bCreateWindow = false;
 				bHaveWindow = true;
-				bWindowCreatedThisTick = true;
-
 				FPlatformMisc::LowLevelOutputDebugStringf(TEXT("ExecWindowCreated, %d, %d, %d"), int(bRunning), int(bHaveWindow), int(bHaveGame));
 			}
 		}
-
-		if (!bRunning && bHaveWindow && bHaveGame)
-		{
-			ResumeRendering();
-			ResumeAudio();
-
-			// broadcast events after the rendering thread has resumed
-			if (FTaskGraphInterface::IsRunning())
-			{
-				FGraphEventRef EnterForegroundTask = FFunctionGraphTask::CreateAndDispatchWhenReady([&]()
-				{
-					FCoreDelegates::ApplicationHasEnteredForegroundDelegate.Broadcast();
-				}, TStatId(), NULL, ENamedThreads::GameThread);
-
-				FGraphEventRef ReactivateTask = FFunctionGraphTask::CreateAndDispatchWhenReady([&]()
-				{
-					FCoreDelegates::ApplicationHasReactivatedDelegate.Broadcast();
-				}, TStatId(), EnterForegroundTask, ENamedThreads::GameThread);
-				FTaskGraphInterface::Get().WaitUntilTaskCompletes(ReactivateTask);
-
-				extern void AndroidThunkCpp_ShowHiddenAlertDialog();
-				AndroidThunkCpp_ShowHiddenAlertDialog();
-			}
-
-			bRunning = true;
-			FPlatformMisc::LowLevelOutputDebugStringf(TEXT("Execution has been resumed!"));
-		}
-		else if (bRunning && (!bHaveWindow || !bHaveGame))
-		{
-			// broadcast events before rendering thread suspends
-			if (FTaskGraphInterface::IsRunning())
-			{
-				FGraphEventRef DeactivateTask = FFunctionGraphTask::CreateAndDispatchWhenReady([&]()
-				{
-					FCoreDelegates::ApplicationWillDeactivateDelegate.Broadcast();
-				}, TStatId(), NULL, ENamedThreads::GameThread);
-				FGraphEventRef EnterBackgroundTask = FFunctionGraphTask::CreateAndDispatchWhenReady([&]()
-				{
-					FCoreDelegates::ApplicationWillEnterBackgroundDelegate.Broadcast();
-				}, TStatId(), DeactivateTask, ENamedThreads::GameThread);
-				FTaskGraphInterface::Get().WaitUntilTaskCompletes(EnterBackgroundTask);
-			}
-
-			PauseRendering();
-			PauseAudio();
-			ReleaseMicrophone(bShuttingDown);
-
-			bRunning = false;
-			FPlatformMisc::LowLevelOutputDebugStringf(TEXT("Execution has been paused..."));
-		}
-
-		if (bDestroyWindow)
-		{
-			FAndroidAppEntry::DestroyWindow();
-			FAndroidWindow::SetHardwareWindow(NULL);
-			bDestroyWindow = false;
-
-			FPlatformMisc::LowLevelOutputDebugStringf(TEXT("FAndroidAppEntry::DestroyWindow() called"));
-		}
 	}
 
 	if (EmptyQueueHandlerEvent)
 	{
 		EmptyQueueHandlerEvent->Trigger();
 	}
-	if (bIsDaydreamApp)
-	{
-		if (!bRunning && FAndroidWindow::GetHardwareWindow() != NULL)
-		{
-			EventHandlerEvent->Wait();
-		}
-	}
-	else
+
+	if (!bRunning)
 	{
-		if (!bRunning && FirstInitialized)
-		{
-			EventHandlerEvent->Wait();
-		}
+		FPlatformMisc::LowLevelOutputDebugStringf(TEXT("FAppEventManager::Tick EventHandlerEvent Wait "));
+		EventHandlerEvent->Wait();
+		FPlatformMisc::LowLevelOutputDebugStringf(TEXT("FAppEventManager::Tick EventHandlerEvent DONE Wait "));
 	}
 }
 
@@ -303,13 +193,10 @@ FAppEventManager::FAppEventManager():
 	,bWindowInFocus(true)
 	,bSaveState(false)
 	,bAudioPaused(false)
-	,PendingWindow(NULL)
 	,bHaveWindow(false)
 	,bHaveGame(false)
 	,bRunning(false)
-	,bDestroyWindowPending(false)
 {
-	pthread_mutex_init(&MainMutex, NULL);
 	pthread_mutex_init(&QueueMutex, NULL);
 
 	IConsoleVariable* CVar = IConsoleManager::Get().FindConsoleVariable(TEXT("r.MobileContentScaleFactor"));
@@ -325,98 +212,33 @@ void FAppEventManager::OnScaleFactorChanged(IConsoleVariable* CVar)
 	}
 }
 
-void FAppEventManager::HandleWindowCreated(void* InWindow)
+void FAppEventManager::HandleWindowCreated_EventThread(void* InWindow)
 {
-	static const bool bIsDaydreamApp = FAndroidMisc::IsDaydreamApplication();
-	if (bIsDaydreamApp)
-	{
-		// We must ALWAYS set the hardware window immediately,
-		// Otherwise we will temporarily end up with an abandoned Window
-		// when the application is pausing/resuming. This is likely
-		// to happen in a Gvr app due to the DON flow pushing an activity
-		// during initialization.
-
-		int rc = pthread_mutex_lock(&MainMutex);
-		check(rc == 0);
-
-		// If we already have a window, destroy it
-		ExecDestroyWindow();
-
-		FAndroidWindow::SetHardwareWindow(InWindow);
-
-		rc = pthread_mutex_unlock(&MainMutex);
-		check(rc == 0);
-
-		// Make sure window will not be deleted until event is processed
-		// Window could be deleted by OS while event queue stuck at game start-up phase
-		FAndroidWindow::AcquireWindowRef((ANativeWindow*)InWindow);
-
-		EnqueueAppEvent(APP_EVENT_STATE_WINDOW_CREATED, InWindow);
-		return;
-	}
-
-	int rc = pthread_mutex_lock(&MainMutex);
-	check(rc == 0);
 	bool AlreadyInited = FirstInitialized;
-	rc = pthread_mutex_unlock(&MainMutex);
-	check(rc == 0);
 
 	// Make sure window will not be deleted until event is processed
 	// Window could be deleted by OS while event queue stuck at game start-up phase
 	FAndroidWindow::AcquireWindowRef((ANativeWindow*)InWindow);
 
-	if (AlreadyInited)
-	{
-		EnqueueAppEvent(APP_EVENT_STATE_WINDOW_CREATED, InWindow);
-	}
-	else
+	check(FAndroidWindow::GetHardwareWindow_EventThread() == NULL);
+	FAndroidWindow::SetHardwareWindow_EventThread(InWindow);
+
+	if (!AlreadyInited)
 	{
 		//This cannot wait until first tick. 
-
-		rc = pthread_mutex_lock(&MainMutex);
-		check(rc == 0);
-
-		check(FAndroidWindow::GetHardwareWindow() == NULL);
-		FAndroidWindow::SetHardwareWindow(InWindow);
 		FirstInitialized = true;
-
-		rc = pthread_mutex_unlock(&MainMutex);
-		check(rc == 0);
-
-		EnqueueAppEvent(APP_EVENT_STATE_WINDOW_CREATED, InWindow);
 	}
+	EnqueueAppEvent(APP_EVENT_STATE_WINDOW_CREATED, FAppEventData((ANativeWindow*)InWindow));
 }
 
-void FAppEventManager::HandleWindowClosed()
+void FAppEventManager::HandleWindowClosed_EventThread()
 {
-	static const bool bIsDaydreamApp = FAndroidMisc::IsDaydreamApplication();
-	if (bIsDaydreamApp)
-	{
-		// We must ALWAYS destroy the hardware window immediately,
-		// Otherwise we will temporarily end up with an abandoned Window
-		// when the application is pausing/resuming. This is likely
-		// to happen in a Gvr app due to the DON flow pushing an activity
-		// during initialization.
-
-		int rc = pthread_mutex_lock(&MainMutex);
-		check(rc == 0);
+	check(FAndroidWindow::GetHardwareWindow_EventThread());
 
-		ExecDestroyWindow();
+	FAndroidWindow::ReleaseWindowRef((ANativeWindow*)FAndroidWindow::GetHardwareWindow_EventThread());
+	FAndroidWindow::SetHardwareWindow_EventThread(nullptr);
 
-		rc = pthread_mutex_unlock(&MainMutex);
-		check(rc == 0);
-	}
-
-	// a "destroy window" event appears on the game preInit routine
-	//     before creating a valid Android window
-	// - override the "create window" data
-	if (!GEngine || !GEngine->IsInitialized())
-	{
-		FirstInitialized = false;
-		FAndroidWindow::SetHardwareWindow(NULL);
-		bDestroyWindowPending = true;
-	}
-	EnqueueAppEvent(APP_EVENT_STATE_WINDOW_DESTROYED, NULL);
+	EnqueueAppEvent(APP_EVENT_STATE_WINDOW_DESTROYED);
 }
 
 
@@ -465,30 +287,13 @@ void FAppEventManager::ResumeRendering()
 void FAppEventManager::ExecWindowCreated()
 {
 	UE_LOG(LogAndroidEvents, Display, TEXT("ExecWindowCreated"));
-
-	static bool bIsDaydreamApp = FAndroidMisc::IsDaydreamApplication();
-	if (!bIsDaydreamApp)
-	{
-		check(PendingWindow);
-		FAndroidWindow::SetHardwareWindow(PendingWindow);
-	}
-
 	// When application launched while device is in sleep mode SystemResolution could be set to opposite orientation values
 	// Force to update SystemResolution to current values whenever we create a new window
 	FPlatformRect ScreenRect = FAndroidWindow::GetScreenRect();
 	FSystemResolution::RequestResolutionChange(ScreenRect.Right, ScreenRect.Bottom, EWindowMode::Fullscreen);
 
-	// ReInit with the new window handle, null for daydream case.
-	FAndroidAppEntry::ReInitWindow(!bIsDaydreamApp ? PendingWindow : nullptr);
-
-	if (!bIsDaydreamApp)
-	{
-		// We hold this reference to ensure that window will not be deleted while game starting up
-		// release it when window is finally initialized
-		FAndroidWindow::ReleaseWindowRef(PendingWindow);
-		PendingWindow = nullptr;
-	}
-
+	// ReInit with the new window handle
+	FAndroidAppEntry::ReInitWindow();
 	FAndroidApplication::OnWindowSizeChanged();
 }
 
@@ -503,17 +308,6 @@ void FAppEventManager::ExecWindowResized()
 	FAndroidApplication::OnWindowSizeChanged();
 }
 
-void FAppEventManager::ExecDestroyWindow()
-{
-	if (FAndroidWindow::GetHardwareWindow() != NULL)
-	{
-		FAndroidWindow::ReleaseWindowRef((ANativeWindow*)FAndroidWindow::GetHardwareWindow());
-
-		FAndroidAppEntry::DestroyWindow();
-		FAndroidWindow::SetHardwareWindow(NULL);
-	}
-}
-
 void FAppEventManager::PauseAudio()
 {
 	if (!GEngine || !GEngine->IsInitialized())
@@ -571,9 +365,9 @@ void FAppEventManager::ResumeAudio()
 }
 
 
-void FAppEventManager::EnqueueAppEvent(EAppEventState InState, void* InData)
+void FAppEventManager::EnqueueAppEvent(EAppEventState InState, FAppEventData&& InData)
 {
-	FAppEventData Event;
+	FAppEventPacket Event;
 	Event.State = InState;
 	Event.Data = InData;
 
@@ -589,22 +383,21 @@ void FAppEventManager::EnqueueAppEvent(EAppEventState InState, void* InData)
 	rc = pthread_mutex_unlock(&QueueMutex);
 	check(rc == 0);
 
-	FPlatformMisc::LowLevelOutputDebugStringf(TEXT("LogAndroidEvents::EnqueueAppEvent : %u, %u, tid = %d, %s"), InState, (uintptr_t)InData, gettid(), GetAppEventName(InState));
+	FPlatformMisc::LowLevelOutputDebugStringf(TEXT("LogAndroidEvents::EnqueueAppEvent : %u, [width=%d, height=%d], tid = %d, %s"), InState, InData.WindowWidth, InData.WindowHeight, gettid(), GetAppEventName(InState));
 }
 
-
-FAppEventData FAppEventManager::DequeueAppEvent()
+FAppEventPacket FAppEventManager::DequeueAppEvent()
 {
 	int rc = pthread_mutex_lock(&QueueMutex);
 	check(rc == 0);
 
-	FAppEventData OutData;
+	FAppEventPacket OutData;
 	Queue.Dequeue( OutData );
 
 	rc = pthread_mutex_unlock(&QueueMutex);
 	check(rc == 0);
 
-	UE_LOG(LogAndroidEvents, Display, TEXT("LogAndroidEvents::DequeueAppEvent : %u, %u, %s"), OutData.State, (uintptr_t)OutData.Data, GetAppEventName(OutData.State))
+	UE_LOG(LogAndroidEvents, Display, TEXT("LogAndroidEvents::DequeueAppEvent : %u, [width=%d, height=%d], %s"), OutData.State, OutData.Data.WindowWidth, OutData.Data.WindowHeight, GetAppEventName(OutData.State))
 
 	return OutData;
 }
@@ -627,7 +420,7 @@ bool FAppEventManager::WaitForEventInQueue(EAppEventState InState, double Timeou
 	bool FoundEvent = false;
 	double StopTime = FPlatformTime::Seconds() + TimeoutSeconds;
 
-	TQueue<FAppEventData, EQueueMode::Spsc> HoldingQueue;
+	TQueue<FAppEventPacket, EQueueMode::Spsc> HoldingQueue;
 	while (!FoundEvent)
 	{
 		int rc = pthread_mutex_lock(&QueueMutex);
@@ -636,7 +429,7 @@ bool FAppEventManager::WaitForEventInQueue(EAppEventState InState, double Timeou
 		// Copy the existing queue (and check for our event)
 		while (!Queue.IsEmpty())
 		{
-			FAppEventData OutData;
+			FAppEventPacket OutData;
 			Queue.Dequeue(OutData);
 
 			if (OutData.State == InState)
@@ -661,7 +454,7 @@ bool FAppEventManager::WaitForEventInQueue(EAppEventState InState, double Timeou
 	// Add events back to queue from holding
 	while (!HoldingQueue.IsEmpty())
 	{
-		FAppEventData OutData;
+		FAppEventPacket OutData;
 		HoldingQueue.Dequeue(OutData);
 		Queue.Enqueue(OutData);
 	}
@@ -674,12 +467,4 @@ bool FAppEventManager::WaitForEventInQueue(EAppEventState InState, double Timeou
 
 extern volatile bool GEventHandlerInitialized;
 
-void FAppEventManager::WaitForEmptyQueue()
-{
-	if (EmptyQueueHandlerEvent && GEventHandlerInitialized && !GIsRequestingExit)
-	{
-		EmptyQueueHandlerEvent->Wait();
-	}
-}
-
 #endif
\ No newline at end of file
diff --git a/Engine/Source/Runtime/Launch/Private/Android/LaunchAndroid.cpp b/Engine/Source/Runtime/Launch/Private/Android/LaunchAndroid.cpp
--- a/Engine/Source/Runtime/Launch/Private/Android/LaunchAndroid.cpp
+++ b/Engine/Source/Runtime/Launch/Private/Android/LaunchAndroid.cpp
@@ -216,6 +216,18 @@ static FEvent* EventHandlerEvent = NULL;
 static volatile bool GResumeMainInit = false;
 volatile bool GEventHandlerInitialized = false;
 
+// The event thread locks this whenever the window unavailable during early init, pause and resume.
+FCriticalSection GAndroidWindowLock;
+bool bReadyToProcessEvents = false;
+void FPlatformMisc::UnlockAndroidWindow()
+{
+	check(IsInGameThread());
+	check(FTaskGraphInterface::IsRunning());
+	UE_LOG(LogAndroid, Log, TEXT("Unlocking android HW window during preinit."));
+	bReadyToProcessEvents = true;
+	GAndroidWindowLock.Unlock();
+}
+
 JNI_METHOD void Java_com_epicgames_ue4_GameActivity_nativeResumeMainInit(JNIEnv* jenv, jobject thiz)
 {
 	GResumeMainInit = true;
@@ -431,14 +443,15 @@ int32 AndroidMain(struct android_app* state)
 	// OBBs and APK are found.
 	IPlatformFile::GetPlatformPhysical().Initialize(nullptr, FCommandLine::Get());
 
-	//wait for re-creating the native window, if previously destroyed by onStop()
 	{
-		SCOPED_BOOT_TIMING("Wait for FAndroidWindow::GetHardwareWindow()");
-		while (FAndroidWindow::GetHardwareWindow() == nullptr)
-		{
-			FPlatformProcess::Sleep(0.01f);
-			FPlatformMisc::MemoryBarrier();
-		}
+		SCOPED_BOOT_TIMING("Wait for GAndroidWindowLock.Lock()");
+		// wait for a valid window
+		// Lock GAndroidWindowLock to ensure no window destroy shenanigans occur between early phase of preinit and UnlockAndroidWindow
+		// Note: this is unlocked after Android's PlatformCreateDynamicRHI when the RHI is then able to process window changes.
+		// We don't wait for all of preinit to complete as PreLoadScreens will need to process events during preinit.
+
+		UE_LOG(LogAndroid, Log, TEXT("PreInit android HW window lock."));
+		GAndroidWindowLock.Lock();
 	}
 
 	// initialize the engine
@@ -640,6 +653,14 @@ void FChoreographer::SetCallback(int64 Delay)
 	AChoreographer_postFrameCallbackDelayed_(choreographer, choreographer_callback, nullptr, Delay / 1000000);
 }
 
+static uint32 EventThreadID = 0;
+
+bool IsInAndroidEventThread()
+{
+	check(EventThreadID != 0);
+	return EventThreadID == FPlatformTLS::GetCurrentThreadId();
+}
+
 static void* AndroidEventThreadWorker( void* param )
 {
 	struct android_app* state = (struct android_app*)param;
@@ -647,6 +668,7 @@ static void* AndroidEventThreadWorker( void* param )
 	FPlatformProcess::SetThreadAffinityMask(FPlatformAffinity::GetMainGameMask());
 
 	FPlatformMisc::LowLevelOutputDebugString(TEXT("Entering event processing thread engine entry point"));
+	EventThreadID = FPlatformTLS::GetCurrentThreadId();
 
 	ALooper* looper = ALooper_prepare(ALOOPER_PREPARE_ALLOW_NON_CALLBACKS);
 	ALooper_addFd(looper, state->msgread, LOOPER_ID_MAIN, ALOOPER_EVENT_INPUT, NULL,
@@ -664,6 +686,10 @@ static void* AndroidEventThreadWorker( void* param )
 
 	TheChoreographer.SetupChoreographer();
 
+	// window is initially invalid/locked.
+	UE_LOG(LogAndroid, Log, TEXT("event thread, Initial HW window lock."));
+	GAndroidWindowLock.Lock();
+
 	//continue to process events until the engine is shutting down
 	while (!GIsRequestingExit)
 	{
@@ -674,6 +700,8 @@ static void* AndroidEventThreadWorker( void* param )
 		sleep(EventRefreshRate);		// this is really 0 since it takes int seconds.
 	}
 
+	GAndroidWindowLock.Unlock();
+
 	UE_LOG(LogAndroid, Log, TEXT("Exiting"));
 
 	return NULL;
@@ -746,6 +774,7 @@ static int32_t HandleInputCB(struct android_app* app, AInputEvent* event)
 {
 //	FPlatformMisc::LowLevelOutputDebugStringf(TEXT("INPUT - type: %x, action: %x, source: %x, keycode: %x, buttons: %x"), AInputEvent_getType(event), 
 //		AMotionEvent_getAction(event), AInputEvent_getSource(event), AKeyEvent_getKeyCode(event), AMotionEvent_getButtonState(event));
+	check(IsInAndroidEventThread());
 
 	int32 EventType = AInputEvent_getType(event);
 	int32 EventSource = AInputEvent_getSource(event);
@@ -894,7 +923,7 @@ static int32_t HandleInputCB(struct android_app* app, AInputEvent* event)
 				return 1;
 			}
 
-			ANativeWindow* Window = (ANativeWindow*)FAndroidWindow::GetHardwareWindow();
+			ANativeWindow* Window = (ANativeWindow*)FAndroidWindow::GetHardwareWindow_EventThread();
 			if (!Window)
 			{
 				return 0;
@@ -905,7 +934,8 @@ static int32_t HandleInputCB(struct android_app* app, AInputEvent* event)
 
 			if(Window)
 			{
-				FAndroidWindow::CalculateSurfaceSize(Window, Width, Height);
+				// we are on the event thread. true here indicates we will retrieve dimensions from the current window.
+				FAndroidWindow::CalculateSurfaceSize(Width, Height, true);
 			}
 
 			// make sure OpenGL context created before accepting touch events.. FAndroidWindow::GetScreenRect() may try to create it early from wrong thread if this is the first call
@@ -913,7 +943,7 @@ static int32_t HandleInputCB(struct android_app* app, AInputEvent* event)
 			{
 				return 1;
 			}
-			FPlatformRect ScreenRect = FAndroidWindow::GetScreenRect();
+			FPlatformRect ScreenRect = FAndroidWindow::GetScreenRect(true);
 
 			if (AndroidThunkCpp_IsVirtuaKeyboardShown() && (type == TouchBegan || type == TouchMoved))
 			{
@@ -1080,13 +1110,121 @@ static bool IsPreLoadScreenPlaying()
         || (FPreLoadScreenManager::Get() && (FPreLoadScreenManager::Get()->HasValidActivePreLoadScreen()));
 }
 
+FAppEventData::FAppEventData(ANativeWindow* WindowIn)
+{
+	check(WindowIn);
+	WindowWidth = ANativeWindow_getWidth(WindowIn);
+	WindowHeight = ANativeWindow_getHeight(WindowIn);
+	check(WindowWidth >= 0 && WindowHeight >= 0);
+}
+
+static bool bAppIsActive_EventThread = false;
+
+
+// called when the app has focus + window + resume.
+static void ActivateApp_EventThread()
+{
+	if (bAppIsActive_EventThread)
+	{
+		// Seems this can occur.
+		return;
+	}
+
+	// Unlock window when we're ready.
+	UE_LOG(LogAndroid, Log, TEXT("event thread, activate app, unlocking HW window"));
+	GAndroidWindowLock.Unlock();
+	// wake the GT up.
+	FAppEventManager::GetInstance()->EnqueueAppEvent(APP_EVENT_STATE_APP_ACTIVATED);
+
+	bAppIsActive_EventThread = true;
+	FAppEventManager::GetInstance()->EnqueueAppEvent(APP_EVENT_RUN_CALLBACK, FAppEventData([]()
+	{
+		UE_LOG(LogAndroid, Log, TEXT("performing app foregrounding callback."));
+		FCoreDelegates::ApplicationHasEnteredForegroundDelegate.Broadcast();
+		FCoreDelegates::ApplicationHasReactivatedDelegate.Broadcast();
+		FAppEventManager::GetInstance()->ResumeAudio();
+	}));
+
+	FPreLoadScreenManager::EnableRendering(true);
+
+	extern void AndroidThunkCpp_ShowHiddenAlertDialog();
+	AndroidThunkCpp_ShowHiddenAlertDialog();	
+}
+
+extern void BlockRendering();
+// called whenever the app loses focus, loses window or pause.
+static void SuspendApp_EventThread()
+{
+	if (!bAppIsActive_EventThread)
+	{
+		return;
+	}
+	bAppIsActive_EventThread = false;
+	// Lock the window, this prevents event thread from removing the window whilst the RHI initializes.
+
+	UE_LOG(LogAndroid, Log, TEXT("event thread, suspending app, acquiring HW window lock."));
+	GAndroidWindowLock.Lock();
+
+	if (bReadyToProcessEvents == false)
+	{
+		// App has stopped before we can process events.
+		// AndroidLaunch will lock GAndroidWindowLock, and set bReadyToProcessEvents when we are able to block the RHI and queue up other events.
+		// we ignore events until this point as acquiring GAndroidWindowLock means requires the window to be properly initialized.
+		UE_LOG(LogAndroid, Log, TEXT("event thread, app not yet ready."));
+		return;
+	};
+
+	TSharedPtr<FEvent, ESPMode::ThreadSafe> EMDoneTrigger = MakeShareable(FPlatformProcess::GetSynchEventFromPool(), [](FEvent* EventToDelete)
+	{
+		FPlatformProcess::ReturnSynchEventToPool(EventToDelete);
+	});
+
+	// perform the delegates before the window handle is cleared.
+	// This ensures any tasks that require a window handle will have it before we block the RT on the invalid window.
+	FAppEventManager::GetInstance()->EnqueueAppEvent(APP_EVENT_RUN_CALLBACK, FAppEventData([EMDoneTrigger]()
+	{
+		UE_LOG(LogAndroid, Log, TEXT("performing app backgrounding callback. %p"), EMDoneTrigger.Get());
+
+		FCoreDelegates::ApplicationWillDeactivateDelegate.Broadcast();
+		FCoreDelegates::ApplicationWillEnterBackgroundDelegate.Broadcast();
+		FAppEventManager::GetInstance()->PauseAudio();
+		FAppEventManager::ReleaseMicrophone(false);
+		EMDoneTrigger->Trigger();
+	}));
+
+	FPreLoadScreenManager::EnableRendering(false);
+
+	// wait for a period of time before blocking rendering
+	UE_LOG(LogAndroid, Log, TEXT("AndroidEGL::  SuspendApp_EventThread, waiting for event manager to process. tid: %d"), FPlatformTLS::GetCurrentThreadId());
+	bool bSuccess = EMDoneTrigger->Wait(4000);
+	UE_CLOG(!bSuccess, LogAndroid, Log, TEXT("backgrounding callback, not responded in timely manner."));
+
+	BlockRendering();
+
+	// Suspend the GT.
+	FAppEventManager::GetInstance()->EnqueueAppEvent(APP_EVENT_STATE_APP_SUSPENDED);
+}
+
 //Called from the event process thread
 static void OnAppCommandCB(struct android_app* app, int32_t cmd)
 {
+	check(IsInAndroidEventThread());
 	static bool bDidGainFocus = false;
-	bool bNeedToSync = false;
 	//FPlatformMisc::LowLevelOutputDebugStringf(TEXT("OnAppCommandCB cmd: %u, tid = %d"), cmd, gettid());
 
+	static bool bHasFocus = false;
+	static bool bHasWindow = false;
+	static bool bIsResumed = false;
+
+	// Set event thread's view of the window dimensions:
+	{
+		ANativeWindow* DimensionWindow = app->pendingWindow ? app->pendingWindow : app->window;
+		if (DimensionWindow)
+		{
+			FAndroidWindow::SetWindowDimensions_EventThread(DimensionWindow);
+		}
+	}
+
 	switch (cmd)
 	{
 	case APP_CMD_SAVE_STATE:
@@ -1110,9 +1248,12 @@ static void OnAppCommandCB(struct android_app* app, int32_t cmd)
 		// get the window ready for showing
 		FPlatformMisc::LowLevelOutputDebugStringf(TEXT("Case APP_CMD_INIT_WINDOW"));
 		UE_LOG(LogAndroid, Log, TEXT("Case APP_CMD_INIT_WINDOW"));
-		FAppEventManager::GetInstance()->HandleWindowCreated(app->pendingWindow);
-
-		bNeedToSync = true;
+		FAppEventManager::GetInstance()->HandleWindowCreated_EventThread(app->pendingWindow);
+		bHasWindow = true;
+		if (bHasWindow && bHasFocus && bIsResumed)
+		{
+			ActivateApp_EventThread();
+		}
 		break;
 	case APP_CMD_TERM_WINDOW:
 		/**
@@ -1124,9 +1265,11 @@ static void OnAppCommandCB(struct android_app* app, int32_t cmd)
 		FPlatformMisc::LowLevelOutputDebugStringf(TEXT("Case APP_CMD_TERM_WINDOW, tid = %d"), gettid());
 		// clean up the window because it is being hidden/closed
 		UE_LOG(LogAndroid, Log, TEXT("Case APP_CMD_TERM_WINDOW"));
-		FAppEventManager::GetInstance()->HandleWindowClosed();
-		
-		bNeedToSync = true;
+
+		SuspendApp_EventThread();
+
+		FAppEventManager::GetInstance()->HandleWindowClosed_EventThread();
+		bHasWindow = false;
 		break;
 	case APP_CMD_LOST_FOCUS:
 		/**
@@ -1134,9 +1277,11 @@ static void OnAppCommandCB(struct android_app* app, int32_t cmd)
 		 * input focus.
 		 */
 		// if the app lost focus, avoid unnecessary processing (like monitoring the accelerometer)
+		bHasFocus = false;
 
 		UE_LOG(LogAndroid, Log, TEXT("Case APP_CMD_LOST_FOCUS"));
-		FAppEventManager::GetInstance()->EnqueueAppEvent(APP_EVENT_STATE_WINDOW_LOST_FOCUS, NULL);
+		FAppEventManager::GetInstance()->EnqueueAppEvent(APP_EVENT_STATE_WINDOW_LOST_FOCUS);
+
 		break;
 	case APP_CMD_GAINED_FOCUS:
 		/**
@@ -1146,10 +1291,16 @@ static void OnAppCommandCB(struct android_app* app, int32_t cmd)
 
 		// remember gaining focus so we know any later pauses are not part of first startup
 		bDidGainFocus = true;
-		 
+		bHasFocus = true;
+
 		// bring back a certain functionality, like monitoring the accelerometer
 		UE_LOG(LogAndroid, Log, TEXT("Case APP_CMD_GAINED_FOCUS"));
-		FAppEventManager::GetInstance()->EnqueueAppEvent(APP_EVENT_STATE_WINDOW_GAINED_FOCUS, NULL);
+		FAppEventManager::GetInstance()->EnqueueAppEvent(APP_EVENT_STATE_WINDOW_GAINED_FOCUS);
+
+		if (bHasWindow && bHasFocus && bIsResumed)
+		{
+			ActivateApp_EventThread();
+		}
 		break;
 	case APP_CMD_INPUT_CHANGED:
 		UE_LOG(LogAndroid, Log, TEXT("Case APP_CMD_INPUT_CHANGED"));
@@ -1160,7 +1311,7 @@ static void OnAppCommandCB(struct android_app* app, int32_t cmd)
 		 * Please redraw with its new size.
 		 */
 		UE_LOG(LogAndroid, Log, TEXT("Case APP_CMD_WINDOW_RESIZED"));
-		FAppEventManager::GetInstance()->EnqueueAppEvent(APP_EVENT_STATE_WINDOW_RESIZED );
+		FAppEventManager::GetInstance()->EnqueueAppEvent(APP_EVENT_STATE_WINDOW_RESIZED, FAppEventData(app->window));
 		break;
 	case APP_CMD_WINDOW_REDRAW_NEEDED:
 		/**
@@ -1188,7 +1339,7 @@ static void OnAppCommandCB(struct android_app* app, int32_t cmd)
 			bool bPortrait = (AConfiguration_getOrientation(app->config) == ACONFIGURATION_ORIENTATION_PORT);
 			if (FAndroidWindow::OnWindowOrientationChanged(bPortrait))
 			{
-				FAppEventManager::GetInstance()->EnqueueAppEvent(APP_EVENT_STATE_WINDOW_CHANGED, nullptr);
+				FAppEventManager::GetInstance()->EnqueueAppEvent(APP_EVENT_STATE_WINDOW_CHANGED);
 			}
 		}
 		break;
@@ -1212,12 +1363,15 @@ static void OnAppCommandCB(struct android_app* app, int32_t cmd)
 		/**
 		 * Command from main thread: the app's activity has been resumed.
 		 */
+		bIsResumed = true;
+		if (bHasWindow && bHasFocus && bIsResumed)
+		{
+			ActivateApp_EventThread();
+		}
 		FPlatformMisc::LowLevelOutputDebugStringf(TEXT("Case APP_CMD_RESUME"));
 		UE_LOG(LogAndroid, Log, TEXT("Case APP_CMD_RESUME"));
 		FAppEventManager::GetInstance()->EnqueueAppEvent(APP_EVENT_STATE_ON_RESUME);
 
-        FPreLoadScreenManager::EnableRendering(true);
-
 		/*
 		* On the initial loading the restart method must be called immediately
 		* in order to restart the app if the startup movie was playing
@@ -1229,6 +1383,7 @@ static void OnAppCommandCB(struct android_app* app, int32_t cmd)
 		break;
 	case APP_CMD_PAUSE:
 	{
+		bIsResumed = false;
 		/**
 		 * Command from main thread: the app's activity has been paused.
 		 */
@@ -1247,12 +1402,12 @@ static void OnAppCommandCB(struct android_app* app, int32_t cmd)
 		// Restart on resuming if did not complete engine initialization
 		if (!bDidCompleteEngineInit && bDidGainFocus  && !bIgnorePauseOnDownloaderStart && bAllowReboot)
 		{
-			// only do this if early startup enabled
-			FString *EarlyRestart = FAndroidMisc::GetConfigRulesVariable(TEXT("earlyrestart"));
-			if (EarlyRestart != NULL && EarlyRestart->Equals("true", ESearchCase::IgnoreCase))
-			{
-				bShouldRestartFromInterrupt = true;
-			}
+// 			// only do this if early startup enabled
+// 			FString *EarlyRestart = FAndroidMisc::GetConfigRulesVariable(TEXT("earlyrestart"));
+// 			if (EarlyRestart != NULL && EarlyRestart->Equals("true", ESearchCase::IgnoreCase))
+// 			{
+// 				bShouldRestartFromInterrupt = true;
+// 			}
 		}
 		bIgnorePauseOnDownloaderStart = false;
 
@@ -1262,13 +1417,12 @@ static void OnAppCommandCB(struct android_app* app, int32_t cmd)
 		*/
 		if (IsPreLoadScreenPlaying() && bAllowReboot)
 		{
-			bShouldRestartFromInterrupt = true;
+			UE_LOG(LogAndroid, Log, TEXT("MoviePlayer force completion"));
 			GetMoviePlayer()->ForceCompletion();
 		}
 
-        FPreLoadScreenManager::EnableRendering(false);
+		SuspendApp_EventThread();
 
-		bNeedToSync = true;
 		break;
 	}
 	case APP_CMD_STOP:
@@ -1277,7 +1431,6 @@ static void OnAppCommandCB(struct android_app* app, int32_t cmd)
 		 */
 		UE_LOG(LogAndroid, Log, TEXT("Case APP_CMD_STOP"));
 		FAppEventManager::GetInstance()->EnqueueAppEvent(APP_EVENT_STATE_ON_STOP);
-
 		break;
 	case APP_CMD_DESTROY:
 		/**
@@ -1285,6 +1438,18 @@ static void OnAppCommandCB(struct android_app* app, int32_t cmd)
 		* and waiting for the app thread to clean up and exit before proceeding.
 		*/
 		UE_LOG(LogAndroid, Log, TEXT("Case APP_CMD_DESTROY"));
+
+		FAppEventManager::GetInstance()->EnqueueAppEvent(APP_EVENT_RUN_CALLBACK, FAppEventData([]()
+		{
+			FGraphEventRef WillTerminateTask = FFunctionGraphTask::CreateAndDispatchWhenReady([]()
+			{
+				FCoreDelegates::ApplicationWillTerminateDelegate.Broadcast();
+			}, TStatId(), NULL, ENamedThreads::GameThread);
+			FTaskGraphInterface::Get().WaitUntilTaskCompletes(WillTerminateTask);
+			GIsRequestingExit = true; //destroy immediately. Game will shutdown.
+		}));
+
+
 		FAppEventManager::GetInstance()->EnqueueAppEvent(APP_EVENT_STATE_ON_DESTROY);
 		break;
 	}
@@ -1294,10 +1459,6 @@ static void OnAppCommandCB(struct android_app* app, int32_t cmd)
 		EventHandlerEvent->Trigger();
 	}
 
-	if (bNeedToSync)
-	{
-		FAppEventManager::GetInstance()->WaitForEmptyQueue();
-	}
 	//FPlatformMisc::LowLevelOutputDebugStringf(TEXT("#### END OF OnAppCommandCB cmd: %u, tid = %d"), cmd, gettid());
 }
 
@@ -1328,10 +1489,12 @@ JNI_METHOD void Java_com_epicgames_ue4_GameActivity_nativeOnConfigurationChanged
 {
 	bool bChangedToPortrait = bPortrait == JNI_TRUE;
 
-	// enqueue a window changed event if orientation changed
+	// enqueue a window changed event if orientation changed, 
+	// note that the HW window handle does not necessarily change.
 	if (FAndroidWindow::OnWindowOrientationChanged(bChangedToPortrait))
-	{
-		FAppEventManager::GetInstance()->EnqueueAppEvent(APP_EVENT_STATE_WINDOW_CHANGED, nullptr);
+	{	
+		// Enqueue an event to trigger gamethread to update the orientation:
+		FAppEventManager::GetInstance()->EnqueueAppEvent(APP_EVENT_STATE_WINDOW_CHANGED);
 
 		if (EventHandlerEvent)
 		{
diff --git a/Engine/Source/Runtime/Launch/Public/Android/AndroidEventManager.h b/Engine/Source/Runtime/Launch/Public/Android/AndroidEventManager.h
--- a/Engine/Source/Runtime/Launch/Public/Android/AndroidEventManager.h
+++ b/Engine/Source/Runtime/Launch/Public/Android/AndroidEventManager.h
@@ -26,17 +26,35 @@ enum EAppEventState
 	APP_EVENT_STATE_WINDOW_LOST_FOCUS,
 	APP_EVENT_STATE_WINDOW_GAINED_FOCUS,
 	APP_EVENT_STATE_SAVE_STATE,
+	APP_EVENT_STATE_APP_SUSPENDED,
+	APP_EVENT_STATE_APP_ACTIVATED,
+	APP_EVENT_RUN_CALLBACK,
 	APP_EVENT_STATE_INVALID = -1,
 };
 
 struct FAppEventData
+{
+	FAppEventData() : WindowWidth(-1), WindowHeight(-1)
+	{}
+	FAppEventData(TFunction<void(void)> CallbackFuncIN) : WindowWidth(-1), WindowHeight(-1), CallbackFunc(MoveTemp(CallbackFuncIN))
+	{}
+
+
+	FAppEventData(ANativeWindow* WindowIn);
+
+	int32 WindowWidth;
+	int32 WindowHeight;
+
+	TFunction<void(void)> CallbackFunc;
+};
+
+struct FAppEventPacket
 {
 	EAppEventState State;
-	void* Data;
+	FAppEventData Data;
 
-	FAppEventData():
+	FAppEventPacket():
 		State(APP_EVENT_STATE_INVALID)
-		,Data(NULL)
 	{	}
 };
 
@@ -46,49 +64,47 @@ public:
 	static FAppEventManager* GetInstance();
 	
 	void Tick();
-	void EnqueueAppEvent(EAppEventState InState, void* InData = NULL);
+	void EnqueueAppEvent(EAppEventState InState, FAppEventData&& InData = FAppEventData());
 	void SetEventHandlerEvent(FEvent* InEventHandlerEvent);
-	void HandleWindowCreated(void* InWindow);
-	void HandleWindowClosed();
+	// These are called directly from the android event thread.
+	void HandleWindowCreated_EventThread(void* InWindow);
+	void HandleWindowClosed_EventThread();
 	bool IsGamePaused();
 	bool IsGameInFocus();
 	bool WaitForEventInQueue(EAppEventState InState, double TimeoutSeconds);
 
 	void SetEmptyQueueHandlerEvent(FEvent* InEventHandlerEvent);
-	void WaitForEmptyQueue();
 	void TriggerEmptyQueue();
 
+	void PauseAudio();
+	void ResumeAudio();
+	static void ReleaseMicrophone(bool shuttingDown);
+
 protected:
 	FAppEventManager();
 
 private:
-	FAppEventData DequeueAppEvent();
+	FAppEventPacket DequeueAppEvent();
 	void PauseRendering();
 	void ResumeRendering();
-	void PauseAudio();
-	void ResumeAudio();
-	void ReleaseMicrophone(bool shuttingDown);
+
 
 	void ExecWindowCreated();
 	void ExecWindowResized();
-	void ExecDestroyWindow();
 	
 	static void OnScaleFactorChanged(IConsoleVariable* CVar);
 
 	static FAppEventManager* sInstance;
 
 	pthread_mutex_t QueueMutex;			//@todo android: can probably get rid of this now that we're using an mpsc queue
-	TQueue<FAppEventData, EQueueMode::Mpsc> Queue;
+	TQueue<FAppEventPacket, EQueueMode::Mpsc> Queue;
 	
-	pthread_mutex_t MainMutex;
-
 	FEvent* EventHandlerEvent;			// triggered every time the event handler thread fires off an event
 	FEvent* EmptyQueueHandlerEvent;		// triggered every time the queue is emptied
 
 	//states
 	bool FirstInitialized;
 	bool bCreateWindow;
-	ANativeWindow* PendingWindow;
 
 	bool bWindowInFocus;
 	bool bSaveState;
@@ -97,9 +113,8 @@ private:
 	bool bHaveWindow;
 	bool bHaveGame;
 	bool bRunning;
-
-	// mark the "destroy window" event to prevent "create window" with invalid data on the same frame
-	bool bDestroyWindowPending;
 };
 
+bool IsInAndroidEventThread();
+
 #endif
\ No newline at end of file
diff --git a/Engine/Source/Runtime/OpenGLDrv/Private/Android/AndroidEGL.cpp b/Engine/Source/Runtime/OpenGLDrv/Private/Android/AndroidEGL.cpp
--- a/Engine/Source/Runtime/OpenGLDrv/Private/Android/AndroidEGL.cpp
+++ b/Engine/Source/Runtime/OpenGLDrv/Private/Android/AndroidEGL.cpp
@@ -196,6 +196,8 @@ struct AndroidESPImpl
 
 
 	double LastTimeEmulatedSync;
+	FPlatformRect CachedWindowRect;
+
 	AndroidESPImpl();
 };
 
@@ -337,7 +339,7 @@ EGLBoolean AndroidEGL::SetCurrentContext(EGLContext InContext, EGLSurface InSurf
 			}
 
 			Result = eglMakeCurrent(PImplData->eglDisplay, Surface, Surface, InContext);
-			checkf(Result == EGL_TRUE, TEXT("ERROR: SetCurrentSharedContext eglMakeCurrent failed : 0x%x"), eglGetError());
+			checkf(Result == EGL_TRUE, TEXT("ERROR: SetCurrentContext eglMakeCurrent failed : 0x%x"), eglGetError());
 		}
 	}
 	return Result;
@@ -624,7 +626,49 @@ eglDisplay(EGL_NO_DISPLAY)
 	,SyncFenceForChoreographerMethod(EGL_NO_SYNC_KHR)
 
 	,LastTimeEmulatedSync(-1.0)
+	,CachedWindowRect(FPlatformRect(0,0,0,0)
+	)
+{
+}
+
+// call out to JNI to see if the application was packaged for Oculus Mobile
+extern bool AndroidThunkCpp_IsOculusMobileApplication();
+
+void AndroidEGL::SetRenderContextWindowSurface()
 {
+	VERIFY_GL_SCOPE();
+
+	FPlatformMisc::LowLevelOutputDebugStringf(TEXT("AndroidEGL::SetRenderContextWindowSurface  recreating context! tid: %d"),  FPlatformTLS::GetCurrentThreadId());
+	UnBind();
+	SetCurrentContext(EGL_NO_CONTEXT, EGL_NO_SURFACE);
+	bool bCreateSurface = !AndroidThunkCpp_IsOculusMobileApplication();
+
+	//SetRenderContextWindowSurface is called only when the window lock is successful.
+	PImplData->Window = (ANativeWindow*)FAndroidWindow::GetHardwareWindow_EventThread();
+	check(PImplData->Window);
+	InitSurface(false, bCreateSurface);
+	SetCurrentRenderingContext();
+	FPlatformMisc::LowLevelOutputDebugStringf(TEXT("AndroidEGL::SetRenderContextWindowSurface  DONE! tid: %d"),  FPlatformTLS::GetCurrentThreadId());
+}
+
+void AndroidEGL::ResizeRenderContextSurface()
+{
+	VERIFY_GL_SCOPE();
+
+	// Resize originates from the gamethread, we cant use Window_Event here.
+	if (PImplData->Window && 
+		(PImplData->eglWidth != (PImplData->CachedWindowRect.Right - PImplData->CachedWindowRect.Left)
+		|| PImplData->eglHeight != (PImplData->CachedWindowRect.Bottom - PImplData->CachedWindowRect.Top))
+		)
+	{
+		UE_LOG(LogAndroid, Log, TEXT("AndroidEGL::ResizeRenderContextSurface, PImplData->Window=%p, PImplData->eglWidth=%d, PImplData->eglHeight=%d!, CachedWidth=%d, CachedHeight=%d, tid: %d"), 
+			PImplData->Window, PImplData->eglWidth, PImplData->eglHeight, (PImplData->CachedWindowRect.Right - PImplData->CachedWindowRect.Left), (PImplData->CachedWindowRect.Bottom - PImplData->CachedWindowRect.Top), FPlatformTLS::GetCurrentThreadId());
+		UnBind();
+		SetCurrentContext(EGL_NO_CONTEXT, EGL_NO_SURFACE);
+		bool bCreateSurface = !AndroidThunkCpp_IsOculusMobileApplication();
+		InitSurface(false, bCreateSurface);
+		SetCurrentRenderingContext();
+	}
 }
 
 AndroidEGL* AndroidEGL::GetInstance()
@@ -673,30 +717,28 @@ void AndroidEGL::InitSurface(bool bUseSmallSurface, bool bCreateWndSurface)
 {
 	FPlatformMisc::LowLevelOutputDebugStringf(TEXT("AndroidEGL::InitSurface %d, %d"), int(bUseSmallSurface), int(bCreateWndSurface));
 
-	ANativeWindow* window = (ANativeWindow*)FAndroidWindow::GetHardwareWindow();
-	if (window == NULL)
-	{
-		// Sleep if the hardware window isn't currently available.
-		// The Window may not exist if the activity is pausing/resuming, in which case we make this thread wait
-		FPlatformMisc::LowLevelOutputDebugString(TEXT("Waiting for Native window in AndroidEGL::InitSurface"));
-		window = (ANativeWindow*)FAndroidWindow::WaitForHardwareWindow();
-	
-		if (window == NULL)
-		{
-			FPlatformMisc::LowLevelOutputDebugString(TEXT("Aborting AndroidEGL::InitSurface, FAndroidWindow::WaitForHardwareWindow() returned null"));
-			return;
-		}
-	}
+	check(PImplData->Window);
 
-	PImplData->Window = (ANativeWindow*)window;
 	int32 Width = 8, Height = 8;
 	if (!bUseSmallSurface)
 	{
 		FPlatformRect WindowSize = FAndroidWindow::GetScreenRect();
+
+		if (PImplData->CachedWindowRect.Right > 0 && PImplData->CachedWindowRect.Bottom > 0)
+		{
+			// If we resumed from a lost window reuse the window size, the game thread will update the window dimensions.
+			FPlatformMisc::LowLevelOutputDebugStringf(TEXT("AndroidEGL::InitSurface, Using CachedWindowRect, width: %d, height %d "), PImplData->CachedWindowRect.Right, PImplData->CachedWindowRect.Bottom);
+			WindowSize = PImplData->CachedWindowRect;
+		}
+
 		Width = WindowSize.Right;
 		Height = WindowSize.Bottom;
+
+		FPlatformMisc::LowLevelOutputDebugStringf(TEXT("AndroidEGL::InitSurface, Using width: %d, height %d "), Width, Height);
 		AndroidThunkCpp_SetDesiredViewSize(Width, Height);
 	}
+
+	FPlatformMisc::LowLevelOutputDebugStringf(TEXT("AndroidEGL::InitSurface, wnd: %p, width: %d, height %d "), PImplData->Window, Width, Height);
 	ANativeWindow_setBuffersGeometry(PImplData->Window, Width, Height, PImplData->NativeVisualID);
 	CreateEGLSurface(PImplData->Window, bCreateWndSurface);
 	
@@ -705,12 +747,14 @@ void AndroidEGL::InitSurface(bool bUseSmallSurface, bool bCreateWndSurface)
 	PImplData->SingleThreadedContext.eglSurface = PImplData->eglSurface;
 }
 
-// call out to JNI to see if the application was packaged for Oculus Mobile
-extern bool AndroidThunkCpp_IsOculusMobileApplication();
-
 void AndroidEGL::ReInit()
 {
 	FPlatformMisc::LowLevelOutputDebugString(TEXT("AndroidEGL::ReInit()"));
+
+	FPlatformRect WindowRect = FAndroidWindow::GetScreenRect();
+	FPlatformMisc::LowLevelOutputDebugStringf(TEXT("AndroidEGL FAndroidAppEntry::ReInitWindow updating window size = %d, %d, cached size: %d, %d tid: %d"), WindowRect.Right, WindowRect.Bottom, PImplData->CachedWindowRect.Right, PImplData->CachedWindowRect.Bottom, FPlatformTLS::GetCurrentThreadId());
+	PImplData->CachedWindowRect = WindowRect;
+
 	SetCurrentContext(EGL_NO_CONTEXT, EGL_NO_SURFACE);
 	bool bCreateSurface = !AndroidThunkCpp_IsOculusMobileApplication();
 	InitSurface(false, bCreateSurface);
@@ -719,6 +763,7 @@ void AndroidEGL::ReInit()
 
 void AndroidEGL::Init(APIVariant API, uint32 MajorVersion, uint32 MinorVersion, bool bDebug)
 {
+	check(IsInGameThread());
 
 	if (PImplData->Initalized)
 	{
@@ -768,6 +813,8 @@ void AndroidEGL::Init(APIVariant API, uint32 MajorVersion, uint32 MinorVersion,
 	}
 
 	InitContexts();
+	// Getting the hardware window is valid during preinit as we have GAndroidWindowLock held.
+	PImplData->Window = (ANativeWindow*)FAndroidWindow::GetHardwareWindow_EventThread();
 	PImplData->Initalized   = true;
 //	INITIATE_GL_FRAME_DUMP();
 }
@@ -1486,30 +1533,15 @@ void AndroidEGL::SetCurrentSharedContext()
 	}
 }
 
-void AndroidEGL::SetSharedContext()
+void AndroidEGL::AcquireCurrentRenderingContext()
 {
-	check(IsInGameThread());
-	PImplData->CurrentContextType = CONTEXT_Shared;
-
-	SetCurrentContext(PImplData->SharedContext.eglContext, PImplData->SharedContext.eglSurface);
-}
-
-void AndroidEGL::SetSingleThreadRenderingContext()
-{
-	PImplData->CurrentContextType = CONTEXT_Rendering;
-	SetCurrentContext(PImplData->SingleThreadedContext.eglContext, PImplData->SingleThreadedContext.eglSurface);
-}
-
-void AndroidEGL::SetMultithreadRenderingContext()
-{
-	PImplData->CurrentContextType = CONTEXT_Rendering;
-	SetCurrentContext(PImplData->RenderingContext.eglContext, PImplData->RenderingContext.eglSurface);
+	SetCurrentRenderingContext();
 }
 
 void AndroidEGL::SetCurrentRenderingContext()
 {
 	PImplData->CurrentContextType = CONTEXT_Rendering;
-	if(GUseThreadedRendering)
+	if (GUseThreadedRendering)
 	{
 		SetCurrentContext(PImplData->RenderingContext.eglContext, PImplData->RenderingContext.eglSurface);
 	}
@@ -1519,8 +1551,23 @@ void AndroidEGL::SetCurrentRenderingContext()
 	}
 }
 
+void AndroidEGL::ReleaseContextOwnership()
+{
+	if (GetCurrentContextType() == CONTEXT_Rendering)
+	{
+		FPlatformMisc::LowLevelOutputDebugStringf(TEXT("AndroidEGL::  ReleaseContextOwnership,  thread releasing rendering context tid: %d"), FPlatformTLS::GetCurrentThreadId());
+		SetCurrentContext(EGL_NO_CONTEXT, EGL_NO_SURFACE);
+	}
+	else
+	{
+		FPlatformMisc::LowLevelOutputDebugStringf(TEXT("AndroidEGL::  ReleaseContextOwnership,  rendering context was not current to this thread tid: %d"), FPlatformTLS::GetCurrentThreadId());
+	}
+}
+
 void AndroidEGL::Terminate()
 {
+	FPlatformMisc::LowLevelOutputDebugStringf(TEXT("AndroidEGL:: Terminate! tid: %d"), FPlatformTLS::GetCurrentThreadId());
+
 	ResetDisplay();
 	DestroyContext(PImplData->SharedContext.eglContext);
 	PImplData->SharedContext.Reset();
@@ -1579,16 +1626,10 @@ void AndroidEGL::UnBind()
 
 void FAndroidAppEntry::ReInitWindow(void* NewNativeWindowHandle)
 {
-	FPlatformMisc::LowLevelOutputDebugString(TEXT("AndroidEGL::ReInitWindow()"));
+	check(IsInGameThread());
 
-	// Check for and call a registered window re-init callback. 
-	// RHIs may wish to perform additional operations when the window handle changes.
-	// Currently only vulkan RHI uses this.
-	const auto& OnReinitWindowCallback = FAndroidMisc::GetOnReInitWindowCallback();
-	if (OnReinitWindowCallback)
-	{
-		OnReinitWindowCallback(NewNativeWindowHandle);
-	}
+	// Window creation is now handled by BlockRendering, when it resumes after a new window is created.
+	FPlatformMisc::LowLevelOutputDebugString(TEXT("AndroidEGL::ReInitWindow()"));
 
 	// It isn't safe to call ShouldUseVulkan if AndroidEGL is not initialized.
 	// However, since we don't need to ReInit the window in that case anyways we
@@ -1601,10 +1642,27 @@ void FAndroidAppEntry::ReInitWindow(void* NewNativeWindowHandle)
 	// @todo vulkan: Clean this up, and does vulkan need any code here?
 	if (!FAndroidMisc::ShouldUseVulkan())
 	{
-		AndroidEGL::GetInstance()->ReInit();
+		// the window size could have been adjusted by the GT by now, if so it must be updated.
+		AndroidEGL::GetInstance()->RefreshWindowSize();
 	}
 }
 
+void AndroidEGL::RefreshWindowSize()
+{
+	check(IsInGameThread());
+	check(!FAndroidMisc::ShouldUseVulkan());
+	FPlatformRect WindowRect = FAndroidWindow::GetScreenRect();
+	UE_LOG(LogAndroid, Log, TEXT("AndroidEGL::RefreshWindowSize updating window size = %d, %d, cached size : %d, %d tid : %d"), WindowRect.Right, WindowRect.Bottom, PImplData->CachedWindowRect.Right, PImplData->CachedWindowRect.Bottom, FPlatformTLS::GetCurrentThreadId());
+	PImplData->CachedWindowRect = WindowRect;
+
+	ENQUEUE_RENDER_COMMAND(EGLResizeRenderContextSurface)(
+		[](FRHICommandListImmediate& RHICmdList)
+	{
+		RunOnGLRenderContextThread([&] {
+			AndroidEGL::GetInstance()->ResizeRenderContextSurface();
+		});
+	});
+}
 
 void FAndroidAppEntry::OnPauseEvent()
 {
@@ -1615,25 +1673,6 @@ void FAndroidAppEntry::OnPauseEvent()
 	}
 }
 
-void FAndroidAppEntry::DestroyWindow()
-{
-	FPlatformMisc::LowLevelOutputDebugString(TEXT("AndroidEGL::DestroyWindow()"));
-
-	// It isn't safe to call ShouldUseVulkan if AndroidEGL is not initialized.
-	// However, since we don't need to UnBind AndoirdEGL in that case anyways we
-	// can return early.
-	if (!AndroidEGL::GetInstance()->IsInitialized())
-	{
-		return;
-	}
-
-	// @todo vulkan: Clean this up, and does vulkan need any code here?
-	if (!FAndroidMisc::ShouldUseVulkan())
-	{
-		AndroidEGL::GetInstance()->UnBind();
-	}
-}
-
 void AndroidEGL::LogConfigInfo(EGLConfig  EGLConfigInfo)
 {
 	VERIFY_EGL_SCOPE();
@@ -1670,4 +1709,70 @@ void AndroidEGL::LogConfigInfo(EGLConfig  EGLConfigInfo)
 }
 
 
+///
+extern FCriticalSection GAndroidWindowLock;
+
+void BlockOnLostWindowRenderCommand(TSharedPtr<FEvent, ESPMode::ThreadSafe> RTBlockedTrigger)
+{
+	check(IsInRenderingThread());
+
+	UE_LOG(LogAndroid, Log, TEXT("Blocking renderer"));
+	if (FAndroidMisc::ShouldUseVulkan())
+	{
+		FRHICommandListImmediate& RHICmdList = FRHICommandListExecutor::GetImmediateCommandList();
+		if (IsRunningRHIInSeparateThread() && !RHICmdList.Bypass()) 
+		{
+			UE_LOG(LogAndroid, Log, TEXT("RendererBlock FlushRHIThread"));
+			RHICmdList.ImmediateFlush(EImmediateFlushType::FlushRHIThread);
+			UE_LOG(LogAndroid, Log, TEXT("RendererBlock DONE FlushRHIThread"));
+		}
+		RTBlockedTrigger->Trigger();
+
+		GAndroidWindowLock.Lock();
+		UE_LOG(LogAndroid, Log, TEXT("RendererBlock acquired window lock"));
+		const auto& OnReinitWindowCallback = FAndroidMisc::GetOnReInitWindowCallback();
+		if (OnReinitWindowCallback)
+		{
+			OnReinitWindowCallback(FAndroidWindow::GetHardwareWindow_EventThread());
+			UE_LOG(LogAndroid, Log, TEXT("RendererBlock updating window"));
+		}
+		GAndroidWindowLock.Unlock();
+	}
+	else
+	{
+		RunOnGLRenderContextThread([&] {
+			RTBlockedTrigger->Trigger();
+			GAndroidWindowLock.Lock();
+			UE_LOG(LogAndroid, Log, TEXT("RendererBlock acquired window lock"));
+			AndroidEGL::GetInstance()->SetRenderContextWindowSurface();
+			UE_LOG(LogAndroid, Log, TEXT("RendererBlock updating window"));
+			GAndroidWindowLock.Unlock();
+		}, true);
+	}
+	UE_LOG(LogAndroid, Log, TEXT("RendererBlock released window lock"));
+}
+
+extern bool IsInAndroidEventThread();
+void BlockRendering()
+{
+	check(IsInAndroidEventThread());
+	check(GIsRHIInitialized);
+
+	UE_LOG(LogAndroid, Log, TEXT("Blocking renderer on invalid window."));
+
+	TSharedPtr<FEvent, ESPMode::ThreadSafe> RTBlockedTrigger = MakeShareable(FPlatformProcess::GetSynchEventFromPool(), [](FEvent* EventToDelete)
+	{
+		FPlatformProcess::ReturnSynchEventToPool(EventToDelete);
+	});
+
+	FGraphEventRef RTBlockTask = FFunctionGraphTask::CreateAndDispatchWhenReady([RTBlockedTrigger]()
+	{
+		BlockOnLostWindowRenderCommand(RTBlockedTrigger);
+	}, TStatId(), NULL, ENamedThreads::GetRenderThread());
+
+	// wait for the render thread to process.
+	UE_LOG(LogAndroid, Log, TEXT("Waiting for renderer to encounter blocking command."));
+	RTBlockedTrigger->Wait();
+}
+
 #endif
diff --git a/Engine/Source/Runtime/OpenGLDrv/Private/Android/AndroidEGL.h b/Engine/Source/Runtime/OpenGLDrv/Private/Android/AndroidEGL.h
--- a/Engine/Source/Runtime/OpenGLDrv/Private/Android/AndroidEGL.h
+++ b/Engine/Source/Runtime/OpenGLDrv/Private/Android/AndroidEGL.h
@@ -80,18 +80,25 @@ public:
 	EGLContext CreateContext(EGLContext InSharedContext = EGL_NO_CONTEXT);
 	int32 GetError();
 	EGLBoolean SetCurrentContext(EGLContext InContext, EGLSurface InSurface);
+
+	void AcquireCurrentRenderingContext();
+	void ReleaseContextOwnership();
+
 	GLuint GetOnScreenColorRenderBuffer();
 	GLuint GetResolveFrameBuffer();
 	bool IsCurrentContextValid();
 	EGLContext  GetCurrentContext(  );
 	void SetCurrentSharedContext();
-	void SetSharedContext();
-	void SetSingleThreadRenderingContext();
-	void SetMultithreadRenderingContext();
 	void SetCurrentRenderingContext();
 	uint32_t GetCurrentContextType();
 	FPlatformOpenGLContext* GetRenderingContext();
-	
+
+	// recreate the EGL surface for the current hardware window.
+	void SetRenderContextWindowSurface();
+
+	// Called from game thread when a window is reinited.
+	void RefreshWindowSize();
+
 protected:
 	AndroidEGL();
 	static AndroidEGL* Singleton ;
@@ -113,6 +120,9 @@ private:
 	void ResetInternal();
 	void LogConfigInfo(EGLConfig  EGLConfigInfo);
 
+	// Actual Update to the egl surface to match the GT's requested size.
+	void ResizeRenderContextSurface();
+
 	bool bSupportsKHRCreateContext;
 	bool bSupportsKHRSurfacelessContext;
 	bool bSupportsKHRNoErrorContext;
diff --git a/Engine/Source/Runtime/OpenGLDrv/Private/Android/AndroidOpenGL.cpp b/Engine/Source/Runtime/OpenGLDrv/Private/Android/AndroidOpenGL.cpp
--- a/Engine/Source/Runtime/OpenGLDrv/Private/Android/AndroidOpenGL.cpp
+++ b/Engine/Source/Runtime/OpenGLDrv/Private/Android/AndroidOpenGL.cpp
@@ -236,7 +236,7 @@ void PlatformSharedContextSetup(FPlatformOpenGLDevice* Device)
 
 void PlatformNULLContextSetup()
 {
-	AndroidEGL::GetInstance()->SetCurrentContext(EGL_NO_CONTEXT, EGL_NO_SURFACE);
+	AndroidEGL::GetInstance()->ReleaseContextOwnership();
 }
 
 EOpenGLCurrentContext PlatformOpenGLCurrentContext(FPlatformOpenGLDevice* Device)
@@ -457,7 +457,7 @@ void PlatformDestroyOpenGLDevice(FPlatformOpenGLDevice* Device)
 
 void FPlatformOpenGLDevice::SetCurrentRenderingContext()
 {
-	AndroidEGL::GetInstance()->SetCurrentRenderingContext();
+	AndroidEGL::GetInstance()->AcquireCurrentRenderingContext();
 }
 
 void PlatformLabelObjects()
diff --git a/Engine/Source/Runtime/OpenGLDrv/Private/OpenGLViewport.cpp b/Engine/Source/Runtime/OpenGLDrv/Private/OpenGLViewport.cpp
--- a/Engine/Source/Runtime/OpenGLDrv/Private/OpenGLViewport.cpp
+++ b/Engine/Source/Runtime/OpenGLDrv/Private/OpenGLViewport.cpp
@@ -232,7 +232,7 @@ FOpenGLViewport::FOpenGLViewport(FOpenGLDynamicRHI* InOpenGLRHI,void* InWindowHa
 {
 	check(OpenGLRHI);
 	// @todo lumin: Add a "PLATFORM_HAS_NO_NATIVE_WINDOW" or something
-#if !PLATFORM_LUMIN
+#if !PLATFORM_LUMIN && !PLATFORM_ANDROID
 	check(InWindowHandle);
 #endif
 	check(IsInGameThread());
diff --git a/Engine/Source/Runtime/RHI/Private/Android/AndroidDynamicRHI.cpp b/Engine/Source/Runtime/RHI/Private/Android/AndroidDynamicRHI.cpp
--- a/Engine/Source/Runtime/RHI/Private/Android/AndroidDynamicRHI.cpp
+++ b/Engine/Source/Runtime/RHI/Private/Android/AndroidDynamicRHI.cpp
@@ -46,5 +46,9 @@ FDynamicRHI* PlatformCreateDynamicRHI()
 		DynamicRHI = DynamicRHIModule->CreateRHI(RequestedFeatureLevel);
 	}
 
+#if !PLATFORM_LUMIN
+	FPlatformMisc::UnlockAndroidWindow();
+#endif
+
 	return DynamicRHI;
 }
diff --git a/Engine/Source/Runtime/VulkanRHI/Private/Android/VulkanAndroidPlatform.cpp b/Engine/Source/Runtime/VulkanRHI/Private/Android/VulkanAndroidPlatform.cpp
--- a/Engine/Source/Runtime/VulkanRHI/Private/Android/VulkanAndroidPlatform.cpp
+++ b/Engine/Source/Runtime/VulkanRHI/Private/Android/VulkanAndroidPlatform.cpp
@@ -114,7 +114,7 @@ void FVulkanAndroidPlatform::FreeVulkanLibrary()
 void FVulkanAndroidPlatform::CreateSurface(void* WindowHandle, VkInstance Instance, VkSurfaceKHR* OutSurface)
 {
 	// don't use cached window handle coming from VulkanViewport, as it could be gone by now
-	WindowHandle = FAndroidWindow::GetHardwareWindow();
+	WindowHandle = FAndroidWindow::GetHardwareWindow_EventThread();
 	if (WindowHandle == NULL)
 	{
 
@@ -196,16 +196,16 @@ void FVulkanAndroidPlatform::OverridePlatformHandlers(bool bInit)
 
 void FVulkanAndroidPlatform::BlockUntilWindowIsAvailable()
 {
-	void* WindowHandle = FAndroidWindow::GetHardwareWindow();
+	void* WindowHandle = FAndroidWindow::GetHardwareWindow_EventThread();
 	if (WindowHandle == nullptr)
 	{
 		// Sleep if the hardware window isn't currently available.
 		// The Window may not exist if the activity is pausing/resuming, in which case we make this thread wait
 		FPlatformMisc::LowLevelOutputDebugString(TEXT("Waiting for Native window in FVulkanAndroidPlatform::BlockUntilWindowIsAvailable()"));
-		while (WindowHandle == nullptr)
+		WindowHandle = FAndroidWindow::WaitForHardwareWindow();
+		if (WindowHandle == NULL)
 		{
-			FPlatformProcess::Sleep(0.001f);
-			WindowHandle = FAndroidWindow::GetHardwareWindow();
+			FPlatformMisc::LowLevelOutputDebugString(TEXT("Aborted FVulkanAndroidPlatform::BlockUntilWindowIsAvailable, FAndroidWindow::WaitForHardwareWindow() returned null"));
 		}
 	}
 }
diff --git a/Engine/Source/Runtime/VulkanRHI/Private/VulkanRHI.cpp b/Engine/Source/Runtime/VulkanRHI/Private/VulkanRHI.cpp
--- a/Engine/Source/Runtime/VulkanRHI/Private/VulkanRHI.cpp
+++ b/Engine/Source/Runtime/VulkanRHI/Private/VulkanRHI.cpp
@@ -1727,17 +1727,25 @@ void FVulkanDynamicRHI::RecreateSwapChain(void* NewNativeWindow)
 {
 	if (NewNativeWindow)
 	{
-		FlushRenderingCommands();
+		if (IsInGameThread())
+		{
+			FlushRenderingCommands();
+		}
+
 		TArray<FVulkanViewport*> Viewports = GVulkanRHI->Viewports;
 		ENQUEUE_RENDER_COMMAND(VulkanRecreateSwapChain)(
 			[Viewports, NewNativeWindow](FRHICommandListImmediate& RHICmdList)
+		{
+			for (auto& Viewport : Viewports)
 			{
-				for (auto& Viewport : Viewports)
-				{
-					Viewport->RecreateSwapchain(NewNativeWindow);
-				}
-			});
-		FlushRenderingCommands();
+				Viewport->RecreateSwapchain(NewNativeWindow);
+			}
+		});
+
+		if (IsInGameThread())
+		{
+			FlushRenderingCommands();
+		}
 	}
 }
 
diff --git a/Engine/Source/Runtime/WebBrowser/Private/Android/AndroidWebBrowserWidget.cpp b/Engine/Source/Runtime/WebBrowser/Private/Android/AndroidWebBrowserWidget.cpp
--- a/Engine/Source/Runtime/WebBrowser/Private/Android/AndroidWebBrowserWidget.cpp
+++ b/Engine/Source/Runtime/WebBrowser/Private/Android/AndroidWebBrowserWidget.cpp
@@ -196,7 +196,7 @@ void SAndroidWebBrowserWidget::Tick(const FGeometry& AllottedGeometry, const dou
 	float UIScale;
 	FPlatformRect ScreenRect = FAndroidWindow::GetScreenRect();
 	int32_t ScreenWidth, ScreenHeight;
-	FAndroidWindow::CalculateSurfaceSize(FAndroidWindow::GetHardwareWindow(), ScreenWidth, ScreenHeight);
+	FAndroidWindow::CalculateSurfaceSize(ScreenWidth, ScreenHeight);
 	if (ScreenWidth > ScreenHeight)
 	{
 		UIScale = (float)ScreenHeight / (ScreenRect.Bottom - ScreenRect.Top);
