cvar to turn accessibility on/off (defaults off)

diff --git a/Engine/Source/Runtime/ApplicationCore/Private/GenericPlatform/GenericAccessibleInterfaces.cpp b/Engine/Source/Runtime/ApplicationCore/Private/GenericPlatform/GenericAccessibleInterfaces.cpp
--- a/Engine/Source/Runtime/ApplicationCore/Private/GenericPlatform/GenericAccessibleInterfaces.cpp
+++ b/Engine/Source/Runtime/ApplicationCore/Private/GenericPlatform/GenericAccessibleInterfaces.cpp
@@ -3,3 +3,40 @@
 #include "GenericPlatform/GenericAccessibleInterfaces.h"
 
 DEFINE_LOG_CATEGORY(LogAccessibility);
+
+#if WITH_ACCESSIBILITY
+#include "HAL/IConsoleManager.h"
+
+/** A flag that can be set by a user to force accessibility off regardless of other settings. */
+bool GAllowAccessibility = false;
+FAutoConsoleVariableRef AllowAccessibilityRef(
+	TEXT("Accessibility.Enable"),
+	GAllowAccessibility,
+	TEXT("If false, all queries from accessible APIs will be ignored. On some platforms, the application must be restarted in order to take effect.")
+);
+
+bool FGenericAccessibleMessageHandler::ApplicationIsAccessible() const
+{
+	return bApplicationIsAccessible && GAllowAccessibility;
+}
+
+void FGenericAccessibleMessageHandler::SetActive(bool bActive)
+{
+	bActive &= GAllowAccessibility;
+	if (bActive != bIsActive)
+	{
+		bIsActive = bActive;
+
+		if (bIsActive)
+		{
+			UE_LOG(LogAccessibility, Verbose, TEXT("Enabling Accessibility"));
+			OnActivate();
+		}
+		else
+		{
+			OnDeactivate();
+			UE_LOG(LogAccessibility, Verbose, TEXT("Accessibility Disabled"));
+		}
+	}
+}
+#endif
diff --git a/Engine/Source/Runtime/ApplicationCore/Private/IOS/Accessibility/IOSAccessibilityCache.cpp b/Engine/Source/Runtime/ApplicationCore/Private/IOS/Accessibility/IOSAccessibilityCache.cpp
--- a/Engine/Source/Runtime/ApplicationCore/Private/IOS/Accessibility/IOSAccessibilityCache.cpp
+++ b/Engine/Source/Runtime/ApplicationCore/Private/IOS/Accessibility/IOSAccessibilityCache.cpp
@@ -27,7 +27,7 @@ static void DumpAccessibilityStatsForwarder()
 #if !UE_BUILD_SHIPPING
 	IConsoleManager::Get().RegisterConsoleCommand
 	(
-		TEXT("DumpAccessibilityStatsIOS"),
+		TEXT("Accessibility.DumpStatsIOS"),
 		TEXT("Writes to LogAccessibility the memory stats for the platform-level accessibility data (AccessibilityElements) required for IOS support."),
 		FConsoleCommandDelegate::CreateStatic(&DumpAccessibilityStatsForwarder),
 		ECVF_Default
diff --git a/Engine/Source/Runtime/ApplicationCore/Private/IOS/IOSAppDelegate.cpp b/Engine/Source/Runtime/ApplicationCore/Private/IOS/IOSAppDelegate.cpp
--- a/Engine/Source/Runtime/ApplicationCore/Private/IOS/IOSAppDelegate.cpp
+++ b/Engine/Source/Runtime/ApplicationCore/Private/IOS/IOSAppDelegate.cpp
@@ -1300,7 +1300,7 @@ static FAutoConsoleVariableRef CVarGEnableThermalsReport(
 #if WITH_ACCESSIBILITY
 -(void)OnVoiceOverStatusChanged
 {
-	if (UIAccessibilityIsVoiceOverRunning())
+	if (UIAccessibilityIsVoiceOverRunning() && self.IOSApplication->GetAccessibleMessageHandler()->ApplicationIsAccessible())
 	{
 		// This must happen asynchronously because when the app activates from a suspended state,
 		// the IOS notification will emit before the game thread wakes up. This does mean that the
@@ -1324,7 +1324,7 @@ static FAutoConsoleVariableRef CVarGEnableThermalsReport(
 			});
 		}, TStatId(), NULL, ENamedThreads::GameThread);
 	}
-	else
+	else if (AccessibilityCacheTimer != nil)
 	{
 		[AccessibilityCacheTimer invalidate];
 		AccessibilityCacheTimer = nil;
diff --git a/Engine/Source/Runtime/ApplicationCore/Private/Windows/Accessibility/WindowsUIAManager.cpp b/Engine/Source/Runtime/ApplicationCore/Private/Windows/Accessibility/WindowsUIAManager.cpp
--- a/Engine/Source/Runtime/ApplicationCore/Private/Windows/Accessibility/WindowsUIAManager.cpp
+++ b/Engine/Source/Runtime/ApplicationCore/Private/Windows/Accessibility/WindowsUIAManager.cpp
@@ -40,7 +40,7 @@ FWindowsUIAManager::FWindowsUIAManager(const FWindowsApplication& InApplication)
 #if !UE_BUILD_SHIPPING
 	IConsoleManager::Get().RegisterConsoleCommand
 	(
-		TEXT("DumpAccessibilityStatsWindows"),
+		TEXT("Accessibility.DumpStatsWindows"),
 		TEXT("Writes to LogAccessibility the memory stats for the platform-level accessibility data (Providers) required for Windows support."),
 		FConsoleCommandDelegate::CreateRaw(this, &FWindowsUIAManager::DumpAccessibilityStats),
 		ECVF_Default
diff --git a/Engine/Source/Runtime/ApplicationCore/Public/GenericPlatform/GenericAccessibleInterfaces.h b/Engine/Source/Runtime/ApplicationCore/Public/GenericPlatform/GenericAccessibleInterfaces.h
--- a/Engine/Source/Runtime/ApplicationCore/Public/GenericPlatform/GenericAccessibleInterfaces.h
+++ b/Engine/Source/Runtime/ApplicationCore/Public/GenericPlatform/GenericAccessibleInterfaces.h
@@ -429,7 +429,7 @@ public:
 		/** If this was a property changed event, the 'after' value. This may also be set for other events such as Notification. */
 		FVariant);
 
-	FGenericAccessibleMessageHandler() : bIsActive(false) {}
+	FGenericAccessibleMessageHandler() : bApplicationIsAccessible(false), bIsActive(false) {}
 
 	virtual ~FGenericAccessibleMessageHandler()
 	{
@@ -444,7 +444,7 @@ public:
 	 *
 	 * @return true if the application intends to return valid accessible widgets when queried.
 	 */
-	virtual bool ApplicationIsAccessible() const { return false; }
+	bool ApplicationIsAccessible() const;
 
 	/**
 	 * Checks if accessibility is enabled in the application. Usually this happens when screen-reading software is turned on.
@@ -459,24 +459,7 @@ public:
 	 *
 	 * @param bActive Whether to enable to disable the message handler.
 	 */
-	void SetActive(bool bActive)
-	{
-		if (bActive != bIsActive)
-		{
-			bIsActive = bActive;
-
-			if (bIsActive)
-			{
-				UE_LOG(LogAccessibility, Verbose, TEXT("Enabling Accessibility"));
-				OnActivate();
-			}
-			else
-			{
-				OnDeactivate();
-				UE_LOG(LogAccessibility, Verbose, TEXT("Accessibility Disabled"));
-			}
-		}
-	}
+	void SetActive(bool bActive);
 
 	/**
 	 * Creates or retrieves an accessible object for a native OS window.
@@ -533,6 +516,9 @@ protected:
 	/** Triggered when bIsActive changes from true to false. */
 	virtual void OnDeactivate() {}
 
+	/** Subclasses should override this to indicate that they support accessibility. */
+	bool bApplicationIsAccessible;
+
 private:
 	/** Whether or not accessibility is currently enabled in the application */
 	bool bIsActive;
diff --git a/Engine/Source/Runtime/SlateCore/Private/Widgets/Accessibility/SlateAccessibleMessageHandler.cpp b/Engine/Source/Runtime/SlateCore/Private/Widgets/Accessibility/SlateAccessibleMessageHandler.cpp
--- a/Engine/Source/Runtime/SlateCore/Private/Widgets/Accessibility/SlateAccessibleMessageHandler.cpp
+++ b/Engine/Source/Runtime/SlateCore/Private/Widgets/Accessibility/SlateAccessibleMessageHandler.cpp
@@ -18,6 +18,7 @@ FSlateAccessibleMessageHandler::FSlateAccessibleMessageHandler()
 	: FGenericAccessibleMessageHandler()
 	, bDirty(false)
 {
+	bApplicationIsAccessible = true;
 }
 
 void FSlateAccessibleMessageHandler::OnActivate()
diff --git a/Engine/Source/Runtime/SlateCore/Private/Widgets/Accessibility/SlateAccessibleWidgetCache.cpp b/Engine/Source/Runtime/SlateCore/Private/Widgets/Accessibility/SlateAccessibleWidgetCache.cpp
--- a/Engine/Source/Runtime/SlateCore/Private/Widgets/Accessibility/SlateAccessibleWidgetCache.cpp
+++ b/Engine/Source/Runtime/SlateCore/Private/Widgets/Accessibility/SlateAccessibleWidgetCache.cpp
@@ -89,7 +89,7 @@ void FSlateAccessibleWidgetCache::DumpAccessibilityStats()
 
 static FAutoConsoleCommand DumpAccessibilityStatsSlateCommand
 (
-	TEXT("DumpAccessibilityStatsSlate"),
+	TEXT("Accessibility.DumpStatsSlate"),
 	TEXT("Writes memory stats for Slate's accessibility data stored to LogAccessibility."),
 	FConsoleCommandDelegate::CreateStatic(&FSlateAccessibleWidgetCache::DumpAccessibilityStats)
 );
diff --git a/Engine/Source/Runtime/SlateCore/Public/Widgets/Accessibility/SlateAccessibleMessageHandler.h b/Engine/Source/Runtime/SlateCore/Public/Widgets/Accessibility/SlateAccessibleMessageHandler.h
--- a/Engine/Source/Runtime/SlateCore/Public/Widgets/Accessibility/SlateAccessibleMessageHandler.h
+++ b/Engine/Source/Runtime/SlateCore/Public/Widgets/Accessibility/SlateAccessibleMessageHandler.h
@@ -26,7 +26,6 @@ public:
 	virtual TSharedPtr<IAccessibleWidget> GetAccessibleWindow(const TSharedRef<FGenericWindow>& InWindow) const override;
 	virtual AccessibleWidgetId GetAccessibleWindowId(const TSharedRef<FGenericWindow>& InWindow) const override;
 	virtual TSharedPtr<IAccessibleWidget> GetAccessibleWidgetFromId(AccessibleWidgetId Id) const override;
-	virtual bool ApplicationIsAccessible() const override { return true; }
 	//~
 
 	/**
