Fix for CIS build
#rnx

diff --git a/Engine/Source/Runtime/ApplicationCore/Private/Android/AndroidWindow.cpp b/Engine/Source/Runtime/ApplicationCore/Private/Android/AndroidWindow.cpp
--- a/Engine/Source/Runtime/ApplicationCore/Private/Android/AndroidWindow.cpp
+++ b/Engine/Source/Runtime/ApplicationCore/Private/Android/AndroidWindow.cpp
@@ -164,8 +164,10 @@ void FAndroidWindow::ReleaseWindowRef(ANativeWindow* InWindow)
 
  void FAndroidWindow::SetHardwareWindow_EventThread(void* InWindow)
 {
+#if USE_ANDROID_JNI
 	check(IsInAndroidEventThread());
 	NativeWindow = InWindow; //using raw native window handle for now. Could be changed to use AndroidWindow later if needed
+#endif
 }
 
 void* FAndroidWindow::GetHardwareWindow_EventThread()
