Fix NonUnity iOS compile

diff --git a/Engine/Source/Runtime/ApplicationCore/Private/IOS/IOSInputInterface.cpp b/Engine/Source/Runtime/ApplicationCore/Private/IOS/IOSInputInterface.cpp
--- a/Engine/Source/Runtime/ApplicationCore/Private/IOS/IOSInputInterface.cpp
+++ b/Engine/Source/Runtime/ApplicationCore/Private/IOS/IOSInputInterface.cpp
@@ -261,8 +261,6 @@ void ModifyVectorByOrientation(FVector& Vec, bool bIsRotation)
 
 void FIOSInputInterface::ProcessTouchesAndKeys(uint32 ControllerId, const TArray<TouchInput>& InTouchInputStack, const TArray<int32>& InKeyInputStack)
 {
-	QUICK_SCOPE_CYCLE_COUNTER(STAT_IOSInputInterface_ProcessTouchesAndKeys);
-	
 	for(int i = 0; i < InTouchInputStack.Num(); ++i)
 	{
 		const TouchInput& Touch = InTouchInputStack[i];
