Optionally allow external audio (audio from other apps) to continue to play while our app is active. Off by default.

diff --git a/Engine/Source/Runtime/ApplicationCore/Private/IOS/IOSAppDelegate.cpp b/Engine/Source/Runtime/ApplicationCore/Private/IOS/IOSAppDelegate.cpp
--- a/Engine/Source/Runtime/ApplicationCore/Private/IOS/IOSAppDelegate.cpp
+++ b/Engine/Source/Runtime/ApplicationCore/Private/IOS/IOSAppDelegate.cpp
@@ -501,9 +501,21 @@ static IOSAppDelegate* CachedDelegate = nil;
 		}
 	}];
 
+	self.bAudioSessionInitialized = true;
+
 	self.bUsingBackgroundMusic = [self IsBackgroundAudioPlaying];
 	self.bForceEmitOtherAudioPlaying = true;
+	
+	self.bAllowExternalAudio = false;
+	bool bAllowExternalAudio = false;
+	GConfig->GetBool(TEXT("/Script/IOSRuntimeSettings.IOSRuntimeSettings"), TEXT("bAudioAllowExternalAudio"), bAllowExternalAudio, GEngineIni);
+	self.bAllowExternalAudio = bAllowExternalAudio;
 
+	if(bAllowExternalAudio)
+	{
+		[self SetFeature:EAudioFeature::ExternalAudio Active:true];
+	}
+	
 #if USE_MUTE_SWITCH_DETECTION
 	// Initialize the mute switch detector.
 	[SharkfoodMuteSwitchDetector shared];
@@ -526,6 +538,11 @@ static IOSAppDelegate* CachedDelegate = nil;
 
 - (void)ToggleAudioSession:(bool)bActive force:(bool)bForce
 {
+	if(!self.bAudioSessionInitialized)
+	{
+		return;
+	}
+	
 	// @todo kairos: resolve old vs new before we go to main
 	if (false)
 	{
@@ -534,30 +551,20 @@ static IOSAppDelegate* CachedDelegate = nil;
 		bool bIsBackground = GIsSuspended;
 		bActive = !bIsBackground || [self IsFeatureActive:EAudioFeature::BackgroundAudio];
 		
+		NSError* ActiveError = nil;
+
 		// @todo maybe check the active states, not bForce?
 		if (self.bAudioActive != bActive || bForce)
 		{
-			// enable or disable audio
-			NSError* ActiveError = nil;
-			[[AVAudioSession sharedInstance] setActive:bActive error:&ActiveError];
-			if (ActiveError)
-			{
-				UE_LOG(LogIOSAudioSession, Error, TEXT("Failed to set audio session to active = %d [Error = %s]"), bActive, *FString([ActiveError description]));
-			}
-			else
-			{
-				self.bAudioActive = bActive;
-			}
-			
-			if (self.bAudioActive)
+			if (bActive || bForce)
 			{
 				// get the category and settings to use
-				/*AVAudioSessionCategory*/NSString* Category = AVAudioSessionCategorySoloAmbient;
+				/*AVAudioSessionCategory*/NSString* Category = self.bAllowExternalAudio ? AVAudioSessionCategoryAmbient : AVAudioSessionCategorySoloAmbient;
 				/*AVAudioSessionMode*/NSString* Mode = AVAudioSessionModeDefault;
 				AVAudioSessionCategoryOptions Options = 0;
 
 				// attempt to mix with other apps if desired
-				if ([self IsFeatureActive:EAudioFeature::BackgroundAudio])
+				if ([self IsFeatureActive:EAudioFeature::BackgroundAudio] || [self IsFeatureActive:EAudioFeature::ExternalAudio])
 				{
 					Options |= AVAudioSessionCategoryOptionMixWithOthers;
 				}
@@ -587,7 +594,7 @@ static IOSAppDelegate* CachedDelegate = nil;
 				{
 					Category = AVAudioSessionCategoryRecord;
 				}
-
+				
 				// set the category (the most important part here)
 				[[AVAudioSession sharedInstance] setCategory:Category mode:Mode options:Options error:&ActiveError];
 				if (ActiveError)
@@ -595,6 +602,17 @@ static IOSAppDelegate* CachedDelegate = nil;
 					UE_LOG(LogIOSAudioSession, Error, TEXT("Failed to set audio session category to %s! [Error = %s]"), *FString(Category), *FString([ActiveError description]));
 				}
 			}
+			
+			// enable or disable audio
+			[[AVAudioSession sharedInstance] setActive:bActive error:&ActiveError];
+			if (ActiveError)
+			{
+				UE_LOG(LogIOSAudioSession, Error, TEXT("Failed to set audio session to active = %d [Error = %s]"), bActive, *FString([ActiveError description]));
+			}
+			else
+			{
+				self.bAudioActive = bActive;
+			}
 		}
 	}
 	// old style below
@@ -1287,8 +1305,11 @@ static FAutoConsoleVariableRef CVarGEnableThermalsReport(
         [[NSNotificationCenter defaultCenter] addObserver:self selector:@selector(batteryStateChanged:) name:UIDeviceBatteryStateDidChangeNotification object:nil];
 	}
 #endif
-    
-	[self InitializeAudioSession];
+	
+	self.bAudioSessionInitialized = false;
+	
+	// InitializeAudioSession is now called from FEngineLoop::AppInit after the config system is initialized
+//	[self InitializeAudioSession];
 
 #if WITH_ACCESSIBILITY
 	[[NSNotificationCenter defaultCenter] addObserver:self selector:@selector(OnVoiceOverStatusChanged) name:UIAccessibilityVoiceOverStatusDidChangeNotification object:nil];
diff --git a/Engine/Source/Runtime/ApplicationCore/Public/IOS/IOSAppDelegate.h b/Engine/Source/Runtime/ApplicationCore/Public/IOS/IOSAppDelegate.h
--- a/Engine/Source/Runtime/ApplicationCore/Public/IOS/IOSAppDelegate.h
+++ b/Engine/Source/Runtime/ApplicationCore/Public/IOS/IOSAppDelegate.h
@@ -190,11 +190,16 @@ APPLICATIONCORE_API
 -(float)GetBackgroundingMainThreadBlockTime;
 -(void)OverrideBackgroundingMainThreadBlockTime:(float)BlockTime;
 
+@property (assign) bool bAudioSessionInitialized;
+
 /** TRUE if the device is playing background music and we want to allow that */
 @property (assign) bool bUsingBackgroundMusic;
 @property (assign) bool bLastOtherAudioPlaying;
 @property (assign) bool bForceEmitOtherAudioPlaying;
 
+/** TRUE if audio from other apps will be allowed to play while this app is active */
+@property (assign) bool bAllowExternalAudio;
+
 #if USE_MUTE_SWITCH_DETECTION
 @property (assign) bool bLastMutedState;
 @property (assign) bool bForceEmitMutedState;
diff --git a/Engine/Source/Runtime/Launch/Private/LaunchEngineLoop.cpp b/Engine/Source/Runtime/Launch/Private/LaunchEngineLoop.cpp
--- a/Engine/Source/Runtime/Launch/Private/LaunchEngineLoop.cpp
+++ b/Engine/Source/Runtime/Launch/Private/LaunchEngineLoop.cpp
@@ -225,6 +225,10 @@ class FFeedbackContext;
 #define WITH_CONFIG_PATCHING 0
 #endif
 
+#if PLATFORM_IOS || PLATFORM_TVOS
+#include "IOS/IOSAppDelegate.h"
+#endif
+
 
 int32 GUseDisregardForGCOnDedicatedServers = 1;
 static FAutoConsoleVariableRef CVarUseDisregardForGCOnDedicatedServers(
@@ -4881,6 +4885,12 @@ bool FEngineLoop::AppInit( )
 	// Put the command line and config info into the suppression system (before plugins start loading)
 	FLogSuppressionInterface::Get().ProcessConfigAndCommandLine();
 
+#if PLATFORM_IOS || PLATFORM_TVOS
+	// Now that the config system is ready, init the audio system.
+	[[IOSAppDelegate GetDelegate] InitializeAudioSession];
+#endif
+
+	
 	// NOTE: This is the earliest place to init the online subsystems (via plugins)
 	// Code needs GConfigFile to be valid
 	// Must be after FThreadStats::StartThread();
