Only update controller force feedback state if the values have changed
See https://udn.unrealengine.com/questions/457822/unnecessary-calls-made-to-xinputsetstate.html

diff --git a/Engine/Source/Runtime/ApplicationCore/Private/Windows/XInputInterface.cpp b/Engine/Source/Runtime/ApplicationCore/Private/Windows/XInputInterface.cpp
--- a/Engine/Source/Runtime/ApplicationCore/Private/Windows/XInputInterface.cpp
+++ b/Engine/Source/Runtime/ApplicationCore/Private/Windows/XInputInterface.cpp
@@ -225,15 +225,21 @@ void XInputInterface::SendControllerEvents()
 			}	
 
 			// apply force feedback
-			XINPUT_VIBRATION VibrationState;
- 
+
 			const float LargeValue = (ControllerState.ForceFeedback.LeftLarge > ControllerState.ForceFeedback.RightLarge ? ControllerState.ForceFeedback.LeftLarge : ControllerState.ForceFeedback.RightLarge);
 			const float SmallValue = (ControllerState.ForceFeedback.LeftSmall > ControllerState.ForceFeedback.RightSmall ? ControllerState.ForceFeedback.LeftSmall : ControllerState.ForceFeedback.RightSmall);
 
-			VibrationState.wLeftMotorSpeed = ( ::WORD ) ( LargeValue * 65535.0f );
-			VibrationState.wRightMotorSpeed = ( ::WORD ) ( SmallValue * 65535.0f );
+			if (!FMath::IsNearlyEqual(LargeValue, ControllerState.LastLargeValue) || !FMath::IsNearlyEqual(SmallValue, ControllerState.LastSmallValue))
+			{
+				XINPUT_VIBRATION VibrationState;
+				VibrationState.wLeftMotorSpeed = ( ::WORD ) ( LargeValue * 65535.0f );
+				VibrationState.wRightMotorSpeed = ( ::WORD ) ( SmallValue * 65535.0f );
  
-			XInputSetState( ( ::DWORD ) ControllerState.ControllerId, &VibrationState );			
+				XInputSetState( ( ::DWORD ) ControllerState.ControllerId, &VibrationState );
+
+				ControllerState.LastLargeValue = LargeValue;
+				ControllerState.LastSmallValue = SmallValue;
+			}
 		}
 	}
 
diff --git a/Engine/Source/Runtime/ApplicationCore/Private/Windows/XInputInterface.h b/Engine/Source/Runtime/ApplicationCore/Private/Windows/XInputInterface.h
--- a/Engine/Source/Runtime/ApplicationCore/Private/Windows/XInputInterface.h
+++ b/Engine/Source/Runtime/ApplicationCore/Private/Windows/XInputInterface.h
@@ -86,14 +86,17 @@ private:
 		/** Right trigger analog value */
 		uint8 RightTriggerAnalog;
 
-		/** Id of the controller */
-		int32 ControllerId;
-
 		/** If the controller is currently connected */
 		bool bIsConnected;
-		
+
+		/** Id of the controller */
+		int32 ControllerId;
+	
 		/** Current force feedback values */
 		FForceFeedbackValues ForceFeedback;
+
+		float LastLargeValue;
+		float LastSmallValue;
 	};
 
 	/** If we've been notified by the system that the controller state may have changed */
