Fixed issues with window position and size when toggling between fullscreen and windowed modes

diff --git a/Engine/Source/Runtime/ApplicationCore/Private/Mac/MacApplication.cpp b/Engine/Source/Runtime/ApplicationCore/Private/Mac/MacApplication.cpp
--- a/Engine/Source/Runtime/ApplicationCore/Private/Mac/MacApplication.cpp
+++ b/Engine/Source/Runtime/ApplicationCore/Private/Mac/MacApplication.cpp
@@ -122,6 +122,8 @@ FMacApplication::FMacApplication()
 	FCoreDelegates::PreSlateModal.AddRaw(this, &FMacApplication::StartScopedModalEvent);
     FCoreDelegates::PostSlateModal.AddRaw(this, &FMacApplication::EndScopedModalEvent);
 #endif
+
+	FMacApplication::OnDisplayReconfiguration(kCGNullDirectDisplay, kCGDisplayDesktopShapeChangedFlag, this);
 }
 
 FMacApplication::~FMacApplication()
diff --git a/Engine/Source/Runtime/ApplicationCore/Private/Mac/MacWindow.cpp b/Engine/Source/Runtime/ApplicationCore/Private/Mac/MacWindow.cpp
--- a/Engine/Source/Runtime/ApplicationCore/Private/Mac/MacWindow.cpp
+++ b/Engine/Source/Runtime/ApplicationCore/Private/Mac/MacWindow.cpp
@@ -610,7 +610,10 @@ void FMacWindow::ApplySizeAndModeChanges(int32 X, int32 Y, int32 Width, int32 He
 			{
 				MainThreadCall(^{
 					SCOPED_AUTORELEASE_POOL;
-					[WindowHandle setFrame:Rect display:YES];
+					if (!bIsFullScreen) // Don't set the frame rect when switching from fullscreen because we don't know the correct window size yet. The OS will resize the window for us.
+					{
+						[WindowHandle setFrame:Rect display:YES];
+					}
 
 					if (Definition->ShouldPreserveAspectRatio)
 					{
@@ -667,6 +670,10 @@ void FMacWindow::UpdateFullScreenState(bool bToggleFullScreen)
 		SCOPED_AUTORELEASE_POOL;
 		if (bToggleFullScreen)
 		{
+			// Make sure we don't limit the window size for fullscreen toggle
+			[WindowHandle setMinSize:NSMakeSize(10.0f, 10.0f)];
+			[WindowHandle setMaxSize:NSMakeSize(10000.0f, 10000.0f)];
+
 			[WindowHandle toggleFullScreen:nil];
 		}
 		else
@@ -699,4 +706,7 @@ void FMacWindow::UpdateFullScreenState(bool bToggleFullScreen)
 		FPlatformApplicationMisc::PumpMessages(true);
 		bModeChanged = [WindowHandle windowMode] == WindowHandle.TargetWindowMode;
 	} while (!bModeChanged);
+
+	// Restore window size limits if needed
+	MacApplication->OnCursorLock();
 }
