Update mouse position on NSDraggingUpdate event on Mac

diff --git a/Engine/Source/Runtime/ApplicationCore/Private/Mac/MacApplication.cpp b/Engine/Source/Runtime/ApplicationCore/Private/Mac/MacApplication.cpp
--- a/Engine/Source/Runtime/ApplicationCore/Private/Mac/MacApplication.cpp
+++ b/Engine/Source/Runtime/ApplicationCore/Private/Mac/MacApplication.cpp
@@ -673,6 +673,20 @@ void FMacApplication::ProcessEvent(const FDeferredMacEvent& Event)
 		}
 		else if (Event.NotificationName == NSDraggingUpdated)
 		{
+			// MouseMoved events are suspended during drag and drop operations, so we need to update the cursor position here
+			NSPoint CursorPos = [NSEvent mouseLocation];
+			FVector2D NewPosition = ConvertCocoaPositionToSlate(CursorPos.x, CursorPos.y);
+			FMacCursor* MacCursor = (FMacCursor*)Cursor.Get();
+			const FVector2D MouseDelta = NewPosition - MacCursor->GetPosition();
+			if (MacCursor->UpdateCursorClipping(NewPosition))
+			{
+				MacCursor->SetPosition(NewPosition.X, NewPosition.Y);
+			}
+			else
+			{
+				MacCursor->UpdateCurrentPosition(NewPosition);
+			}
+
 			MessageHandler->OnDragOver(EventWindow.ToSharedRef());
 		}
 		else if (Event.NotificationName == NSPrepareForDragOperation)
