Linux: Dont try to make current if our current window was the SplashScreen window we just destroyed

diff --git a/Engine/Source/Runtime/ApplicationCore/Private/Linux/LinuxPlatformSplash.cpp b/Engine/Source/Runtime/ApplicationCore/Private/Linux/LinuxPlatformSplash.cpp
--- a/Engine/Source/Runtime/ApplicationCore/Private/Linux/LinuxPlatformSplash.cpp
+++ b/Engine/Source/Runtime/ApplicationCore/Private/Linux/LinuxPlatformSplash.cpp
@@ -133,7 +133,7 @@ FLinuxSplashState::~FLinuxSplashState()
 		FT_Done_FreeType(FontLibrary);
 	}
 
-	if (CurrentWindow)  // put back any old GL state...
+	if (CurrentWindow && CurrentWindow != SplashWindow)  // put back any old GL state...
 	{
 		SDL_GL_MakeCurrent(CurrentWindow, CurrentContext);
 	}
