Ignore tooltips and cursor decorator windows when determining the window under cursor on Mac

diff --git a/Engine/Source/Runtime/ApplicationCore/Private/Mac/MacApplication.cpp b/Engine/Source/Runtime/ApplicationCore/Private/Mac/MacApplication.cpp
--- a/Engine/Source/Runtime/ApplicationCore/Private/Mac/MacApplication.cpp
+++ b/Engine/Source/Runtime/ApplicationCore/Private/Mac/MacApplication.cpp
@@ -285,7 +285,7 @@ bool FMacApplication::IsCursorDirectlyOverSlateWindow() const
 
 TSharedPtr<FGenericWindow> FMacApplication::GetWindowUnderCursor()
 {
-	if (WindowUnderCursor)
+	if (WindowUnderCursor && WindowUnderCursor->bAcceptsInput)
 	{
 		return StaticCastSharedPtr<FGenericWindow>(FindWindowByNSWindow(WindowUnderCursor));
 	}
