Linux: Add a timer when we lose focus to check if we have lost application focus

diff --git a/Engine/Source/Runtime/ApplicationCore/Private/Linux/LinuxApplication.cpp b/Engine/Source/Runtime/ApplicationCore/Private/Linux/LinuxApplication.cpp
--- a/Engine/Source/Runtime/ApplicationCore/Private/Linux/LinuxApplication.cpp
+++ b/Engine/Source/Runtime/ApplicationCore/Private/Linux/LinuxApplication.cpp
@@ -18,6 +18,12 @@
 #define GAMECONTROLLER_RIGHT_THUMB_DEADZONE 8689
 #define GAMECONTROLLER_TRIGGER_THRESHOLD    30
 
+namespace
+{
+	// How long we wait from a FocusOut event to deactivate the application (100ms default)
+	double DeactivationThreadshold = 0.1;
+}
+
 float ShortToNormalFloat(short AxisVal)
 {
 	// normalize [-32768..32767] -> [-1..1]
@@ -72,6 +78,7 @@ FLinuxApplication::FLinuxApplication()
 	,	bInsideOwnWindow(false)
 	,	bIsDragWindowButtonPressed(false)
 	,	bActivateApp(false)
+	,	FocusOutDeactivationTime(0.0)
 	,	LastTimeCachedDisplays(-1.0)
 {
 	bUsingHighPrecisionMouseInput = false;
@@ -863,6 +870,8 @@ void FLinuxApplication::ProcessDeferredMessage( SDL_Event Event )
 							// Only do if the application is active.
 							if(bActivateApp)
 							{
+								FocusOutDeactivationTime = FPlatformTime::Seconds() + DeactivationThreadshold;
+
 								SDL_Event event;
 								event.type = SDL_USEREVENT;
 								event.user.code = CheckForDeactivation;
@@ -932,9 +941,17 @@ void FLinuxApplication::ProcessDeferredMessage( SDL_Event Event )
 		{
 			if(Event.user.code == CheckForDeactivation)
 			{
+				// We still havent hit our timeout limit, try again
+				if (FocusOutDeactivationTime > FPlatformTime::Seconds())
+				{
+					SDL_Event event;
+					event.type = SDL_USEREVENT;
+					event.user.code = CheckForDeactivation;
+					SDL_PushEvent(&event);
+				}
 				// If we don't use bIsDragWindowButtonPressed the draged window will be destroyed because we
 				// deactivate the whole appliacton. TODO Is that a bug? Do we have to do something?
-				if (!CurrentFocusWindow.IsValid() && !bIsDragWindowButtonPressed)
+				else if (!CurrentFocusWindow.IsValid() && !bIsDragWindowButtonPressed)
 				{
 					DeactivateApplication();
 				}
diff --git a/Engine/Source/Runtime/ApplicationCore/Public/Linux/LinuxApplication.h b/Engine/Source/Runtime/ApplicationCore/Public/Linux/LinuxApplication.h
--- a/Engine/Source/Runtime/ApplicationCore/Public/Linux/LinuxApplication.h
+++ b/Engine/Source/Runtime/ApplicationCore/Public/Linux/LinuxApplication.h
@@ -331,6 +331,9 @@ private:
 	/** Used to check if the application is active or not. */
 	bool bActivateApp;
 
+	/** Time before deactivating the application if no FocusIn event happens on any of our Windows */
+	double FocusOutDeactivationTime;
+
 	/** Cached displays - to reduce costly communication with X server (may be better cached in SDL? avoids ugly 'mutable') */
 	mutable TArray<SDL_Rect>	CachedDisplays;
 
