Fixed issues with cursor locking on window activation on Mac

diff --git a/Engine/Source/Runtime/ApplicationCore/Private/Mac/MacCursor.cpp b/Engine/Source/Runtime/ApplicationCore/Private/Mac/MacCursor.cpp
--- a/Engine/Source/Runtime/ApplicationCore/Private/Mac/MacCursor.cpp
+++ b/Engine/Source/Runtime/ApplicationCore/Private/Mac/MacCursor.cpp
@@ -293,10 +293,12 @@ void FMacCursor::Lock(const RECT* const Bounds)
 
 	MacApplication->OnCursorLock();
 
+	bIsPositionInitialised = false; // Force GetPosition() to update its cached position in case the cursor was warped by another app while we were in high precision mode
 	FVector2D Position = GetPosition();
 	if (UpdateCursorClipping(Position))
 	{
 		SetPosition(Position.X, Position.Y);
+		WarpCursor(Position.X, Position.Y);
 	}
 }
 
