Linux: Break grabs and bring down pointer barreirs when we have crashed

diff --git a/Engine/Source/Runtime/ApplicationCore/Private/Linux/LinuxPlatformApplicationMisc.cpp b/Engine/Source/Runtime/ApplicationCore/Private/Linux/LinuxPlatformApplicationMisc.cpp
--- a/Engine/Source/Runtime/ApplicationCore/Private/Linux/LinuxPlatformApplicationMisc.cpp
+++ b/Engine/Source/Runtime/ApplicationCore/Private/Linux/LinuxPlatformApplicationMisc.cpp
@@ -222,7 +222,6 @@ EAppReturnType::Type MessageBoxExtImpl(EAppMsgType::Type MsgType, const TCHAR* T
 	return Answer;
 }
 
-#if !UE_BUILD_SHIPPING
 void UngrabAllInputImpl()
 {
 	if (GInitializedSDL)
@@ -237,7 +236,6 @@ void UngrabAllInputImpl()
 		SDL_CaptureMouse(SDL_FALSE);
 	}
 }
-#endif // !UE_BUILD_SHIPPING
 
 uint32 FLinuxPlatformApplicationMisc::WindowStyle()
 {
@@ -260,9 +258,7 @@ void FLinuxPlatformApplicationMisc::Init()
 
 	FGenericPlatformApplicationMisc::Init();
 
-#if !UE_BUILD_SHIPPING
 	UngrabAllInputCallback = UngrabAllInputImpl;
-#endif
 }
 
 bool FLinuxPlatformApplicationMisc::InitSDL()
@@ -340,9 +336,7 @@ void FLinuxPlatformApplicationMisc::TearDown()
 		GInitializedSDL = false;
 
 		MessageBoxExtCallback = nullptr;
-#if !UE_BUILD_SHIPPING
 		UngrabAllInputCallback = nullptr;
-#endif
 	}
 }
 
diff --git a/Engine/Source/Runtime/Core/Private/Unix/UnixPlatformCrashContext.cpp b/Engine/Source/Runtime/Core/Private/Unix/UnixPlatformCrashContext.cpp
--- a/Engine/Source/Runtime/Core/Private/Unix/UnixPlatformCrashContext.cpp
+++ b/Engine/Source/Runtime/Core/Private/Unix/UnixPlatformCrashContext.cpp
@@ -646,6 +646,9 @@ void PlatformCrashHandler(int32 Signal, siginfo_t* Info, void* Context)
 	CrashContext.InitFromSignal(Signal, Info, Context);
 	CrashContext.FirstCrashHandlerFrame = static_cast<uint64*>(__builtin_return_address(0));
 
+	// This will ungrab cursor/keyboard and bring down any pointer barriers which will be stuck on when opening the CRC
+	FLinuxPlatformMisc::UngrabAllInput();
+
 	if (GCrashHandlerPointer)
 	{
 		GCrashHandlerPointer(CrashContext);
diff --git a/Engine/Source/Runtime/Core/Private/Unix/UnixPlatformMisc.cpp b/Engine/Source/Runtime/Core/Private/Unix/UnixPlatformMisc.cpp
--- a/Engine/Source/Runtime/Core/Private/Unix/UnixPlatformMisc.cpp
+++ b/Engine/Source/Runtime/Core/Private/Unix/UnixPlatformMisc.cpp
@@ -918,17 +918,23 @@ void FUnixPlatformMisc::CustomNamedStat(const ANSICHAR* Text, float Value, const
 }
 #endif
 
-#if !UE_BUILD_SHIPPING
 CORE_API TFunction<void()> UngrabAllInputCallback;
 
+void FUnixPlatformMisc::UngrabAllInput()
+{
+	if(UngrabAllInputCallback)
+	{
+		UngrabAllInputCallback();
+	}
+}
+
+#if !UE_BUILD_SHIPPING
 void FUnixPlatformMisc::DebugBreakInternal()
 {
 	if( IsDebuggerPresent() )
 	{
-		if(UngrabAllInputCallback)
-		{
-			UngrabAllInputCallback();
-		}
+		UngrabAllInput();
+
 #if PLATFORM_CPU_X86_FAMILY
 		__asm__ volatile("int $0x03");
 #else
diff --git a/Engine/Source/Runtime/Core/Public/Unix/UnixPlatformMisc.h b/Engine/Source/Runtime/Core/Public/Unix/UnixPlatformMisc.h
--- a/Engine/Source/Runtime/Core/Public/Unix/UnixPlatformMisc.h
+++ b/Engine/Source/Runtime/Core/Public/Unix/UnixPlatformMisc.h
@@ -150,12 +150,10 @@ struct CORE_API FUnixPlatformMisc : public FGenericPlatformMisc
 	static bool HasNonoptionalCPUFeatures();
 	static bool NeedsNonoptionalCPUFeaturesCheck();
 
-#if !UE_BUILD_SHIPPING	// only in non-shipping because we break into the debugger in non-shipping builds only
 	/**
 	 * Ungrabs input (useful before breaking into debugging)
 	 */
 	static void UngrabAllInput();
-#endif // !UE_BUILD_SHIPPING
 
 	/**
 	 * Returns whether the program has been started remotely (e.g. over SSH)
