Comitting shelf 4433043 for UE-64237.

diff --git a/Engine/Source/Runtime/ApplicationCore/Public/Android/AndroidWindowUtils.h b/Engine/Source/Runtime/ApplicationCore/Public/Android/AndroidWindowUtils.h
--- a/Engine/Source/Runtime/ApplicationCore/Public/Android/AndroidWindowUtils.h
+++ b/Engine/Source/Runtime/ApplicationCore/Public/Android/AndroidWindowUtils.h
@@ -5,6 +5,11 @@
 
 #if PLATFORM_ANDROID
 #include "Android/AndroidMisc.h"
+
+#if USE_ANDROID_JNI
+extern bool AndroidThunkCpp_IsGearVRApplication();
+#endif
+
 #endif
 
 namespace AndroidWindowUtils
@@ -12,10 +17,16 @@ namespace AndroidWindowUtils
 	inline bool DeviceRequiresMosaic()
 	{
 #if PLATFORM_ANDROID
-		bool bDeviceRequiresMosaic = !FAndroidMisc::SupportsFloatingPointRenderTargets() && !FAndroidMisc::SupportsShaderFramebufferFetch();
+		bool bDeviceRequiresMosaic = false;
+
+#if USE_ANDROID_JNI
+		bDeviceRequiresMosaic = !AndroidThunkCpp_IsGearVRApplication() && !FAndroidMisc::IsDaydreamApplication()
+			&& !FAndroidMisc::SupportsFloatingPointRenderTargets() && !FAndroidMisc::SupportsShaderFramebufferFetch();
+#endif //USE_ANDROID_JNI
+
 #else
 		bool bDeviceRequiresMosaic = !GSupportsRenderTargetFormat_PF_FloatRGBA && !GSupportsShaderFramebufferFetch;
-#endif
+#endif // PLATFORM_ANDROID
 
 		return bDeviceRequiresMosaic;
 	}
@@ -139,4 +150,4 @@ namespace AndroidWindowUtils
 		}
 	}
 
-}
\ No newline at end of file
+} // end AndroidWindowUtils
\ No newline at end of file
