UE-61698 - Android + Mouse: Hover events and Get Mouse Position on Viewport not working

diff --git a/Engine/Source/Runtime/ApplicationCore/Private/Android/AndroidApplication.cpp b/Engine/Source/Runtime/ApplicationCore/Private/Android/AndroidApplication.cpp
--- a/Engine/Source/Runtime/ApplicationCore/Private/Android/AndroidApplication.cpp
+++ b/Engine/Source/Runtime/ApplicationCore/Private/Android/AndroidApplication.cpp
@@ -81,6 +81,14 @@ void FAndroidApplication::PollGameDeviceState( const float TimeDelta )
 		FDisplayMetrics DisplayMetrics;
 		FDisplayMetrics::RebuildDisplayMetrics(DisplayMetrics);
 		BroadcastDisplayMetricsChanged(DisplayMetrics);
+
+		// the cursor needs to compute the proper slate scaling factor each time the display metrics change
+		TSharedPtr<FAndroidCursor> AndroidCursor = StaticCastSharedPtr<FAndroidCursor>(Cursor);
+		if (AndroidCursor.IsValid())
+		{
+			AndroidCursor->ComputeUIScaleFactor();
+		}
+
 		bWindowSizeChanged = false;
 	}
 }
diff --git a/Engine/Source/Runtime/ApplicationCore/Private/Android/AndroidCursor.cpp b/Engine/Source/Runtime/ApplicationCore/Private/Android/AndroidCursor.cpp
--- a/Engine/Source/Runtime/ApplicationCore/Private/Android/AndroidCursor.cpp
+++ b/Engine/Source/Runtime/ApplicationCore/Private/Android/AndroidCursor.cpp
@@ -8,14 +8,34 @@ FAndroidCursor::FAndroidCursor()
 	, CurrentPosition(FVector2D::ZeroVector)
 	, CursorClipRect()
 	, bShow(false)
-{	
+	, UIScaleFactor(1.0f)
+{
+	ComputeUIScaleFactor();
+}
+
+void FAndroidCursor::ComputeUIScaleFactor()
+{
+	FPlatformRect ScreenRect = FAndroidWindow::GetScreenRect();
+	int32_t ScreenWidth, ScreenHeight;
+	FAndroidWindow::CalculateSurfaceSize(FAndroidWindow::GetHardwareWindow(), ScreenWidth, ScreenHeight);
+
+	if (ScreenWidth > ScreenHeight)
+	{
+		UIScaleFactor = (ScreenRect.Bottom - ScreenRect.Top) / (float)ScreenHeight;
+	}
+	else
+	{
+		UIScaleFactor = (ScreenRect.Right - ScreenRect.Left) / (float)ScreenWidth;
+	}
 }
 
 void FAndroidCursor::SetPosition( const int32 X, const int32 Y )
 {
 	FVector2D NewPosition(X, Y);
+
+	NewPosition *= UIScaleFactor;
+
 	UpdateCursorClipping(NewPosition);
-	
 	CurrentPosition = NewPosition;
 }
 
diff --git a/Engine/Source/Runtime/ApplicationCore/Private/Android/AndroidInputInterface.cpp b/Engine/Source/Runtime/ApplicationCore/Private/Android/AndroidInputInterface.cpp
--- a/Engine/Source/Runtime/ApplicationCore/Private/Android/AndroidInputInterface.cpp
+++ b/Engine/Source/Runtime/ApplicationCore/Private/Android/AndroidInputInterface.cpp
@@ -1017,7 +1017,7 @@ void FAndroidInputInterface::SendControllerEvents()
 			case MouseEventType::MouseMove:
 				if (Cursor.IsValid())
 				{
-					Cursor->SetPosition(mouse_data.AbsoluteX, mouse_data.AbsoluteX);
+					Cursor->SetPosition(mouse_data.AbsoluteX, mouse_data.AbsoluteY);
 					MessageHandler->OnMouseMove();
 				}
 				MessageHandler->OnRawMouseMove(mouse_data.DeltaX, mouse_data.DeltaY);
diff --git a/Engine/Source/Runtime/ApplicationCore/Public/Android/AndroidCursor.h b/Engine/Source/Runtime/ApplicationCore/Public/Android/AndroidCursor.h
--- a/Engine/Source/Runtime/ApplicationCore/Public/Android/AndroidCursor.h
+++ b/Engine/Source/Runtime/ApplicationCore/Public/Android/AndroidCursor.h
@@ -37,6 +37,9 @@ public:
 
 	virtual void SetTypeShape(EMouseCursor::Type InCursorType, void* CursorHandle) override { }
 
+	/** this function will compute the necessary position scale factor needed to transform the cursor in Slate space */
+	void ComputeUIScaleFactor();
+
 private:
 	bool UpdateCursorClipping(FVector2D& CursorPosition);
 
@@ -44,4 +47,7 @@ private:
 	FVector2D CurrentPosition;
 	FIntRect CursorClipRect;
 	bool bShow;
+
+	/** scale factor computed by ComputeUIScaleFactor() and used in SetPosition() */
+	float UIScaleFactor;
 };
