Clone of Dev-Rendering CL 4444343:

Fix for eyedropper in the Color Picker window not grabbing any color values on Mac.  Bug also says Linux - unable to verify if its the same issue so kept change to Mac specific code.

diff --git a/Engine/Source/Runtime/ApplicationCore/Private/Mac/MacPlatformApplicationMisc.cpp b/Engine/Source/Runtime/ApplicationCore/Private/Mac/MacPlatformApplicationMisc.cpp
--- a/Engine/Source/Runtime/ApplicationCore/Private/Mac/MacPlatformApplicationMisc.cpp
+++ b/Engine/Source/Runtime/ApplicationCore/Private/Mac/MacPlatformApplicationMisc.cpp
@@ -418,7 +418,8 @@ FLinearColor FMacPlatformApplicationMisc::GetScreenPixelColor(const FVector2D& I
 {
 	SCOPED_AUTORELEASE_POOL;
 
-	CGImageRef ScreenImage = CGWindowListCreateImage(CGRectMake(InScreenPos.X, InScreenPos.Y, 1, 1), kCGWindowListOptionOnScreenBelowWindow, kCGNullWindowID, kCGWindowImageDefault);
+	const CGPoint Pt = FMacApplication::ConvertSlatePositionToCGPoint(InScreenPos.X, InScreenPos.Y);
+	CGImageRef ScreenImage = CGWindowListCreateImage(CGRectMake(Pt.x, Pt.y, 1, 1), kCGWindowListOptionOnScreenBelowWindow, kCGNullWindowID, kCGWindowImageDefault);
 	
 	CGDataProviderRef provider = CGImageGetDataProvider(ScreenImage);
 	NSData* data = (id)CGDataProviderCopyData(provider);
