Fixed force feedback on iPhone 6 and 6s.

diff --git a/Engine/Source/Runtime/ApplicationCore/Private/IOS/IOSInputInterface.cpp b/Engine/Source/Runtime/ApplicationCore/Private/IOS/IOSInputInterface.cpp
--- a/Engine/Source/Runtime/ApplicationCore/Private/IOS/IOSInputInterface.cpp
+++ b/Engine/Source/Runtime/ApplicationCore/Private/IOS/IOSInputInterface.cpp
@@ -67,6 +67,8 @@ FIOSInputInterface::FIOSInputInterface( const TSharedRef< FGenericApplicationMes
 	[GCController startWirelessControllerDiscoveryWithCompletionHandler:^{ }];
 
 	FMemory::Memzero(Controllers, sizeof(Controllers));
+	
+	HapticFeedbackSupportLevel = [[[UIDevice currentDevice] valueForKey:@"_feedbackSupportLevel"] intValue];
 }
 
 void FIOSInputInterface::SetMessageHandler( const TSharedRef< FGenericApplicationMessageHandler >& InMessageHandler )
@@ -633,56 +635,62 @@ bool FIOSInputInterface::IsControllerAssignedToGamepad(int32 ControllerId)
 
 void FIOSInputInterface::SetForceFeedbackChannelValue(int32 ControllerId, FForceFeedbackChannelType ChannelType, float Value)
 {
-	// if we are at rest, then kick when we are over the Kick cutoff
-	if (LastHapticValue == 0.0f && Value > 0.0f)
+	if(HapticFeedbackSupportLevel >= 2)
 	{
-		const float HeavyKickVal = CVarHapticsKickHeavy.GetValueOnGameThread();
-		const float MediumKickVal = CVarHapticsKickMedium.GetValueOnGameThread();
-		const float LightKickVal = CVarHapticsKickLight.GetValueOnGameThread();
-		// once we get past the
-		if (Value > LightKickVal)
+		// if we are at rest, then kick when we are over the Kick cutoff
+		if (LastHapticValue == 0.0f && Value > 0.0f)
 		{
-			if (Value > HeavyKickVal)
+			const float HeavyKickVal = CVarHapticsKickHeavy.GetValueOnGameThread();
+			const float MediumKickVal = CVarHapticsKickMedium.GetValueOnGameThread();
+			const float LightKickVal = CVarHapticsKickLight.GetValueOnGameThread();
+			// once we get past the
+			if (Value > LightKickVal)
 			{
-				FPlatformMisc::PrepareMobileHaptics(EMobileHapticsType::ImpactHeavy);
+				if (Value > HeavyKickVal)
+				{
+					FPlatformMisc::PrepareMobileHaptics(EMobileHapticsType::ImpactHeavy);
+				}
+				else if (Value > MediumKickVal)
+				{
+					FPlatformMisc::PrepareMobileHaptics(EMobileHapticsType::ImpactMedium);
+				}
+				else
+				{
+					FPlatformMisc::PrepareMobileHaptics(EMobileHapticsType::ImpactLight);
+				}
+
+				FPlatformMisc::TriggerMobileHaptics();
+				
+				// remember it to not kick again
+				LastHapticValue = Value;
 			}
-			else if (Value > MediumKickVal)
+		}
+		else
+		{
+			const float RestVal = CVarHapticsRest.GetValueOnGameThread();
+
+			if (Value >= RestVal)
 			{
-				FPlatformMisc::PrepareMobileHaptics(EMobileHapticsType::ImpactMedium);
+				// always remember the last value if we are over the Rest amount
+				LastHapticValue = Value;
 			}
 			else
 			{
-				FPlatformMisc::PrepareMobileHaptics(EMobileHapticsType::ImpactLight);
+				// release the haptics
+				FPlatformMisc::ReleaseMobileHaptics();
+				
+				// rest
+				LastHapticValue = 0.0f;
 			}
-
-			FPlatformMisc::TriggerMobileHaptics();
-			
-			// remember it to not kick again
-			LastHapticValue = Value;
 		}
 	}
 	else
 	{
-		const float RestVal = CVarHapticsRest.GetValueOnGameThread();
-
-		if (Value >= RestVal)
+		if(Value >= 0.3f)
 		{
-			// always remember the last value if we are over the Rest amount
-			LastHapticValue = Value;
-		}
-		else
-		{
-			// release the haptics
-			FPlatformMisc::ReleaseMobileHaptics();
-			
-			// rest
-			LastHapticValue = 0.0f;
+			AudioServicesPlaySystemSound(kSystemSoundID_Vibrate);
 		}
 	}
-	
-//	{
-//		AudioServicesPlaySystemSound(kSystemSoundID_Vibrate);
-//	}
 }
 
 void FIOSInputInterface::SetForceFeedbackChannelValues(int32 ControllerId, const FForceFeedbackValues &Values)
diff --git a/Engine/Source/Runtime/ApplicationCore/Public/IOS/IOSInputInterface.h b/Engine/Source/Runtime/ApplicationCore/Public/IOS/IOSInputInterface.h
--- a/Engine/Source/Runtime/ApplicationCore/Public/IOS/IOSInputInterface.h
+++ b/Engine/Source/Runtime/ApplicationCore/Public/IOS/IOSInputInterface.h
@@ -177,4 +177,6 @@ private:
 
 	/** Last value sent to mobile haptics */
 	float LastHapticValue;
+	
+	int HapticFeedbackSupportLevel;
 };
