Copy of CL 4496633

Fixed an issue with mouse dragged events on Mac not being generated when cursor was locked while outside of the window.

diff --git a/Engine/Source/Runtime/ApplicationCore/Private/Mac/MacApplication.cpp b/Engine/Source/Runtime/ApplicationCore/Private/Mac/MacApplication.cpp
--- a/Engine/Source/Runtime/ApplicationCore/Private/Mac/MacApplication.cpp
+++ b/Engine/Source/Runtime/ApplicationCore/Private/Mac/MacApplication.cpp
@@ -535,7 +535,26 @@ NSEvent* FMacApplication::HandleNSEvent(NSEvent* Event)
 		{
 			MacApplication->DeferEvent(Event);
 
-			if ([Event type] == NSEventTypeKeyDown || [Event type] == NSEventTypeKeyUp)
+			bool bShouldStopEventDispatch = false;
+			switch ([Event type])
+			{
+				case NSEventTypeKeyDown:
+				case NSEventTypeKeyUp:
+					bShouldStopEventDispatch = true;
+					break;
+					
+				case NSEventTypeLeftMouseDown:
+				case NSEventTypeRightMouseDown:
+				case NSEventTypeOtherMouseDown:
+				case NSEventTypeLeftMouseUp:
+				case NSEventTypeRightMouseUp:
+				case NSEventTypeOtherMouseUp:
+					// Make sure we don't stop receiving mouse dragged events if the cursor is over the edge of the screen and the cursor is locked
+					bShouldStopEventDispatch = ((FMacCursor*)MacApplication->Cursor.Get())->IsLocked();
+					break;
+			}
+			
+			if (bShouldStopEventDispatch)
 			{
 				ReturnEvent = nil;
 			}
diff --git a/Engine/Source/Runtime/ApplicationCore/Private/Mac/MacCursor.cpp b/Engine/Source/Runtime/ApplicationCore/Private/Mac/MacCursor.cpp
--- a/Engine/Source/Runtime/ApplicationCore/Private/Mac/MacCursor.cpp
+++ b/Engine/Source/Runtime/ApplicationCore/Private/Mac/MacCursor.cpp
@@ -231,7 +231,7 @@ void FMacCursor::SetPosition(const int32 X, const int32 Y)
 
 	if (!bIsPositionInitialised || FIntVector(NewPos.X, NewPos.Y, 0) != FIntVector(CurrentPosition.X, CurrentPosition.Y, 0))
 	{
-		if (!bUseHighPrecisionMode || (CurrentCursor && bIsVisible))
+		if (!bUseHighPrecisionMode || (CurrentCursor && bIsVisible) || !bIsPositionInitialised)
 		{
 			WarpCursor(NewPos.X, NewPos.Y);
 		}
