Copy of CL 4494336

Added support for Nintendo Switch Pro controller and newer versions of PS4 and Xbox One contollers for Mac

diff --git a/Engine/Source/Runtime/ApplicationCore/Private/Mac/HIDInputInterface.cpp b/Engine/Source/Runtime/ApplicationCore/Private/Mac/HIDInputInterface.cpp
--- a/Engine/Source/Runtime/ApplicationCore/Private/Mac/HIDInputInterface.cpp
+++ b/Engine/Source/Runtime/ApplicationCore/Private/Mac/HIDInputInterface.cpp
@@ -128,6 +128,8 @@ void HIDInputInterface::FHIDDeviceInfo::SetupMappings()
 	const int32 VendorID = GetDevicePropertyAsInt32(DeviceRef, CFSTR(kIOHIDVendorIDKey));
 	const int32 ProductID = GetDevicePropertyAsInt32(DeviceRef, CFSTR(kIOHIDProductIDKey));
 
+	UE_LOG(LogMac, Log, TEXT("Controller attached: VendorID 0x%x, ProductID 0x%x"), VendorID, ProductID);
+
 	if (VendorID == 0x54c && ProductID == 0x268)
 	{
 		// PlayStation 3 Controller
@@ -155,7 +157,7 @@ void HIDInputInterface::FHIDDeviceInfo::SetupMappings()
 		RightAnalogYMapping = kHIDUsage_GD_Rz;
 		RightTriggerAnalogMapping = kHIDUsage_GD_Ry;
 	}
-	else if (VendorID == 0x54c && ProductID == 0x5c4)
+	else if (VendorID == 0x54c && (ProductID == 0x5c4 || ProductID == 0x9cc))
 	{
 		// PlayStation 4 Controller
 		ButtonsMapping[1]	= 2;	// Square		->	X
@@ -178,7 +180,7 @@ void HIDInputInterface::FHIDDeviceInfo::SetupMappings()
 		RightAnalogYMapping = kHIDUsage_GD_Rz;
 		RightTriggerAnalogMapping = kHIDUsage_GD_Ry;
 	}
-	else if (VendorID == 0x45e && ProductID == 0x2d1)
+	else if (VendorID == 0x45e && (ProductID == 0x2d1 || ProductID == 0x2ea))
 	{
 		// Xbox One Controller
 		ButtonsMapping[1]	= 0;	// A
@@ -203,6 +205,27 @@ void HIDInputInterface::FHIDDeviceInfo::SetupMappings()
 		RightAnalogYMapping = kHIDUsage_GD_Ry;
 		RightTriggerAnalogMapping = kHIDUsage_GD_Rz;
 	}
+	else if (VendorID == 0x57e && ProductID == 0x2009)
+	{
+		// Nintendo Switch Pro Controller
+		ButtonsMapping[1]	= 0;	// B
+		ButtonsMapping[2]	= 1;	// A
+		ButtonsMapping[3]	= 2;	// Y
+		ButtonsMapping[4]	= 3;	// X
+		ButtonsMapping[5]	= 4;	// L
+		ButtonsMapping[6]	= 5;	// R
+		ButtonsMapping[7]	= 10;	// ZL
+		ButtonsMapping[8]	= 11;	// ZR
+		ButtonsMapping[9]	= 7;	// -
+		ButtonsMapping[10]	= 6;	// +
+		ButtonsMapping[11]	= 8;	// Left Thumbstick
+		ButtonsMapping[12]	= 9;	// Right Thumbstick
+
+		LeftAnalogXMapping = kHIDUsage_GD_X;
+		LeftAnalogYMapping = kHIDUsage_GD_Y;
+		RightAnalogXMapping = kHIDUsage_GD_Rx;
+		RightAnalogYMapping = kHIDUsage_GD_Ry;
+	}
 	else
 	{
 		// Array generated from [REDACTED_PROFILE_URL]/360Controller/blob/master/360Controller/Info.plist (with additions)
