Linux: Workaround for corrupted windows on Vulkan

diff --git a/Engine/Source/Runtime/ApplicationCore/Private/Linux/LinuxWindow.cpp b/Engine/Source/Runtime/ApplicationCore/Private/Linux/LinuxWindow.cpp
--- a/Engine/Source/Runtime/ApplicationCore/Private/Linux/LinuxWindow.cpp
+++ b/Engine/Source/Runtime/ApplicationCore/Private/Linux/LinuxWindow.cpp
@@ -503,6 +503,17 @@ static void _GetBestFullscreenResolution( SDL_HWindow hWnd, int32 *pWidth, int32
 
 void FLinuxWindow::ReshapeWindow( int32 NewX, int32 NewY, int32 NewWidth, int32 NewHeight )
 {
+	// Some vulkan video drivers have issues with specific height ranges causing them to corrupt the texture rendered
+	// Moving these nearest values removes this corruption.
+	if (NewHeight >= 33 && NewHeight <= 43)
+	{
+		NewHeight = 44;
+	}
+	else if (NewHeight >= 65 && NewHeight <= 85)
+	{
+		NewHeight = 86;
+	}
+
 	// X11 will take until the next frame to send a SizeChanged event. This means the X11 window
 	// will most likely have resized already by the time we render but the slate renderer will
 	// not have been updated leading to an incorrect frame.
