Fix for TVOS crash on launch caused by missing _feedbackSupportLevel key

diff --git a/Engine/Source/Runtime/ApplicationCore/Private/IOS/IOSInputInterface.cpp b/Engine/Source/Runtime/ApplicationCore/Private/IOS/IOSInputInterface.cpp
--- a/Engine/Source/Runtime/ApplicationCore/Private/IOS/IOSInputInterface.cpp
+++ b/Engine/Source/Runtime/ApplicationCore/Private/IOS/IOSInputInterface.cpp
@@ -68,7 +68,11 @@ FIOSInputInterface::FIOSInputInterface( const TSharedRef< FGenericApplicationMes
 
 	FMemory::Memzero(Controllers, sizeof(Controllers));
 	
+#if !PLATFORM_TVOS
 	HapticFeedbackSupportLevel = [[[UIDevice currentDevice] valueForKey:@"_feedbackSupportLevel"] intValue];
+#else
+	HapticFeedbackSupportLevel = 0;
+#endif
 }
 
 void FIOSInputInterface::SetMessageHandler( const TSharedRef< FGenericApplicationMessageHandler >& InMessageHandler )
