Linux: Fix blackscreen on window when leaving fullscreen

Backs are part of CL 4501575

diff --git a/Engine/Source/Runtime/ApplicationCore/Private/Linux/LinuxWindow.cpp b/Engine/Source/Runtime/ApplicationCore/Private/Linux/LinuxWindow.cpp
--- a/Engine/Source/Runtime/ApplicationCore/Private/Linux/LinuxWindow.cpp
+++ b/Engine/Source/Runtime/ApplicationCore/Private/Linux/LinuxWindow.cpp
@@ -537,7 +537,6 @@ void FLinuxWindow::ReshapeWindow( int32 NewX, int32 NewY, int32 NewWidth, int32
 		case EWindowMode::WindowedFullscreen:
 		{
 			SDL_SetWindowFullscreen( HWnd, 0 );
-			SDL_SetWindowSize( HWnd, NewWidth, NewHeight );
 			SDL_SetWindowFullscreen( HWnd, SDL_WINDOW_FULLSCREEN_DESKTOP );
 			bWasFullscreen = true;
 		}
@@ -594,9 +593,25 @@ void FLinuxWindow::SetWindowMode( EWindowMode::Type NewWindowMode )
 			{
 				if ( bWasFullscreen != true )
 				{
-					SDL_SetWindowSize( HWnd, VirtualWidth, VirtualHeight );
+					TSharedPtr< FLinuxWindow > LinuxWindow = OwningApplication->FindWindowBySDLWindow(HWnd);
+					if ( LinuxWindow )
+					{
+						OwningApplication->GetMessageHandler()->OnResizingWindow( LinuxWindow.ToSharedRef() );
+					}
+
 					SDL_SetWindowFullscreen( HWnd, SDL_WINDOW_FULLSCREEN_DESKTOP );
 					bWasFullscreen = true;
+
+					if ( LinuxWindow )
+					{
+						OwningApplication->GetMessageHandler()->OnSizeChanged(
+							LinuxWindow.ToSharedRef(),
+							VirtualWidth,
+							VirtualHeight,
+							//  bWasMinimized
+							false
+						);
+					}
 				}
 			}
 			break;
@@ -634,11 +649,16 @@ void FLinuxWindow::AdjustCachedSize( FVector2D& Size ) const
 	{
 		Size = FVector2D( VirtualWidth, VirtualHeight );
 	}
-	else
-	if	( HWnd )
+	else if	( HWnd )
 	{
 		int SizeW, SizeH;
 
+		SDL_GetWindowSize( HWnd, &SizeW, &SizeH );
+
+		/*
+		 * Currently we are not correctly supporting up-scaling on all RHIs. For now disable this
+		 * until all RHIs are working with up-scaling
+		 *
 		if ( WindowMode == EWindowMode::Windowed )
 		{
 			SDL_GetWindowSize( HWnd, &SizeW, &SizeH );
@@ -650,6 +670,7 @@ void FLinuxWindow::AdjustCachedSize( FVector2D& Size ) const
 
 			_GetBestFullscreenResolution( HWnd, &SizeW, &SizeH );
 		}
+		*/
 
 		Size = FVector2D( SizeW, SizeH );
 	}
diff --git a/Engine/Source/Runtime/VulkanRHI/Private/Linux/VulkanLinuxPlatform.cpp b/Engine/Source/Runtime/VulkanRHI/Private/Linux/VulkanLinuxPlatform.cpp
--- a/Engine/Source/Runtime/VulkanRHI/Private/Linux/VulkanLinuxPlatform.cpp
+++ b/Engine/Source/Runtime/VulkanRHI/Private/Linux/VulkanLinuxPlatform.cpp
@@ -187,8 +187,3 @@ void FVulkanLinuxPlatform::WriteCrashMarker(const FOptionalVulkanDeviceExtension
 		}
 	}
 }
-
-void FVulkanLinuxPlatform::UpdateWindowSize(void* WindowHandle, uint32& Width, uint32& Height)
-{
-	SDL_Vulkan_GetDrawableSize((SDL_Window*) WindowHandle, (int32*) &Width, (int32*) &Height);
-}
diff --git a/Engine/Source/Runtime/VulkanRHI/Private/Linux/VulkanLinuxPlatform.h b/Engine/Source/Runtime/VulkanRHI/Private/Linux/VulkanLinuxPlatform.h
--- a/Engine/Source/Runtime/VulkanRHI/Private/Linux/VulkanLinuxPlatform.h
+++ b/Engine/Source/Runtime/VulkanRHI/Private/Linux/VulkanLinuxPlatform.h
@@ -57,8 +57,6 @@ public:
 #endif
 	}
 
-	static void UpdateWindowSize(void* WindowHandle, uint32& Width, uint32& Height);
-
 	static void WriteCrashMarker(const FOptionalVulkanDeviceExtensions& OptionalExtensions, VkCommandBuffer CmdBuffer, VkBuffer DestBuffer, const TArrayView<uint32>& Entries, bool bAdding);
 
 protected:
diff --git a/Engine/Source/Runtime/VulkanRHI/Private/VulkanGenericPlatform.h b/Engine/Source/Runtime/VulkanRHI/Private/VulkanGenericPlatform.h
--- a/Engine/Source/Runtime/VulkanRHI/Private/VulkanGenericPlatform.h
+++ b/Engine/Source/Runtime/VulkanRHI/Private/VulkanGenericPlatform.h
@@ -102,7 +102,4 @@ public:
 
 	// Ensure the last frame completed on the GPU
 	static bool RequiresWaitingForFrameCompletionEvent() { return true; }
-
-	// Some platforms imitate fullscreen mode by using bigger window
-	static void UpdateWindowSize(void* WindowHandle, uint32& Width, uint32& Height) {}
 };
diff --git a/Engine/Source/Runtime/VulkanRHI/Private/VulkanSwapChain.cpp b/Engine/Source/Runtime/VulkanRHI/Private/VulkanSwapChain.cpp
--- a/Engine/Source/Runtime/VulkanRHI/Private/VulkanSwapChain.cpp
+++ b/Engine/Source/Runtime/VulkanRHI/Private/VulkanSwapChain.cpp
@@ -386,9 +386,6 @@ FVulkanSwapChain::FVulkanSwapChain(VkInstance InInstance, FVulkanDevice& InDevic
 	uint32 SizeY = FVulkanPlatform::SupportsQuerySurfaceProperties() ? (SurfProperties.currentExtent.height == 0xFFFFFFFF ? Height : SurfProperties.currentExtent.height) : Height;
 	//FPlatformMisc::LowLevelOutputDebugStringf(TEXT("Create swapchain: %ux%u \n"), SizeX, SizeY);
 
-	// Some platforms cheat with the window size. Allow them to update.
-	FVulkanPlatform::UpdateWindowSize(WindowHandle, SizeX, SizeY);
-
 	VkSwapchainCreateInfoKHR SwapChainInfo;
 	ZeroVulkanStruct(SwapChainInfo, VK_STRUCTURE_TYPE_SWAPCHAIN_CREATE_INFO_KHR);
 	SwapChainInfo.surface = Surface;
diff --git a/Engine/Source/Runtime/VulkanRHI/Private/VulkanViewport.cpp b/Engine/Source/Runtime/VulkanRHI/Private/VulkanViewport.cpp
--- a/Engine/Source/Runtime/VulkanRHI/Private/VulkanViewport.cpp
+++ b/Engine/Source/Runtime/VulkanRHI/Private/VulkanViewport.cpp
@@ -689,7 +689,6 @@ bool FVulkanViewport::Present(FVulkanCommandListContext* Context, FVulkanCmdBuff
 
 			uint32 WindowSizeX = SizeX;
 			uint32 WindowSizeY = SizeY;
-			FVulkanPlatform::UpdateWindowSize(WindowHandle, WindowSizeX, WindowSizeY);
 
 			Context->RHIPushEvent(TEXT("CopyImageToBackBuffer"), FColor::Blue);
 			CopyImageToBackBuffer(CmdBuffer, true, RenderingBackBuffer->Surface.Image, BackBufferImages[AcquiredImageIndex], SizeX, SizeY, WindowSizeX, WindowSizeY);
