Linux: Additional height range found for corrupt windows

diff --git a/Engine/Source/Runtime/ApplicationCore/Private/Linux/LinuxWindow.cpp b/Engine/Source/Runtime/ApplicationCore/Private/Linux/LinuxWindow.cpp
--- a/Engine/Source/Runtime/ApplicationCore/Private/Linux/LinuxWindow.cpp
+++ b/Engine/Source/Runtime/ApplicationCore/Private/Linux/LinuxWindow.cpp
@@ -505,7 +505,15 @@ void FLinuxWindow::ReshapeWindow( int32 NewX, int32 NewY, int32 NewWidth, int32
 {
 	// Some vulkan video drivers have issues with specific height ranges causing them to corrupt the texture rendered
 	// Moving these nearest values removes this corruption.
-	if (NewHeight >= 33 && NewHeight <= 43)
+	if (NewHeight >= 9 && NewHeight <= 10)
+	{
+		NewHeight = 11;
+	}
+	else if (NewHeight >= 17 && NewHeight <= 21)
+	{
+		NewHeight = 22;
+	}
+	else if (NewHeight >= 33 && NewHeight <= 43)
 	{
 		NewHeight = 44;
 	}
