On Windows, RebuildDisplayMetrics atempts to retrieve DPI values from the EDID data structure

diff --git a/Engine/Source/Runtime/ApplicationCore/Private/Windows/WindowsApplication.cpp b/Engine/Source/Runtime/ApplicationCore/Private/Windows/WindowsApplication.cpp
--- a/Engine/Source/Runtime/ApplicationCore/Private/Windows/WindowsApplication.cpp
+++ b/Engine/Source/Runtime/ApplicationCore/Private/Windows/WindowsApplication.cpp
@@ -466,7 +466,7 @@ FPlatformRect FWindowsApplication::GetWorkArea( const FPlatformRect& CurrentWind
  * @param OutHeight - Reference to output variable for monitor native height
  * @returns 'true' if data was extracted successfully, 'false' otherwise
  **/
-static bool GetMonitorSizeFromEDID(const HKEY hDevRegKey, int32& OutWidth, int32& OutHeight)
+static bool GetMonitorSizeFromEDID(const HKEY hDevRegKey, int32& OutWidth, int32& OutHeight, int32& OutDPI)
 {	
 	static const uint32 NameSize = 512;
 	static TCHAR ValueName[NameSize];
@@ -498,6 +498,25 @@ static bool GetMonitorSizeFromEDID(const HKEY hDevRegKey, int32& OutWidth, int32
 		OutWidth = ((EDIDData[DetailTimingDescriptorStartIndex+4] >> 4) << 8) | EDIDData[DetailTimingDescriptorStartIndex+2];
 		OutHeight = ((EDIDData[DetailTimingDescriptorStartIndex+7] >> 4) << 8) | EDIDData[DetailTimingDescriptorStartIndex+5];
 
+		const int32 HorizontalSizeOffset = 21;
+		const int32 VerticalSizeOffset = 22;
+		const float CmToInch = 0.393701f;
+
+		if (EDIDData[HorizontalSizeOffset] > 0 && EDIDData[VerticalSizeOffset] > 0)
+		{
+			float PhysicalWidth = CmToInch * (float)EDIDData[HorizontalSizeOffset];
+			float PhysicalHeight = CmToInch * (float)EDIDData[VerticalSizeOffset];
+
+			int32 HDpi = FMath::TruncToInt((float)OutWidth / PhysicalWidth);
+			int32 VDpi = FMath::TruncToInt((float)OutHeight / PhysicalHeight);
+
+			OutDPI = (HDpi + VDpi) / 2;
+		}
+		else
+		{
+			OutDPI = 0;
+		}
+
 		return true; // valid EDID found
 	}
 
@@ -511,7 +530,7 @@ static bool GetMonitorSizeFromEDID(const HKEY hDevRegKey, int32& OutWidth, int32
  * @praam OutHeight - Reference to output variable for monitor native height
  * @returns TRUE if data was extracted successfully, FALSE otherwise
  **/
-inline bool GetSizeForDevID(const FString& TargetDevID, int32& Width, int32& Height)
+inline bool GetSizeForDevID(const FString& TargetDevID, int32& Width, int32& Height, int32& DPI)
 {
 	static const GUID ClassMonitorGuid = {0x4d36e96e, 0xe325, 0x11ce, {0xbf, 0xc1, 0x08, 0x00, 0x2b, 0xe1, 0x03, 0x18}};
 
@@ -550,7 +569,7 @@ inline bool GetSizeForDevID(const FString& TargetDevID, int32& Width, int32& Hei
 
 					if (hDevRegKey && hDevRegKey != INVALID_HANDLE_VALUE)
 					{
-						bRes = GetMonitorSizeFromEDID(hDevRegKey, Width, Height);
+						bRes = GetMonitorSizeFromEDID(hDevRegKey, Width, Height, DPI);
 						RegCloseKey(hDevRegKey);
 						break;
 					}
@@ -628,12 +647,25 @@ static void GetMonitorsInfo(TArray<FMonitorInfo>& OutMonitorInfo)
 					Info.ID = FString::Printf(TEXT("%s"), Monitor.DeviceID);
 					Info.Name = Info.ID.Mid (8, Info.ID.Find (TEXT("\\"), ESearchCase::CaseSensitive, ESearchDir::FromStart, 9) - 8);
 
-					if (GetSizeForDevID(Info.Name, Info.NativeWidth, Info.NativeHeight))
+					if (GetSizeForDevID(Info.Name, Info.NativeWidth, Info.NativeHeight, Info.DPI))
 					{
 						Info.ID = Monitor.DeviceID;
 						Info.bIsPrimary = (DisplayDevice.StateFlags & DISPLAY_DEVICE_PRIMARY_DEVICE) > 0;
 
-						Info.DPI = FWindowsPlatformApplicationMisc::GetMonitorDPI(Info);
+						// sanity check for DPI values
+						if (Info.DPI < 96 || Info.DPI > 300)
+						{
+							// switch to default winapi value
+							Info.DPI = FWindowsPlatformApplicationMisc::GetMonitorDPI(Info);
+						}
+						else
+						{
+							// we also need to include the OS scaling value
+							const float CenterX = 0.5f * (Info.WorkArea.Right + Info.WorkArea.Left);
+							const float CenterY = 0.5f * (Info.WorkArea.Top + Info.WorkArea.Bottom);
+							const float DPIScaleFactor = FWindowsPlatformApplicationMisc::GetDPIScaleFactorAtPoint(CenterX, CenterY);
+							Info.DPI *= DPIScaleFactor;
+						}
 
 						OutMonitorInfo.Add(Info);
 
