UE-65327 Project with Preferred Landscape Orientation set to Left Home Button still launches oriented Right Home

#iOS
#4.21

Default values for bSupportsXXOrientation in IOSRuntimeSettingsconstructor and BaseEngine.ini
preferredInterfaceOrientationForPresentation will return right or left only if both landscape orientation are supported

diff --git a/Engine/Config/BaseEngine.ini b/Engine/Config/BaseEngine.ini
--- a/Engine/Config/BaseEngine.ini
+++ b/Engine/Config/BaseEngine.ini
@@ -1769,6 +1769,7 @@ bSupportsPortraitOrientation=False
 bSupportsUpsideDownOrientation=False
 bSupportsLandscapeLeftOrientation=True
 bSupportsLandscapeRightOrientation=True
+PreferredLandscapeOrientation=LandscapeLeft
 bSupportsMetal=True
 bSupportsOpenGLES2=False
 bCookPVRTCTextures=True
diff --git a/Engine/Source/Runtime/ApplicationCore/Private/IOS/IOSView.cpp b/Engine/Source/Runtime/ApplicationCore/Private/IOS/IOSView.cpp
--- a/Engine/Source/Runtime/ApplicationCore/Private/IOS/IOSView.cpp
+++ b/Engine/Source/Runtime/ApplicationCore/Private/IOS/IOSView.cpp
@@ -1105,20 +1105,23 @@ id<MTLDevice> GMetalDevice = nil;
 - (UIInterfaceOrientation)preferredInterfaceOrientationForPresentation
 {
 	FString PreferredLandscapeOrientation = "";
+	bool bSupportsLandscapeLeft = false;
+	bool bSupportsLandscapeRight = false;
 
+	GConfig->GetBool(TEXT("/Script/IOSRuntimeSettings.IOSRuntimeSettings"), TEXT("bSupportsLandscapeLeftOrientation"), bSupportsLandscapeLeft, GEngineIni);
+	GConfig->GetBool(TEXT("/Script/IOSRuntimeSettings.IOSRuntimeSettings"), TEXT("bSupportsLandscapeRightOrientation"), bSupportsLandscapeRight, GEngineIni);
 	GConfig->GetString(TEXT("/Script/IOSRuntimeSettings.IOSRuntimeSettings"), TEXT("PreferredLandscapeOrientation"), PreferredLandscapeOrientation, GEngineIni);
-	if (PreferredLandscapeOrientation.Equals("LandscapeLeft"))
+
+	if(bSupportsLandscapeLeft && bSupportsLandscapeRight)
 	{
+		if (PreferredLandscapeOrientation.Equals("LandscapeRight"))
+		{
+			return UIInterfaceOrientationLandscapeRight;
+		}
 		return UIInterfaceOrientationLandscapeLeft;
 	}
-	else if (PreferredLandscapeOrientation.Equals("LandscapeRight"))
-	{
-		return UIInterfaceOrientationLandscapeRight;
-	}
-	else
-	{
-		return UIInterfaceOrientationPortrait;
-	}
+
+	return UIInterfaceOrientationPortrait;
 }
 
 - (UIRectEdge)preferredScreenEdgesDeferringSystemGestures
diff --git a/Engine/Source/Runtime/IOS/IOSRuntimeSettings/Classes/IOSRuntimeSettings.h b/Engine/Source/Runtime/IOS/IOSRuntimeSettings/Classes/IOSRuntimeSettings.h
--- a/Engine/Source/Runtime/IOS/IOSRuntimeSettings/Classes/IOSRuntimeSettings.h
+++ b/Engine/Source/Runtime/IOS/IOSRuntimeSettings/Classes/IOSRuntimeSettings.h
@@ -350,7 +350,7 @@ public:
 	uint32 bSupportsLandscapeRightOrientation : 1;
 
 	// The Preferred Orientation will be used as the initial orientation at launch when both Landscape Left and Landscape Right orientations are to be supported.
-	UPROPERTY(GlobalConfig, EditAnywhere, Category = DeviceOrientations)
+	UPROPERTY(GlobalConfig, EditAnywhere, Category = DeviceOrientations, meta = (DisplayName = "Preferred Landscape Orientation"))
 	EIOSLandscapeOrientation PreferredLandscapeOrientation;
 
 	// Specifies the the display name for the application. This will be displayed under the icon on the device.
@@ -512,4 +512,8 @@ public:
 	virtual void PostInitProperties() override;
 	// End of UObject interface
 #endif
+
+private:
+	virtual void EnsureOrientationInProjectDefaultEngine();
+
 };
diff --git a/Engine/Source/Runtime/IOS/IOSRuntimeSettings/Private/IOSRuntimeSettings.cpp b/Engine/Source/Runtime/IOS/IOSRuntimeSettings/Private/IOSRuntimeSettings.cpp
--- a/Engine/Source/Runtime/IOS/IOSRuntimeSettings/Private/IOSRuntimeSettings.cpp
+++ b/Engine/Source/Runtime/IOS/IOSRuntimeSettings/Private/IOSRuntimeSettings.cpp
@@ -184,5 +184,16 @@ void UIOSRuntimeSettings::PostInitProperties()
 	{
 		MinimumiOSVersion = EIOSVersion::IOS_10;
 	}
+
+	EnsureOrientationInProjectDefaultEngine();
 }
 #endif
+
+void UIOSRuntimeSettings::EnsureOrientationInProjectDefaultEngine()
+{
+	// Make sure the values from BaseEngine.ini are written in the project's DefaultEngine.ini
+	UpdateSinglePropertyInConfigFile(GetClass()->FindPropertyByName(GET_MEMBER_NAME_CHECKED(UIOSRuntimeSettings, bSupportsPortraitOrientation)), GetDefaultConfigFilename());
+	UpdateSinglePropertyInConfigFile(GetClass()->FindPropertyByName(GET_MEMBER_NAME_CHECKED(UIOSRuntimeSettings, bSupportsUpsideDownOrientation)), GetDefaultConfigFilename());
+	UpdateSinglePropertyInConfigFile(GetClass()->FindPropertyByName(GET_MEMBER_NAME_CHECKED(UIOSRuntimeSettings, bSupportsLandscapeLeftOrientation)), GetDefaultConfigFilename());
+	UpdateSinglePropertyInConfigFile(GetClass()->FindPropertyByName(GET_MEMBER_NAME_CHECKED(UIOSRuntimeSettings, PreferredLandscapeOrientation)), GetDefaultConfigFilename());
+}
