Linux: Underscores followed by an uppercase letter are reserved identifiers

diff --git a/Engine/Source/Runtime/ApplicationCore/Private/Linux/LinuxWindow.cpp b/Engine/Source/Runtime/ApplicationCore/Private/Linux/LinuxWindow.cpp
--- a/Engine/Source/Runtime/ApplicationCore/Private/Linux/LinuxWindow.cpp
+++ b/Engine/Source/Runtime/ApplicationCore/Private/Linux/LinuxWindow.cpp
@@ -469,7 +469,7 @@ void FLinuxWindow::Hide()
 	}
 }
 
-static void _GetBestFullscreenResolution( SDL_HWindow hWnd, int32 *pWidth, int32 *pHeight )
+static void GetBestFullscreenResolution( SDL_HWindow hWnd, int32 *pWidth, int32 *pHeight )
 {
 	uint32 InitializedMode = false;
 	uint32 BestWidth = 0;
@@ -637,7 +637,7 @@ void FLinuxWindow::AdjustCachedSize( FVector2D& Size ) const
 			SizeW = VirtualWidth ;
 			SizeH = VirtualHeight;
 
-			_GetBestFullscreenResolution( HWnd, &SizeW, &SizeH );
+			GetBestFullscreenResolution( HWnd, &SizeW, &SizeH );
 		}
 
 		Size = FVector2D( SizeW, SizeH );
