Fixed crash at exit on Mac that could happen if quit action was initiated during window<->fullscreen transition

diff --git a/Engine/Source/Runtime/ApplicationCore/Private/Mac/CocoaWindow.cpp b/Engine/Source/Runtime/ApplicationCore/Private/Mac/CocoaWindow.cpp
--- a/Engine/Source/Runtime/ApplicationCore/Private/Mac/CocoaWindow.cpp
+++ b/Engine/Source/Runtime/ApplicationCore/Private/Mac/CocoaWindow.cpp
@@ -207,7 +207,7 @@ NSString* NSPerformDragOperation = @"NSPerformDragOperation";
 
 - (void)windowWillEnterFullScreen:(NSNotification*)Notification
 {
-	FMacCursor* MacCursor = (FMacCursor*)MacApplication->Cursor.Get();
+	FMacCursor* MacCursor = MacApplication ? (FMacCursor*)MacApplication->Cursor.Get() : nullptr;
 	if (MacCursor)
 	{
 		MacCursor->SetShouldIgnoreLocking(true);
@@ -242,7 +242,7 @@ NSString* NSPerformDragOperation = @"NSPerformDragOperation";
 
 - (void)windowWillExitFullScreen:(NSNotification *)Notification
 {
-	FMacCursor* MacCursor = (FMacCursor*)MacApplication->Cursor.Get();
+	FMacCursor* MacCursor = MacApplication ? (FMacCursor*)MacApplication->Cursor.Get() : nullptr;
 	if (MacCursor)
 	{
 		MacCursor->SetShouldIgnoreLocking(true);
