GetMonitorWorkAreaFromPoint no longer assumes 0,0 is the top left corner, which won't be true in cases where the top left monitor isn't primary.

#rnx

diff --git a/Engine/Source/Runtime/ApplicationCore/Private/GenericPlatform/GenericApplication.cpp b/Engine/Source/Runtime/ApplicationCore/Private/GenericPlatform/GenericApplication.cpp
--- a/Engine/Source/Runtime/ApplicationCore/Private/GenericPlatform/GenericApplication.cpp
+++ b/Engine/Source/Runtime/ApplicationCore/Private/GenericPlatform/GenericApplication.cpp
@@ -138,7 +138,7 @@ FPlatformRect FDisplayMetrics::GetMonitorWorkAreaFromPoint(const FVector2D& Poin
 	for (const FMonitorInfo& Info : MonitorInfo)
 	{
 		// The point may not actually be inside the work area (for example on Windows taskbar or Mac menu bar), so we use DisplayRect to find the monitor
-		if (Point.X >= Info.DisplayRect.Left && Point.X <= Info.DisplayRect.Right && Point.Y >= Info.DisplayRect.Top && Point.Y <= Info.DisplayRect.Bottom)
+		if (Point.X >= Info.DisplayRect.Left && Point.X < Info.DisplayRect.Right && Point.Y >= Info.DisplayRect.Top && Point.Y < Info.DisplayRect.Bottom)
 		{
 			return Info.WorkArea;
 		}
