Disable heartbeat during windows console text output. To allow selecting text in windows log console without crashing due to heartbeat after 25 seconds. The main thread will still block but will not crash due to heartbeat anymore.

diff --git a/Engine/Source/Runtime/ApplicationCore/Private/Windows/WindowsConsoleOutputDevice.cpp b/Engine/Source/Runtime/ApplicationCore/Private/Windows/WindowsConsoleOutputDevice.cpp
--- a/Engine/Source/Runtime/ApplicationCore/Private/Windows/WindowsConsoleOutputDevice.cpp
+++ b/Engine/Source/Runtime/ApplicationCore/Private/Windows/WindowsConsoleOutputDevice.cpp
@@ -203,6 +203,9 @@ void FWindowsConsoleOutputDevice::Serialize( const TCHAR* Data, ELogVerbosity::T
 				}
 				TCHAR OutputString[MAX_SPRINTF]=TEXT(""); //@warning: this is safe as FCString::Sprintf only use 1024 characters max
 				FCString::Sprintf(OutputString,TEXT("%s%s"),*FOutputDeviceHelper::FormatLogLine(Verbosity, Category, Data, GPrintLogTimes,RealTime),LINE_TERMINATOR);
+
+				// WriteConsole blocks when user interacting with console window (text is selected, scrollbar held, etc..)
+				FSlowHeartBeatScope SuspendHeartBeat;
 				uint32 Written;
 				WriteConsole( ConsoleHandle, OutputString, FCString::Strlen(OutputString), (::DWORD*)&Written, NULL );
 
