Linux: Add a command line option to disable tooltip workaround

diff --git a/Engine/Source/Runtime/ApplicationCore/Private/Linux/LinuxWindow.cpp b/Engine/Source/Runtime/ApplicationCore/Private/Linux/LinuxWindow.cpp
--- a/Engine/Source/Runtime/ApplicationCore/Private/Linux/LinuxWindow.cpp
+++ b/Engine/Source/Runtime/ApplicationCore/Private/Linux/LinuxWindow.cpp
@@ -505,21 +505,26 @@ void FLinuxWindow::ReshapeWindow( int32 NewX, int32 NewY, int32 NewWidth, int32
 {
 	// Some vulkan video drivers have issues with specific height ranges causing them to corrupt the texture rendered
 	// Moving these nearest values removes this corruption.
-	if (NewHeight >= 9 && NewHeight <= 10)
-	{
-		NewHeight = 11;
-	}
-	else if (NewHeight >= 17 && NewHeight <= 21)
-	{
-		NewHeight = 22;
-	}
-	else if (NewHeight >= 33 && NewHeight <= 43)
-	{
-		NewHeight = 44;
-	}
-	else if (NewHeight >= 65 && NewHeight <= 85)
+	static bool bDisableVulkanWorkaround = FParse::Param(FCommandLine::Get(), TEXT("disablevulkanworkaround"));
+
+	if (!bDisableVulkanWorkaround)
 	{
-		NewHeight = 86;
+		if (NewHeight >= 9 && NewHeight <= 10)
+		{
+			NewHeight = 11;
+		}
+		else if (NewHeight >= 17 && NewHeight <= 21)
+		{
+			NewHeight = 22;
+		}
+		else if (NewHeight >= 33 && NewHeight <= 43)
+		{
+			NewHeight = 44;
+		}
+		else if (NewHeight >= 65 && NewHeight <= 85)
+		{
+			NewHeight = 86;
+		}
 	}
 
 	// X11 will take until the next frame to send a SizeChanged event. This means the X11 window
