Optionally disable android back from being sent by controllers.

diff --git a/Engine/Source/Runtime/Android/AndroidRuntimeSettings/Classes/AndroidRuntimeSettings.h b/Engine/Source/Runtime/Android/AndroidRuntimeSettings/Classes/AndroidRuntimeSettings.h
--- a/Engine/Source/Runtime/Android/AndroidRuntimeSettings/Classes/AndroidRuntimeSettings.h
+++ b/Engine/Source/Runtime/Android/AndroidRuntimeSettings/Classes/AndroidRuntimeSettings.h
@@ -420,6 +420,10 @@ public:
 	UPROPERTY(GlobalConfig, EditAnywhere, Category = Input, meta = (DisplayName = "Allow Bluetooth controllers"))
 	bool bAllowControllers;
 
+	// If checked, controllers will not send Android_Back and Android_Menu events that might cause unnecce
+	UPROPERTY(GlobalConfig, EditAnywhere, Category = Input, meta = (DisplayName = "Block Android system keys being sent from controllers"))
+	bool bBlockAndroidKeysOnControllers;
+
 	/** Android encoding options. */
 	UPROPERTY(GlobalConfig, EditAnywhere, Category = Audio, meta = (DisplayName = "Encoding Format"))
 	TEnumAsByte<EAndroidAudio::Type> AndroidAudio;
diff --git a/Engine/Source/Runtime/Android/AndroidRuntimeSettings/Private/AndroidRuntimeSettings.cpp b/Engine/Source/Runtime/Android/AndroidRuntimeSettings/Private/AndroidRuntimeSettings.cpp
--- a/Engine/Source/Runtime/Android/AndroidRuntimeSettings/Private/AndroidRuntimeSettings.cpp
+++ b/Engine/Source/Runtime/Android/AndroidRuntimeSettings/Private/AndroidRuntimeSettings.cpp
@@ -22,6 +22,7 @@ UAndroidRuntimeSettings::UAndroidRuntimeSettings(const FObjectInitializer& Objec
 	, bEnableGooglePlaySupport(false)
 	, bUseGetAccounts(false)
 	, bSupportAdMob(true)
+	, bBlockAndroidKeysOnControllers(false)
 	, AudioSampleRate(44100)
 	, AudioCallbackBufferFrameSize(1024)
 	, AudioNumBuffersToEnqueue(4)
diff --git a/Engine/Source/Runtime/ApplicationCore/Private/Android/AndroidInputInterface.cpp b/Engine/Source/Runtime/ApplicationCore/Private/Android/AndroidInputInterface.cpp
--- a/Engine/Source/Runtime/ApplicationCore/Private/Android/AndroidInputInterface.cpp
+++ b/Engine/Source/Runtime/ApplicationCore/Private/Android/AndroidInputInterface.cpp
@@ -21,6 +21,7 @@ bool FAndroidInputInterface::VibeIsOn;
 FForceFeedbackValues FAndroidInputInterface::VibeValues;
 
 bool FAndroidInputInterface::bAllowControllers = true;
+bool FAndroidInputInterface::bBlockAndroidKeysOnControllers = false;
 
 FAndroidControllerData FAndroidInputInterface::OldControllerData[MAX_NUM_CONTROLLERS];
 FAndroidControllerData FAndroidInputInterface::NewControllerData[MAX_NUM_CONTROLLERS];
@@ -66,6 +67,7 @@ FAndroidInputInterface::FAndroidInputInterface(const TSharedRef< FGenericApplica
 	, Cursor(InCursor)
 {
 	GConfig->GetBool(TEXT("/Script/AndroidRuntimeSettings.AndroidRuntimeSettings"), TEXT("bAllowControllers"), bAllowControllers, GEngineIni);
+	GConfig->GetBool(TEXT("/Script/AndroidRuntimeSettings.AndroidRuntimeSettings"), TEXT("bBlockAndroidKeysOnControllers"), bBlockAndroidKeysOnControllers, GEngineIni);
 
 	ButtonMapping[0] = FGamepadKeyNames::FaceButtonBottom;
 	ButtonMapping[1] = FGamepadKeyNames::FaceButtonRight;
@@ -1324,9 +1326,17 @@ void FAndroidInputInterface::JoystickButtonEvent(int32 deviceId, int32 buttonId,
 											 }
 											 break;
 				case AKEYCODE_BUTTON_START:
-				case AKEYCODE_MENU:          NewControllerData[deviceId].ButtonStates[6] = buttonDown; NewControllerData[deviceId].ButtonStates[17] = buttonDown;  break;
+				case AKEYCODE_MENU:          NewControllerData[deviceId].ButtonStates[6] = buttonDown;
+											 if (!bBlockAndroidKeysOnControllers)
+											 {
+												 NewControllerData[deviceId].ButtonStates[17] = buttonDown;  break;
+											 }
 				case AKEYCODE_BUTTON_SELECT:
-				case AKEYCODE_BACK:          NewControllerData[deviceId].ButtonStates[7] = buttonDown; NewControllerData[deviceId].ButtonStates[16] = buttonDown;  break;
+				case AKEYCODE_BACK:          NewControllerData[deviceId].ButtonStates[7] = buttonDown;
+											 if (!bBlockAndroidKeysOnControllers)
+											 {
+												 NewControllerData[deviceId].ButtonStates[16] = buttonDown;  break;
+											 }
 				case AKEYCODE_BUTTON_THUMBL: NewControllerData[deviceId].ButtonStates[8] = buttonDown; break;
 				case AKEYCODE_BUTTON_THUMBR: NewControllerData[deviceId].ButtonStates[9] = buttonDown; break;
 				case AKEYCODE_BUTTON_L2:     NewControllerData[deviceId].ButtonStates[10] = buttonDown; break;
@@ -1347,8 +1357,18 @@ void FAndroidInputInterface::JoystickButtonEvent(int32 deviceId, int32 buttonId,
 				case AKEYCODE_BUTTON_X:      NewControllerData[deviceId].ButtonStates[3] = buttonDown; break; // Y
 				case AKEYCODE_BUTTON_Y:      NewControllerData[deviceId].ButtonStates[4] = buttonDown; break; // L1
 				case AKEYCODE_BUTTON_Z:      NewControllerData[deviceId].ButtonStates[5] = buttonDown; break; // R1
-				case AKEYCODE_BUTTON_R1:     NewControllerData[deviceId].ButtonStates[6] = buttonDown; NewControllerData[deviceId].ButtonStates[17] = buttonDown;  break; // Menu
-				case AKEYCODE_BUTTON_L1:     NewControllerData[deviceId].ButtonStates[7] = buttonDown; NewControllerData[deviceId].ButtonStates[16] = buttonDown;  break; // View
+				case AKEYCODE_BUTTON_R1:     NewControllerData[deviceId].ButtonStates[6] = buttonDown;
+											 if (!bBlockAndroidKeysOnControllers)
+											 {
+												 NewControllerData[deviceId].ButtonStates[17] = buttonDown; // Menu
+											 }
+											 break;
+				case AKEYCODE_BUTTON_L1:     NewControllerData[deviceId].ButtonStates[7] = buttonDown;
+											 if (!bBlockAndroidKeysOnControllers)
+											 {
+												 NewControllerData[deviceId].ButtonStates[16] = buttonDown; // View
+											 }
+											 break;
 				case AKEYCODE_BUTTON_L2:     NewControllerData[deviceId].ButtonStates[8] = buttonDown; break; // ThumbL
 				case AKEYCODE_BUTTON_R2:     NewControllerData[deviceId].ButtonStates[9] = buttonDown; break; // ThumbR
 			}
@@ -1363,8 +1383,18 @@ void FAndroidInputInterface::JoystickButtonEvent(int32 deviceId, int32 buttonId,
 				case AKEYCODE_BUTTON_X:      NewControllerData[deviceId].ButtonStates[3] = buttonDown; break; // Triangle
 				case AKEYCODE_BUTTON_Y:      NewControllerData[deviceId].ButtonStates[4] = buttonDown; break; // L1
 				case AKEYCODE_BUTTON_Z:      NewControllerData[deviceId].ButtonStates[5] = buttonDown; break; // R1
-				case AKEYCODE_BUTTON_L2:     NewControllerData[deviceId].ButtonStates[6] = buttonDown; NewControllerData[deviceId].ButtonStates[17] = buttonDown;  break; // Options
-				case AKEYCODE_MENU:     NewControllerData[deviceId].ButtonStates[7] = buttonDown; NewControllerData[deviceId].ButtonStates[16] = buttonDown;  break; // Share
+				case AKEYCODE_BUTTON_L2:     NewControllerData[deviceId].ButtonStates[6] = buttonDown;
+											 if (!bBlockAndroidKeysOnControllers)
+											 {
+											 	NewControllerData[deviceId].ButtonStates[17] = buttonDown; // Options
+											 }
+											 break;
+				case AKEYCODE_MENU:          NewControllerData[deviceId].ButtonStates[7] = buttonDown;
+											 if (!bBlockAndroidKeysOnControllers)
+											 {
+												 NewControllerData[deviceId].ButtonStates[16] = buttonDown; // Touchpad
+											 }
+											 break;
 				case AKEYCODE_BUTTON_SELECT: NewControllerData[deviceId].ButtonStates[8] = buttonDown; break; // ThumbL
 				case AKEYCODE_BUTTON_START:  NewControllerData[deviceId].ButtonStates[9] = buttonDown; break; // ThumbR
 				case AKEYCODE_BUTTON_L1:     NewControllerData[deviceId].ButtonStates[10] = buttonDown; break; // L2
diff --git a/Engine/Source/Runtime/ApplicationCore/Public/Android/AndroidInputInterface.h b/Engine/Source/Runtime/ApplicationCore/Public/Android/AndroidInputInterface.h
--- a/Engine/Source/Runtime/ApplicationCore/Public/Android/AndroidInputInterface.h
+++ b/Engine/Source/Runtime/ApplicationCore/Public/Android/AndroidInputInterface.h
@@ -313,6 +313,9 @@ private:
 	// should we allow controllers to send input
 	static bool bAllowControllers;
 
+	// should we allow controllers to send Android_Back and Android_Menu events
+	static bool bBlockAndroidKeysOnControllers;
+
 	static FAndroidGamepadDeviceMapping DeviceMapping[MAX_NUM_CONTROLLERS];
 
 	static FAndroidControllerData OldControllerData[MAX_NUM_CONTROLLERS];
diff --git a/Engine/Source/Runtime/Launch/Private/Android/LaunchAndroid.cpp b/Engine/Source/Runtime/Launch/Private/Android/LaunchAndroid.cpp
--- a/Engine/Source/Runtime/Launch/Private/Android/LaunchAndroid.cpp
+++ b/Engine/Source/Runtime/Launch/Private/Android/LaunchAndroid.cpp
@@ -95,7 +95,7 @@ static const uint32 ValidGamepadKeyCodesList[] =
 	AKEYCODE_BUTTON_START,
 	AKEYCODE_MENU,
 	AKEYCODE_BUTTON_SELECT,
-//	AKEYCODE_BACK,
+	AKEYCODE_BACK,
 	AKEYCODE_BUTTON_THUMBL,
 	AKEYCODE_BUTTON_THUMBR,
 	AKEYCODE_BUTTON_L2,
@@ -956,17 +956,19 @@ static int32_t HandleInputCB(struct android_app* app, AInputEvent* event)
 
 		FPlatformMisc::LowLevelOutputDebugStringf(TEXT("Received keycode: %d, softkey: %d"), keyCode, bSoftKey ? 1 : 0);
 
+		//Only pass on the device id if really a gamepad, joystick or dpad (allows menu and back to be treated as gamepad events)
+		int32 device = -1;
+		if ((((EventSource & AINPUT_SOURCE_JOYSTICK) == AINPUT_SOURCE_JOYSTICK) && (GetAxes != NULL)) ||
+			((EventSource & AINPUT_SOURCE_GAMEPAD) == AINPUT_SOURCE_GAMEPAD) ||
+			((EventSource & AINPUT_SOURCE_DPAD) == AINPUT_SOURCE_DPAD))
+		{
+			device = AInputEvent_getDeviceId(event);
+		}
+
 		//Trap codes handled as possible gamepad events
-		if (ValidGamepadKeyCodes.Contains(keyCode))
+		if (device >= 0 && ValidGamepadKeyCodes.Contains(keyCode))
 		{
-			//Only pass on the device id if really a gamepad, joystick or dpad (allows menu and back to be treated as gamepad events)
-			int32 device = 0;
-			if ( (((EventSource & AINPUT_SOURCE_JOYSTICK) == AINPUT_SOURCE_JOYSTICK) && (GetAxes != NULL)) ||
-				 ((EventSource & AINPUT_SOURCE_GAMEPAD) == AINPUT_SOURCE_GAMEPAD) ||
-				 ((EventSource & AINPUT_SOURCE_DPAD) == AINPUT_SOURCE_DPAD) )
-			{
-				device = AInputEvent_getDeviceId(event);
-			}
+			
 			bool down = AKeyEvent_getAction(event) != AKEY_EVENT_ACTION_UP;
 			FAndroidInputInterface::JoystickButtonEvent(device, keyCode, down);
 			FPlatformMisc::LowLevelOutputDebugStringf(TEXT("Received gamepad button: %d"), keyCode);
