FAndroidInputInterface changes not added to the lumin version of the interface.

diff --git a/Engine/Source/Runtime/ApplicationCore/Private/Lumin/LuminInputInterface.cpp b/Engine/Source/Runtime/ApplicationCore/Private/Lumin/LuminInputInterface.cpp
--- a/Engine/Source/Runtime/ApplicationCore/Private/Lumin/LuminInputInterface.cpp
+++ b/Engine/Source/Runtime/ApplicationCore/Private/Lumin/LuminInputInterface.cpp
@@ -53,6 +53,11 @@ bool FAndroidInputInterface::IsControllerAssignedToGamepad(int32 ControllerId)
 	return true;
 }
 
+FString FAndroidInputInterface::GetGamepadControllerName(int32 ControllerId)
+{
+	return FString(TEXT("Generic"));
+}
+
 void FAndroidInputInterface::SetForceFeedbackChannelValue(int32 ControllerId, FForceFeedbackChannelType ChannelType, float Value)
 {
 	for (TSharedPtr<IInputDevice>& InputDevice : LuminInputDevices)
diff --git a/Engine/Source/Runtime/ApplicationCore/Public/Lumin/LuminInputInterface.h b/Engine/Source/Runtime/ApplicationCore/Public/Lumin/LuminInputInterface.h
--- a/Engine/Source/Runtime/ApplicationCore/Public/Lumin/LuminInputInterface.h
+++ b/Engine/Source/Runtime/ApplicationCore/Public/Lumin/LuminInputInterface.h
@@ -28,6 +28,7 @@ public:
 	static void ResetGamepadAssignments();
 	static void ResetGamepadAssignmentToController(int32 ControllerId);
 	static bool IsControllerAssignedToGamepad(int32 ControllerId);
+	static FString GetGamepadControllerName(int32 ControllerId);
 
 	/** IInputInterface interface */
 	virtual void SetForceFeedbackChannelValue(int32 ControllerId, FForceFeedbackChannelType ChannelType, float Value) override;
@@ -36,6 +37,7 @@ public:
 	virtual void SetLightColor(int32 ControllerId, struct FColor Color) override {};
 	virtual void ResetLightColor(int32 ControllerId) override {};
 
+	void SetGamepadsAllowed(bool bAllowed) {}
 	virtual void AddExternalInputDevice(TSharedPtr<class IInputDevice>& InputDevice);
 	virtual bool IsGamepadAttached() const;
 
