Fixing repeat input for analog events to function like console/pc for android and fixing some misplaced breaks.

diff --git a/Engine/Source/Runtime/ApplicationCore/Private/Android/AndroidInputInterface.cpp b/Engine/Source/Runtime/ApplicationCore/Private/Android/AndroidInputInterface.cpp
--- a/Engine/Source/Runtime/ApplicationCore/Private/Android/AndroidInputInterface.cpp
+++ b/Engine/Source/Runtime/ApplicationCore/Private/Android/AndroidInputInterface.cpp
@@ -949,21 +949,24 @@ void FAndroidInputInterface::SendControllerEvents()
 			FAndroidControllerData& OldControllerState = OldControllerData[ControllerIndex];
 			FAndroidControllerData& NewControllerState = NewControllerData[ControllerIndex];
 
-			if (NewControllerState.LXAnalog != OldControllerState.LXAnalog)
+			// Send controller events any time we have a large enough input threshold similarly to PC/Console (see: XInputInterface.cpp)
+			const float RepeatDeadzone = 0.24f;
+
+			if (NewControllerState.LXAnalog != OldControllerState.LXAnalog || FMath::Abs(NewControllerState.LXAnalog) >= RepeatDeadzone)
 			{
 				MessageHandler->OnControllerAnalog(FGamepadKeyNames::LeftAnalogX, NewControllerState.DeviceId, NewControllerState.LXAnalog);
 			}
-			if (NewControllerState.LYAnalog != OldControllerState.LYAnalog)
+			if (NewControllerState.LYAnalog != OldControllerState.LYAnalog || FMath::Abs(NewControllerState.LYAnalog) >= RepeatDeadzone)
 			{
 				//LOGD("    Sending updated LeftAnalogY value of %f", NewControllerState.LYAnalog);
 				MessageHandler->OnControllerAnalog(FGamepadKeyNames::LeftAnalogY, NewControllerState.DeviceId, NewControllerState.LYAnalog);
 			}
-			if (NewControllerState.RXAnalog != OldControllerState.RXAnalog)
+			if (NewControllerState.RXAnalog != OldControllerState.RXAnalog || FMath::Abs(NewControllerState.RXAnalog) >= RepeatDeadzone)
 			{
 				//LOGD("    Sending updated RightAnalogX value of %f", NewControllerState.RXAnalog);
 				MessageHandler->OnControllerAnalog(FGamepadKeyNames::RightAnalogX, NewControllerState.DeviceId, NewControllerState.RXAnalog);
 			}
-			if (NewControllerState.RYAnalog != OldControllerState.RYAnalog)
+			if (NewControllerState.RYAnalog != OldControllerState.RYAnalog || FMath::Abs(NewControllerState.RYAnalog) >= RepeatDeadzone)
 			{
 				//LOGD("    Sending updated RightAnalogY value of %f", NewControllerState.RYAnalog);
 				MessageHandler->OnControllerAnalog(FGamepadKeyNames::RightAnalogY, NewControllerState.DeviceId, NewControllerState.RYAnalog);
@@ -1329,14 +1332,16 @@ void FAndroidInputInterface::JoystickButtonEvent(int32 deviceId, int32 buttonId,
 				case AKEYCODE_MENU:          NewControllerData[deviceId].ButtonStates[6] = buttonDown;
 											 if (!bBlockAndroidKeysOnControllers)
 											 {
-												 NewControllerData[deviceId].ButtonStates[17] = buttonDown;  break;
+												 NewControllerData[deviceId].ButtonStates[17] = buttonDown;
 											 }
+											 break;
 				case AKEYCODE_BUTTON_SELECT:
 				case AKEYCODE_BACK:          NewControllerData[deviceId].ButtonStates[7] = buttonDown;
 											 if (!bBlockAndroidKeysOnControllers)
 											 {
-												 NewControllerData[deviceId].ButtonStates[16] = buttonDown;  break;
+												 NewControllerData[deviceId].ButtonStates[16] = buttonDown;
 											 }
+											 break;
 				case AKEYCODE_BUTTON_THUMBL: NewControllerData[deviceId].ButtonStates[8] = buttonDown; break;
 				case AKEYCODE_BUTTON_THUMBR: NewControllerData[deviceId].ButtonStates[9] = buttonDown; break;
 				case AKEYCODE_BUTTON_L2:     NewControllerData[deviceId].ButtonStates[10] = buttonDown; break;
