Adding support for repeat button input on ios.

Josh.Adams, Peter.Sauerbrei

diff --git a/Engine/Source/Runtime/ApplicationCore/Private/IOS/IOSInputInterface.cpp b/Engine/Source/Runtime/ApplicationCore/Private/IOS/IOSInputInterface.cpp
--- a/Engine/Source/Runtime/ApplicationCore/Private/IOS/IOSInputInterface.cpp
+++ b/Engine/Source/Runtime/ApplicationCore/Private/IOS/IOSInputInterface.cpp
@@ -5,6 +5,7 @@
 #include "Misc/ConfigCacheIni.h"
 #include "Misc/ScopeLock.h"
 #include "HAL/IConsoleManager.h"
+#include "HAL/PlatformTime.h"
 
 #import <AudioToolbox/AudioToolbox.h>
 
@@ -334,17 +335,36 @@ void FIOSInputInterface::SendControllerEvents()
 			Controller.bPauseWasPressed = false;
 		}
 		
-		// @todo tvos: Handle repeated buttons?
+        const double CurrentTime = FPlatformTime::Seconds();
+        const float IniitialRepeatDelay = 0.2f;
+        const float RepeatDelay = 0.1;
 		
 #define HANDLE_BUTTON(Gamepad, GCButton, UEButton) \
 if ((Controller.Previous##Gamepad == nil && Gamepad.GCButton.pressed) || Controller.Previous##Gamepad.GCButton.pressed != Gamepad.GCButton.pressed) \
 { \
 	NSLog(@"%@ button %s on controller %d", (ExtendedGamepad.GCButton.pressed) ? @"Pressed" : @"Released", TCHAR_TO_ANSI(*UEButton.ToString()), (int32)Cont.playerIndex); \
 	(Gamepad.GCButton.pressed) ? MessageHandler->OnControllerButtonPressed(UEButton, Cont.playerIndex, false) : MessageHandler->OnControllerButtonReleased(UEButton, Cont.playerIndex, false); \
+    NextKeyRepeatTime.FindOrAdd(UEButton) = CurrentTime + IniitialRepeatDelay; \
+} \
+else if(Gamepad.GCButton.pressed) \
+{ \
+    double* NextRepeatTime = NextKeyRepeatTime.Find(UEButton); \
+    if(NextRepeatTime && *NextRepeatTime <= CurrentTime) \
+    { \
+        MessageHandler->OnControllerButtonPressed(UEButton, Cont.playerIndex, false); \
+*NextRepeatTime = CurrentTime + RepeatDelay; \
+    } \
+} \
+else \
+{ \
+    NextKeyRepeatTime.Remove(UEButton); \
 }
 
+        // Send controller events any time we are passed the given input threshold similarly to PC/Console (see: XInputInterface.cpp)
+        const float RepeatDeadzone = 0.24;
+        
 #define HANDLE_ANALOG(Gamepad, GCAxis, UEAxis) \
-if (Controller.Previous##Gamepad != nil && Gamepad.GCAxis.value != Controller.Previous##Gamepad.GCAxis.value) \
+if ((Controller.Previous##Gamepad != nil && Gamepad.GCAxis.value != Controller.Previous##Gamepad.GCAxis.value) || (Gamepad.GCAxis.value < -RepeatDeadzone || Gamepad.GCAxis.value > RepeatDeadzone)) \
 { \
 	NSLog(@"Axis %s is %f", TCHAR_TO_ANSI(*UEAxis.ToString()), Gamepad.GCAxis.value); \
 	MessageHandler->OnControllerAnalog(UEAxis, Cont.playerIndex, Gamepad.GCAxis.value); \
@@ -431,7 +451,7 @@ if (Controller.Previous##Gamepad != nil && Gamepad.GCAxis.value != Controller.Pr
 			Controller.PreviousMicroGamepad = [MicroGamepad saveSnapshot];
 			[Controller.PreviousMicroGamepad retain];
         }
-		
+        
 		// motion is orthogonal to buttons
 // @todo tvos: handle motion without attitude or rotation rate
 #if 0
diff --git a/Engine/Source/Runtime/ApplicationCore/Public/IOS/IOSInputInterface.h b/Engine/Source/Runtime/ApplicationCore/Public/IOS/IOSInputInterface.h
--- a/Engine/Source/Runtime/ApplicationCore/Public/IOS/IOSInputInterface.h
+++ b/Engine/Source/Runtime/ApplicationCore/Public/IOS/IOSInputInterface.h
@@ -182,4 +182,6 @@ private:
 	float LastHapticValue;
 	
 	int HapticFeedbackSupportLevel;
+    
+    TMap<FName, double> NextKeyRepeatTime;
 };
