Adding support for analog virtual buttons on Android/IOs.

Peter.Sauerbrei, Josh.Adams, Chris.Babcock

diff --git a/Engine/Source/Runtime/ApplicationCore/Private/Android/AndroidInputInterface.cpp b/Engine/Source/Runtime/ApplicationCore/Private/Android/AndroidInputInterface.cpp
--- a/Engine/Source/Runtime/ApplicationCore/Private/Android/AndroidInputInterface.cpp
+++ b/Engine/Source/Runtime/ApplicationCore/Private/Android/AndroidInputInterface.cpp
@@ -88,6 +88,16 @@ FAndroidInputInterface::FAndroidInputInterface(const TSharedRef< FGenericApplica
 	ButtonMapping[16] = AndroidKeyNames::Android_Back;  // Technically just an alias for SpecialLeft
 	ButtonMapping[17] = AndroidKeyNames::Android_Menu;  // Technically just an alias for SpecialRight
 
+	// Virtual buttons
+	ButtonMapping[MAX_NUM_PHYSICAL_CONTROLLER_BUTTONS + 0] = FGamepadKeyNames::LeftStickLeft;
+	ButtonMapping[MAX_NUM_PHYSICAL_CONTROLLER_BUTTONS + 1] = FGamepadKeyNames::LeftStickRight;
+	ButtonMapping[MAX_NUM_PHYSICAL_CONTROLLER_BUTTONS + 2] = FGamepadKeyNames::LeftStickUp;
+	ButtonMapping[MAX_NUM_PHYSICAL_CONTROLLER_BUTTONS + 3] = FGamepadKeyNames::LeftStickDown;
+	ButtonMapping[MAX_NUM_PHYSICAL_CONTROLLER_BUTTONS + 4] = FGamepadKeyNames::RightStickLeft;
+	ButtonMapping[MAX_NUM_PHYSICAL_CONTROLLER_BUTTONS + 5] = FGamepadKeyNames::RightStickRight;
+	ButtonMapping[MAX_NUM_PHYSICAL_CONTROLLER_BUTTONS + 6] = FGamepadKeyNames::RightStickUp;
+	ButtonMapping[MAX_NUM_PHYSICAL_CONTROLLER_BUTTONS + 7] = FGamepadKeyNames::RightStickDown;
+
 	InitialButtonRepeatDelay = 0.2f;
 	ButtonRepeatDelay = 0.1f;
 
@@ -955,21 +965,29 @@ void FAndroidInputInterface::SendControllerEvents()
 			if (NewControllerState.LXAnalog != OldControllerState.LXAnalog || FMath::Abs(NewControllerState.LXAnalog) >= RepeatDeadzone)
 			{
 				MessageHandler->OnControllerAnalog(FGamepadKeyNames::LeftAnalogX, NewControllerState.DeviceId, NewControllerState.LXAnalog);
+				NewControllerState.ButtonStates[MAX_NUM_PHYSICAL_CONTROLLER_BUTTONS + 1] = NewControllerState.LXAnalog >= RepeatDeadzone;
+				NewControllerState.ButtonStates[MAX_NUM_PHYSICAL_CONTROLLER_BUTTONS + 0] = NewControllerState.LXAnalog <= -RepeatDeadzone;
 			}
 			if (NewControllerState.LYAnalog != OldControllerState.LYAnalog || FMath::Abs(NewControllerState.LYAnalog) >= RepeatDeadzone)
 			{
 				//LOGD("    Sending updated LeftAnalogY value of %f", NewControllerState.LYAnalog);
 				MessageHandler->OnControllerAnalog(FGamepadKeyNames::LeftAnalogY, NewControllerState.DeviceId, NewControllerState.LYAnalog);
+				NewControllerState.ButtonStates[MAX_NUM_PHYSICAL_CONTROLLER_BUTTONS + 2] = NewControllerState.LYAnalog >= RepeatDeadzone;
+				NewControllerState.ButtonStates[MAX_NUM_PHYSICAL_CONTROLLER_BUTTONS + 3] = NewControllerState.LYAnalog <= -RepeatDeadzone;
 			}
 			if (NewControllerState.RXAnalog != OldControllerState.RXAnalog || FMath::Abs(NewControllerState.RXAnalog) >= RepeatDeadzone)
 			{
 				//LOGD("    Sending updated RightAnalogX value of %f", NewControllerState.RXAnalog);
 				MessageHandler->OnControllerAnalog(FGamepadKeyNames::RightAnalogX, NewControllerState.DeviceId, NewControllerState.RXAnalog);
+				NewControllerState.ButtonStates[MAX_NUM_PHYSICAL_CONTROLLER_BUTTONS + 5] = NewControllerState.RXAnalog >= RepeatDeadzone;
+				NewControllerState.ButtonStates[MAX_NUM_PHYSICAL_CONTROLLER_BUTTONS + 4] = NewControllerState.RXAnalog <= -RepeatDeadzone;
 			}
 			if (NewControllerState.RYAnalog != OldControllerState.RYAnalog || FMath::Abs(NewControllerState.RYAnalog) >= RepeatDeadzone)
 			{
 				//LOGD("    Sending updated RightAnalogY value of %f", NewControllerState.RYAnalog);
 				MessageHandler->OnControllerAnalog(FGamepadKeyNames::RightAnalogY, NewControllerState.DeviceId, NewControllerState.RYAnalog);
+				NewControllerState.ButtonStates[MAX_NUM_PHYSICAL_CONTROLLER_BUTTONS + 6] = NewControllerState.RYAnalog >= RepeatDeadzone;
+				NewControllerState.ButtonStates[MAX_NUM_PHYSICAL_CONTROLLER_BUTTONS + 7] = NewControllerState.RYAnalog <= -RepeatDeadzone;
 			}
 			if (NewControllerState.LTAnalog != OldControllerState.LTAnalog)
 			{
diff --git a/Engine/Source/Runtime/ApplicationCore/Private/IOS/IOSInputInterface.cpp b/Engine/Source/Runtime/ApplicationCore/Private/IOS/IOSInputInterface.cpp
--- a/Engine/Source/Runtime/ApplicationCore/Private/IOS/IOSInputInterface.cpp
+++ b/Engine/Source/Runtime/ApplicationCore/Private/IOS/IOSInputInterface.cpp
@@ -338,15 +338,15 @@ void FIOSInputInterface::SendControllerEvents()
         const double CurrentTime = FPlatformTime::Seconds();
         const float IniitialRepeatDelay = 0.2f;
         const float RepeatDelay = 0.1;
-		
-#define HANDLE_BUTTON(Gamepad, GCButton, UEButton) \
-if ((Controller.Previous##Gamepad == nil && Gamepad.GCButton.pressed) || Controller.Previous##Gamepad.GCButton.pressed != Gamepad.GCButton.pressed) \
+
+#define HANDLE_BUTTON_INTERNAL(Gamepad, bWasPressed, bIsPressed, UEButton) \
+if (bWasPressed != bIsPressed) \
 { \
-	NSLog(@"%@ button %s on controller %d", (ExtendedGamepad.GCButton.pressed) ? @"Pressed" : @"Released", TCHAR_TO_ANSI(*UEButton.ToString()), (int32)Cont.playerIndex); \
-	(Gamepad.GCButton.pressed) ? MessageHandler->OnControllerButtonPressed(UEButton, Cont.playerIndex, false) : MessageHandler->OnControllerButtonReleased(UEButton, Cont.playerIndex, false); \
+	NSLog(@"%@ button %s on controller %d", bIsPressed ? @"Pressed" : @"Released", TCHAR_TO_ANSI(*UEButton.ToString()), (int32)Cont.playerIndex); \
+	bIsPressed ? MessageHandler->OnControllerButtonPressed(UEButton, Cont.playerIndex, false) : MessageHandler->OnControllerButtonReleased(UEButton, Cont.playerIndex, false); \
     NextKeyRepeatTime.FindOrAdd(UEButton) = CurrentTime + IniitialRepeatDelay; \
 } \
-else if(Gamepad.GCButton.pressed) \
+else if(bIsPressed) \
 { \
     double* NextRepeatTime = NextKeyRepeatTime.Find(UEButton); \
     if(NextRepeatTime && *NextRepeatTime <= CurrentTime) \
@@ -359,6 +359,13 @@ else \
 { \
     NextKeyRepeatTime.Remove(UEButton); \
 }
+		
+#define HANDLE_BUTTON(Gamepad, GCButton, UEButton) \
+{ \
+	const bool bWasPressed = Controller.Previous##Gamepad != nil && Controller.Previous##Gamepad.GCButton.pressed; \
+	const bool bPressed = Gamepad.GCButton.pressed; \
+	HANDLE_BUTTON_INTERNAL(Gamepad, bWasPressed, bPressed, UEButton); \
+}
 
         // Send controller events any time we are passed the given input threshold similarly to PC/Console (see: XInputInterface.cpp)
         const float RepeatDeadzone = 0.24;
@@ -369,6 +376,16 @@ if ((Controller.Previous##Gamepad != nil && Gamepad.GCAxis.value != Controller.P
 	NSLog(@"Axis %s is %f", TCHAR_TO_ANSI(*UEAxis.ToString()), Gamepad.GCAxis.value); \
 	MessageHandler->OnControllerAnalog(UEAxis, Cont.playerIndex, Gamepad.GCAxis.value); \
 }
+
+#define HANDLE_ANALOG_VIRTUAL_BUTTONS(Gamepad, GCAxis, UEButtonNegative, UEButtonPositive) \
+{ \
+	const bool bWasNegativePressed = Controller.Previous##Gamepad != nil && Controller.Previous##Gamepad.GCAxis.value <= -RepeatDeadzone; \
+	const bool bNegativePressed = Gamepad.GCAxis.value <= -RepeatDeadzone; \
+	HANDLE_BUTTON_INTERNAL(Gamepad, bWasNegativePressed, bNegativePressed, UEButtonNegative) \
+	const bool bWasPositivePressed = Controller.Previous##Gamepad != nil && Controller.Previous##Gamepad.GCAxis.value >= RepeatDeadzone; \
+	const bool bPositivePressed = Gamepad.GCAxis.value >= RepeatDeadzone; \
+	HANDLE_BUTTON_INTERNAL(Gamepad, bWasPositivePressed, bPositivePressed, UEButtonPositive) \
+}
 		
 		if (ExtendedGamepad != nil)
 		{
@@ -392,6 +409,11 @@ if ((Controller.Previous##Gamepad != nil && Gamepad.GCAxis.value != Controller.P
 			HANDLE_ANALOG(ExtendedGamepad, leftTrigger,				FGamepadKeyNames::LeftTriggerAnalog);
 			HANDLE_ANALOG(ExtendedGamepad, rightTrigger,			FGamepadKeyNames::RightTriggerAnalog);
 
+			HANDLE_ANALOG_VIRTUAL_BUTTONS(ExtendedGamepad, leftThumbstick.xAxis, FGamepadKeyNames::LeftStickLeft, FGamepadKeyNames::LeftStickRight);
+			HANDLE_ANALOG_VIRTUAL_BUTTONS(ExtendedGamepad, leftThumbstick.yAxis, FGamepadKeyNames::LeftStickDown, FGamepadKeyNames::LeftStickUp);
+			HANDLE_ANALOG_VIRTUAL_BUTTONS(ExtendedGamepad, rightThumbstick.xAxis, FGamepadKeyNames::RightStickLeft, FGamepadKeyNames::RightStickRight);
+			HANDLE_ANALOG_VIRTUAL_BUTTONS(ExtendedGamepad, rightThumbstick.yAxis, FGamepadKeyNames::RightStickDown, FGamepadKeyNames::RightStickUp);
+
 			[Controller.PreviousExtendedGamepad release];
 			Controller.PreviousExtendedGamepad = [ExtendedGamepad saveSnapshot];
 			[Controller.PreviousExtendedGamepad retain];
diff --git a/Engine/Source/Runtime/ApplicationCore/Public/Android/AndroidInputInterface.h b/Engine/Source/Runtime/ApplicationCore/Public/Android/AndroidInputInterface.h
--- a/Engine/Source/Runtime/ApplicationCore/Public/Android/AndroidInputInterface.h
+++ b/Engine/Source/Runtime/ApplicationCore/Public/Android/AndroidInputInterface.h
@@ -153,7 +153,9 @@ struct TouchInput
 };
 
 #define MAX_NUM_CONTROLLERS					8  // reasonable limit for now
-#define MAX_NUM_CONTROLLER_BUTTONS			18
+#define MAX_NUM_PHYSICAL_CONTROLLER_BUTTONS	18
+#define MAX_NUM_VIRTUAL_CONTROLLER_BUTTONS	8
+#define MAX_NUM_CONTROLLER_BUTTONS			MAX_NUM_PHYSICAL_CONTROLLER_BUTTONS + MAX_NUM_VIRTUAL_CONTROLLER_BUTTONS
 #define MAX_DEFERRED_MESSAGE_QUEUE_SIZE		128
 
 struct FAndroidControllerData
