Ios repeat button presses were not being passed as repeat button presses.

diff --git a/Engine/Source/Runtime/ApplicationCore/Private/IOS/IOSInputInterface.cpp b/Engine/Source/Runtime/ApplicationCore/Private/IOS/IOSInputInterface.cpp
--- a/Engine/Source/Runtime/ApplicationCore/Private/IOS/IOSInputInterface.cpp
+++ b/Engine/Source/Runtime/ApplicationCore/Private/IOS/IOSInputInterface.cpp
@@ -351,7 +351,7 @@ else if(bIsPressed) \
     double* NextRepeatTime = NextKeyRepeatTime.Find(UEButton); \
     if(NextRepeatTime && *NextRepeatTime <= CurrentTime) \
     { \
-        MessageHandler->OnControllerButtonPressed(UEButton, Cont.playerIndex, false); \
+        MessageHandler->OnControllerButtonPressed(UEButton, Cont.playerIndex, true); \
 *NextRepeatTime = CurrentTime + RepeatDelay; \
     } \
 } \
