Adding support for Razer controllers in Android input and making trigger threshold in line with other platforms.

diff --git a/Engine/Build/Android/Java/src/com/epicgames/ue4/GameActivity.java.template b/Engine/Build/Android/Java/src/com/epicgames/ue4/GameActivity.java.template
--- a/Engine/Build/Android/Java/src/com/epicgames/ue4/GameActivity.java.template
+++ b/Engine/Build/Android/Java/src/com/epicgames/ue4/GameActivity.java.template
@@ -4563,7 +4563,8 @@ public class GameActivity extends NativeActivity implements SurfaceHolder.Callba
 		new DeviceInfoData(0x054c, 0x05c4, "PS4 Wireless Controller"),
 		new DeviceInfoData(0x054c, 0x09cc, "PS4 Wireless Controller (v2)"),
 		new DeviceInfoData(0x05ac, 0x056a, "glap QXPGP001"),
-		new DeviceInfoData(0x0483, 0x5750, "STMicroelectronics Lenovo GamePad")
+		new DeviceInfoData(0x0483, 0x5750, "STMicroelectronics Lenovo GamePad"),
+		new DeviceInfoData(0x1532, 0x0705, "Razer Razer Claire T1 Wired")
 	};
 
 	public class InputDeviceInfo {
diff --git a/Engine/Source/Runtime/ApplicationCore/Private/Android/AndroidInputInterface.cpp b/Engine/Source/Runtime/ApplicationCore/Private/Android/AndroidInputInterface.cpp
--- a/Engine/Source/Runtime/ApplicationCore/Private/Android/AndroidInputInterface.cpp
+++ b/Engine/Source/Runtime/ApplicationCore/Private/Android/AndroidInputInterface.cpp
@@ -12,6 +12,10 @@
 #include "HAL/IConsoleManager.h"
 #include "IHapticDevice.h"
 
+#ifndef ANDROID_GAMEPAD_TRIGGER_THRESHOLD
+	#define ANDROID_GAMEPAD_TRIGGER_THRESHOLD	0.30f
+#endif
+
 TArray<TouchInput> FAndroidInputInterface::TouchInputStack = TArray<TouchInput>();
 FCriticalSection FAndroidInputInterface::TouchInputCriticalSection;
 
@@ -888,6 +892,10 @@ void FAndroidInputInterface::SendControllerEvents()
 						{
 							CurrentDevice.bSupportsHat = true;
 						}
+						else if (CurrentDevice.DeviceInfo.Name.StartsWith(TEXT("Razer")))
+						{
+							CurrentDevice.bSupportsHat = true;
+						}
 
 						FCoreDelegates::OnControllerConnectionChange.Broadcast(true, -1, DeviceIndex);
 
@@ -996,7 +1004,7 @@ void FAndroidInputInterface::SendControllerEvents()
 
 				// Handle the trigger theshold "virtual" button state
 				//check(ButtonMapping[10] == FGamepadKeyNames::LeftTriggerThreshold);
-				NewControllerState.ButtonStates[10] = NewControllerState.LTAnalog >= 0.1f;
+				NewControllerState.ButtonStates[10] = NewControllerState.LTAnalog >= ANDROID_GAMEPAD_TRIGGER_THRESHOLD;
 			}
 			if (NewControllerState.RTAnalog != OldControllerState.RTAnalog)
 			{
@@ -1005,7 +1013,7 @@ void FAndroidInputInterface::SendControllerEvents()
 
 				// Handle the trigger theshold "virtual" button state
 				//check(ButtonMapping[11] == FGamepadKeyNames::RightTriggerThreshold);
-				NewControllerState.ButtonStates[11] = NewControllerState.RTAnalog >= 0.1f;
+				NewControllerState.ButtonStates[11] = NewControllerState.RTAnalog >= ANDROID_GAMEPAD_TRIGGER_THRESHOLD;
 			}
 
 			const double CurrentTime = FPlatformTime::Seconds();
