UE-64046 iOS Media Player freezes when unplugging headphones from device

#4.2.2
#iOS

Adding a new event - iOS specific - to be handled by the media player when un/plugging the headphones

diff --git a/Engine/Plugins/Media/AvfMedia/Source/AvfMedia/Private/Player/AvfMediaPlayer.cpp b/Engine/Plugins/Media/AvfMedia/Source/AvfMedia/Private/Player/AvfMediaPlayer.cpp
--- a/Engine/Plugins/Media/AvfMedia/Source/AvfMedia/Private/Player/AvfMediaPlayer.cpp
+++ b/Engine/Plugins/Media/AvfMedia/Source/AvfMedia/Private/Player/AvfMediaPlayer.cpp
@@ -268,6 +268,12 @@ void FAvfMediaPlayer::Close()
         WillDeactivateHandle.Reset();
     }
 
+	if (AudioRouteChangedHandle.IsValid())
+	{
+		FCoreDelegates::AudioRouteChangedDelegate.Remove(AudioRouteChangedHandle);
+		AudioRouteChangedHandle.Reset();
+	}
+
 	CurrentTime = 0;
 	MediaUrl = FString();
 	
@@ -498,6 +504,11 @@ bool FAvfMediaPlayer::Open(const FString& Url, const IMediaOptions* /*Options*/)
         WillDeactivateHandle = FCoreDelegates::ApplicationWillDeactivateDelegate.AddRaw(this, &FAvfMediaPlayer::HandleApplicationDeactivate);
     }
 
+	if (!AudioRouteChangedHandle.IsValid())
+	{
+		AudioRouteChangedHandle = FCoreDelegates::AudioRouteChangedDelegate.AddRaw(this, &FAvfMediaPlayer::HandleAudioRouteChanged);
+	}
+	
 	return true;
 }
 
@@ -737,3 +748,18 @@ void FAvfMediaPlayer::HandleApplicationDeactivate()
 		[MediaPlayer pause];
 	}
 }
+
+void FAvfMediaPlayer::HandleAudioRouteChanged(bool InDeviceAvailable)
+{
+	if ((CurrentState == EMediaState::Playing) && MediaPlayer != nil)
+	{
+		if (!InDeviceAvailable)
+		{
+			// restart the media - route it to the active audio device
+			// i.e. when unplugging the headphones
+			[MediaPlayer pause];
+
+			[MediaPlayer play];
+		}
+	}
+}
diff --git a/Engine/Plugins/Media/AvfMedia/Source/AvfMedia/Private/Player/AvfMediaPlayer.h b/Engine/Plugins/Media/AvfMedia/Source/AvfMedia/Private/Player/AvfMediaPlayer.h
--- a/Engine/Plugins/Media/AvfMedia/Source/AvfMedia/Private/Player/AvfMediaPlayer.h
+++ b/Engine/Plugins/Media/AvfMedia/Source/AvfMedia/Private/Player/AvfMediaPlayer.h
@@ -100,6 +100,11 @@ private:
 	/** Callback for when the application is moved from the active to inactive state */
 	void HandleApplicationDeactivate();
 
+	/** Callback for when the audio device is changed 
+	* @param InDeviceAvailable - true if the device is available
+	*/
+	void HandleAudioRouteChanged(bool InDeviceAvailable);
+
 	/** The current playback rate. */
 	float CurrentRate;
 
@@ -156,4 +161,7 @@ private:
     FDelegateHandle EnteredBackgroundHandle;
     FDelegateHandle WillDeactivateHandle;
 
+	/** Delegate for changing the audio device  */
+	FDelegateHandle AudioRouteChangedHandle;
+
 };
diff --git a/Engine/Source/Runtime/ApplicationCore/Private/IOS/IOSAppDelegate.cpp b/Engine/Source/Runtime/ApplicationCore/Private/IOS/IOSAppDelegate.cpp
--- a/Engine/Source/Runtime/ApplicationCore/Private/IOS/IOSAppDelegate.cpp
+++ b/Engine/Source/Runtime/ApplicationCore/Private/IOS/IOSAppDelegate.cpp
@@ -427,6 +427,22 @@ bool FIOSCoreDelegates::PassesPushNotificationFilters(NSDictionary* Payload)
 		}
 	}];
 
+	[[NSNotificationCenter defaultCenter] addObserverForName:AVAudioSessionRouteChangeNotification object : nil queue : nil usingBlock : ^ (NSNotification *notification)
+	{
+		switch ([[[notification userInfo] objectForKey:AVAudioSessionRouteChangeReasonKey] unsignedIntegerValue])
+		{
+			case AVAudioSessionRouteChangeReasonNewDeviceAvailable:
+				// headphones plugged in
+				FCoreDelegates::AudioRouteChangedDelegate.Broadcast(true);
+				break;
+
+			case AVAudioSessionRouteChangeReasonOldDeviceUnavailable:
+				// headphones unplugged
+				FCoreDelegates::AudioRouteChangedDelegate.Broadcast(false);
+				break;
+		}
+	}];
+
 	self.bUsingBackgroundMusic = [self IsBackgroundAudioPlaying];
 	self.bForceEmitOtherAudioPlaying = true;
 
@@ -705,7 +721,7 @@ bool GIsSuspended = 0;
 		// Don't deadlock here because a msg box may appear super early blocking the game thread and then the app may go into the background
 		double	startTime = FPlatformTime::Seconds();
 
-		// don't wait for FDefaultGameMoviePlayer::WaitForMovieToFinish(), crash with 0x8badf00d if â€œWait for Movies to Completeâ€ is checked
+		// don't wait for FDefaultGameMoviePlayer::WaitForMovieToFinish(), crash with 0x8badf00d if “Wait for Movies to Complete” is checked
 		while(!self.bHasSuspended && !FAppEntry::IsStartupMoviePlaying() &&  (FPlatformTime::Seconds() - startTime) < cMaxThreadWaitTime)
 		{
             FIOSPlatformRHIFramePacer::Suspend();
diff --git a/Engine/Source/Runtime/Core/Private/Misc/CoreDelegates.cpp b/Engine/Source/Runtime/Core/Private/Misc/CoreDelegates.cpp
--- a/Engine/Source/Runtime/Core/Private/Misc/CoreDelegates.cpp
+++ b/Engine/Source/Runtime/Core/Private/Misc/CoreDelegates.cpp
@@ -96,6 +96,7 @@ FCoreDelegates::FApplicationLifetimeDelegate FCoreDelegates::ApplicationShouldUn
 FCoreDelegates::FApplicationStartupArgumentsDelegate FCoreDelegates::ApplicationReceivedStartupArgumentsDelegate;
 
 FCoreDelegates::FUserMusicInterruptDelegate FCoreDelegates::UserMusicInterruptDelegate;
+FCoreDelegates::FAudioRouteChangedDelegate FCoreDelegates::AudioRouteChangedDelegate;
 FCoreDelegates::FAudioMuteDelegate FCoreDelegates::AudioMuteDelegate;
 FCoreDelegates::FApplicationRequestAudioState FCoreDelegates::ApplicationRequestAudioState;
 
diff --git a/Engine/Source/Runtime/Core/Public/Misc/CoreDelegates.h b/Engine/Source/Runtime/Core/Public/Misc/CoreDelegates.h
--- a/Engine/Source/Runtime/Core/Public/Misc/CoreDelegates.h
+++ b/Engine/Source/Runtime/Core/Public/Misc/CoreDelegates.h
@@ -330,6 +330,11 @@ public:
 	DECLARE_MULTICAST_DELEGATE_TwoParams(FAudioMuteDelegate, bool, int);
 	static FAudioMuteDelegate AudioMuteDelegate;
 	
+	// [iOS only] Called when the audio device changes
+	// For instance, when the headphones are plugged in or removed
+	DECLARE_MULTICAST_DELEGATE_OneParam(FAudioRouteChangedDelegate, bool);
+	static FAudioRouteChangedDelegate AudioRouteChangedDelegate;
+
 	// Generally, events triggering UserMusicInterruptDelegate or AudioMuteDelegate happen only
 	// when a change occurs. When a system comes online needing the current audio state but the
 	// event has already been broadcast, calling ApplicationRequestAudioState will force the
