Fixed a Mac crash caused by an attempt to access info about a disconnected display before the display reconfiguration event was received.

diff --git a/Engine/Source/Runtime/ApplicationCore/Private/Mac/MacApplication.cpp b/Engine/Source/Runtime/ApplicationCore/Private/Mac/MacApplication.cpp
--- a/Engine/Source/Runtime/ApplicationCore/Private/Mac/MacApplication.cpp
+++ b/Engine/Source/Runtime/ApplicationCore/Private/Mac/MacApplication.cpp
@@ -2027,77 +2027,81 @@ void FDisplayMetrics::RebuildDisplayMetrics(FDisplayMetrics& OutDisplayMetrics)
 		Info.ID = FString::Printf(TEXT("%u"), DisplayID);
 
 		CFArrayRef ArrDisplay = CGDisplayCopyAllDisplayModes(DisplayID, nullptr);
-
-		Info.NativeWidth = 0;
-		Info.NativeHeight = 0;
-		const CFIndex AppsCount = CFArrayGetCount(ArrDisplay);
-		for (CFIndex i = 0; i < AppsCount; ++i)
+		if (ArrDisplay)
 		{
-			const CGDisplayModeRef Mode = (const CGDisplayModeRef)CFArrayGetValueAtIndex(ArrDisplay, i);
-			const int32 Width = (int32)CGDisplayModeGetWidth(Mode);
-			const int32 Height = (int32)CGDisplayModeGetHeight(Mode);
-
-			if (Width * Height > Info.NativeWidth * Info.NativeHeight)
+			Info.NativeWidth = 0;
+			Info.NativeHeight = 0;
+			const CFIndex AppsCount = CFArrayGetCount(ArrDisplay);
+			for (CFIndex i = 0; i < AppsCount; ++i)
 			{
-				Info.NativeWidth = Width;
-				Info.NativeHeight = Height;
-			}
-		}
+				const CGDisplayModeRef Mode = (const CGDisplayModeRef)CFArrayGetValueAtIndex(ArrDisplay, i);
+				const int32 Width = (int32)CGDisplayModeGetWidth(Mode);
+				const int32 Height = (int32)CGDisplayModeGetHeight(Mode);
 
-		if (!Info.NativeWidth || !Info.NativeHeight)
-		{
-			Info.NativeWidth = CGDisplayPixelsWide(DisplayID);
-			Info.NativeHeight = CGDisplayPixelsHigh(DisplayID);
-		}
-
-		Info.DisplayRect = FPlatformRect(Screen->FramePixels.origin.x, Screen->FramePixels.origin.y, Screen->FramePixels.origin.x + Screen->FramePixels.size.width, Screen->FramePixels.origin.y + Screen->FramePixels.size.height);
-		Info.WorkArea = FPlatformRect(Screen->VisibleFramePixels.origin.x, Screen->VisibleFramePixels.origin.y, Screen->VisibleFramePixels.origin.x + Screen->VisibleFramePixels.size.width, Screen->VisibleFramePixels.origin.y + Screen->VisibleFramePixels.size.height);
-		Info.bIsPrimary = Screen->Screen == [NSScreen mainScreen];
+				if (Width * Height > Info.NativeWidth * Info.NativeHeight)
+				{
+					Info.NativeWidth = Width;
+					Info.NativeHeight = Height;
+				}
+			}
 
-		// dpi computations
-		const CGSize DisplayPhysicalSize = CGDisplayScreenSize(DisplayID);
-		const float MilimetreInch = 25.4f;
-		float HorizontalDPI = MilimetreInch * (float)Info.NativeWidth / (float)DisplayPhysicalSize.width;
-		float VerticalDPI = MilimetreInch * (float)Info.NativeHeight / (float)DisplayPhysicalSize.height;
-		Info.DPI = FMath::CeilToInt((HorizontalDPI + VerticalDPI) / 2.0f);
+			if (!Info.NativeWidth || !Info.NativeHeight)
+			{
+				Info.NativeWidth = CGDisplayPixelsWide(DisplayID);
+				Info.NativeHeight = CGDisplayPixelsHigh(DisplayID);
+			}
 
-		// Monitor's name can only be obtained from IOKit
-		io_iterator_t IOIterator;
-		kern_return_t Result = IOServiceGetMatchingServices(kIOMasterPortDefault, IOServiceMatching("IODisplayConnect"), &IOIterator);
-		if (Result == kIOReturnSuccess)
-		{
-			io_object_t Device;
-			while ((Device = IOIteratorNext(IOIterator)))
+			Info.DisplayRect = FPlatformRect(Screen->FramePixels.origin.x, Screen->FramePixels.origin.y, Screen->FramePixels.origin.x + Screen->FramePixels.size.width, Screen->FramePixels.origin.y + Screen->FramePixels.size.height);
+			Info.WorkArea = FPlatformRect(Screen->VisibleFramePixels.origin.x, Screen->VisibleFramePixels.origin.y, Screen->VisibleFramePixels.origin.x + Screen->VisibleFramePixels.size.width, Screen->VisibleFramePixels.origin.y + Screen->VisibleFramePixels.size.height);
+			Info.bIsPrimary = Screen->Screen == [NSScreen mainScreen];
+
+			// dpi computations
+			const CGSize DisplayPhysicalSize = CGDisplayScreenSize(DisplayID);
+			const float MilimetreInch = 25.4f;
+			float HorizontalDPI = MilimetreInch * (float)Info.NativeWidth / (float)DisplayPhysicalSize.width;
+			float VerticalDPI = MilimetreInch * (float)Info.NativeHeight / (float)DisplayPhysicalSize.height;
+			Info.DPI = FMath::CeilToInt((HorizontalDPI + VerticalDPI) / 2.0f);
+
+			// Monitor's name can only be obtained from IOKit
+			io_iterator_t IOIterator;
+			kern_return_t Result = IOServiceGetMatchingServices(kIOMasterPortDefault, IOServiceMatching("IODisplayConnect"), &IOIterator);
+			if (Result == kIOReturnSuccess)
 			{
-				CFDictionaryRef Dictionary = IODisplayCreateInfoDictionary(Device, kIODisplayOnlyPreferredName);
-				if (Dictionary)
+				io_object_t Device;
+				while ((Device = IOIteratorNext(IOIterator)))
 				{
-					const uint32 VendorID = [(__bridge NSNumber*)CFDictionaryGetValue(Dictionary, CFSTR(kDisplayVendorID)) unsignedIntegerValue];
-					const uint32 ProductID = [(__bridge NSNumber*)CFDictionaryGetValue(Dictionary, CFSTR(kDisplayProductID)) unsignedIntegerValue];
-					const uint32 SerialNumber = [(__bridge NSNumber*)CFDictionaryGetValue(Dictionary, CFSTR(kDisplaySerialNumber)) unsignedIntegerValue];
-
-					if (VendorID == CGDisplayVendorNumber(DisplayID) && ProductID == CGDisplayModelNumber(DisplayID) && SerialNumber == CGDisplaySerialNumber(DisplayID))
+					CFDictionaryRef Dictionary = IODisplayCreateInfoDictionary(Device, kIODisplayOnlyPreferredName);
+					if (Dictionary)
 					{
-						NSDictionary* NamesDictionary = (__bridge NSDictionary*)CFDictionaryGetValue(Dictionary, CFSTR(kDisplayProductName));
-						if (NamesDictionary && NamesDictionary.count > 0)
+						const uint32 VendorID = [(__bridge NSNumber*)CFDictionaryGetValue(Dictionary, CFSTR(kDisplayVendorID)) unsignedIntegerValue];
+						const uint32 ProductID = [(__bridge NSNumber*)CFDictionaryGetValue(Dictionary, CFSTR(kDisplayProductID)) unsignedIntegerValue];
+						const uint32 SerialNumber = [(__bridge NSNumber*)CFDictionaryGetValue(Dictionary, CFSTR(kDisplaySerialNumber)) unsignedIntegerValue];
+
+						if (VendorID == CGDisplayVendorNumber(DisplayID) && ProductID == CGDisplayModelNumber(DisplayID) && SerialNumber == CGDisplaySerialNumber(DisplayID))
 						{
-							Info.Name = (NSString*)[NamesDictionary objectForKey:[NamesDictionary.allKeys objectAtIndex:0]];
-							CFRelease(Dictionary);
-							IOObjectRelease(Device);
-							break;
+							NSDictionary* NamesDictionary = (__bridge NSDictionary*)CFDictionaryGetValue(Dictionary, CFSTR(kDisplayProductName));
+							if (NamesDictionary && NamesDictionary.count > 0)
+							{
+								Info.Name = (NSString*)[NamesDictionary objectForKey:[NamesDictionary.allKeys objectAtIndex:0]];
+								CFRelease(Dictionary);
+								IOObjectRelease(Device);
+								break;
+							}
 						}
+
+						CFRelease(Dictionary);
 					}
 
-					CFRelease(Dictionary);
+					IOObjectRelease(Device);
 				}
 
-				IOObjectRelease(Device);
+				IOObjectRelease(IOIterator);
 			}
 
-			IOObjectRelease(IOIterator);
-		}
+			OutDisplayMetrics.MonitorInfo.Add(Info);
 
-		OutDisplayMetrics.MonitorInfo.Add(Info);
+			CFRelease(ArrDisplay);
+		}
 	}
 
 	// Virtual desktop area
