Hide Mac native window buttons when the editor's viewport is in immersive mode

diff --git a/Engine/Source/Editor/LevelEditor/Private/LevelViewportLayout.cpp b/Engine/Source/Editor/LevelEditor/Private/LevelViewportLayout.cpp
--- a/Engine/Source/Editor/LevelEditor/Private/LevelViewportLayout.cpp
+++ b/Engine/Source/Editor/LevelEditor/Private/LevelViewportLayout.cpp
@@ -415,6 +415,7 @@ void FLevelViewportLayout::MaximizeViewport( FName ViewportToMaximize, const boo
 			if( bWasImmersive && !bIsImmersive )
 			{
 				OwnerWindow->BeginFullWindowOverlayTransition();
+				OwnerWindow->SetNativeWindowButtonsVisibility(true);
 			}
 		}
 		else
@@ -653,6 +654,7 @@ void FLevelViewportLayout::FinishMaximizeTransition()
 			TSharedPtr< SWindow > OwnerWindow( CachedOwnerWindow.Pin() );
 			if( OwnerWindow.IsValid() )
 			{
+				OwnerWindow->SetNativeWindowButtonsVisibility(false);
 				OwnerWindow->EndFullWindowOverlayTransition();
 			}
 
diff --git a/Engine/Source/Runtime/ApplicationCore/Private/GenericPlatform/GenericWindow.cpp b/Engine/Source/Runtime/ApplicationCore/Private/GenericPlatform/GenericWindow.cpp
--- a/Engine/Source/Runtime/ApplicationCore/Private/GenericPlatform/GenericWindow.cpp
+++ b/Engine/Source/Runtime/ApplicationCore/Private/GenericPlatform/GenericWindow.cpp
@@ -195,3 +195,8 @@ void FGenericWindow::DrawAttention(const FWindowDrawAttentionParameters& Paramet
 {
 	// empty default functionality
 }
+
+void FGenericWindow::SetNativeWindowButtonsVisibility(bool bVisible)
+{
+	// empty default functionality
+}
diff --git a/Engine/Source/Runtime/ApplicationCore/Private/Mac/MacWindow.cpp b/Engine/Source/Runtime/ApplicationCore/Private/Mac/MacWindow.cpp
--- a/Engine/Source/Runtime/ApplicationCore/Private/Mac/MacWindow.cpp
+++ b/Engine/Source/Runtime/ApplicationCore/Private/Mac/MacWindow.cpp
@@ -520,6 +520,32 @@ float FMacWindow::GetDPIScaleFactor() const
 	return FPlatformApplicationMisc::IsHighDPIModeEnabled() ? WindowHandle.backingScaleFactor : 1.0f;
 }
 
+void FMacWindow::SetNativeWindowButtonsVisibility(bool bVisible)
+{
+	const bool bHidden = !(bVisible && Definition->IsRegularWindow);
+	MainThreadCall(^{
+		SCOPED_AUTORELEASE_POOL;
+
+		NSButton* CloseButton = [WindowHandle standardWindowButton:NSWindowCloseButton];
+		if (CloseButton)
+		{
+			CloseButton.hidden = bHidden;
+		}
+
+		NSButton* MinimizeButton = [WindowHandle standardWindowButton:NSWindowMiniaturizeButton];
+		if (MinimizeButton)
+		{
+			MinimizeButton.hidden = bHidden;
+		}
+
+		NSButton* MaximizeButton = [WindowHandle standardWindowButton:NSWindowZoomButton];
+		if (MaximizeButton)
+		{
+			MaximizeButton.hidden = bHidden;
+		}
+	}, NSDefaultRunLoopMode, false);
+}
+
 void FMacWindow::OnDisplayReconfiguration(CGDirectDisplayID Display, CGDisplayChangeSummaryFlags Flags)
 {
 	if(WindowHandle)
diff --git a/Engine/Source/Runtime/ApplicationCore/Public/GenericPlatform/GenericWindow.h b/Engine/Source/Runtime/ApplicationCore/Public/GenericPlatform/GenericWindow.h
--- a/Engine/Source/Runtime/ApplicationCore/Public/GenericPlatform/GenericWindow.h
+++ b/Engine/Source/Runtime/ApplicationCore/Public/GenericPlatform/GenericWindow.h
@@ -227,7 +227,10 @@ public:
 	 *        the platform, not all parameters may be supported.
 	 */
 	virtual void DrawAttention(const FWindowDrawAttentionParameters& Parameters);
-	
+
+	/** Shows or hides native window buttons on platforms that use them */
+	virtual void SetNativeWindowButtonsVisibility(bool bVisible);
+
 protected:
 
 	TSharedPtr< FGenericWindowDefinition > Definition;
diff --git a/Engine/Source/Runtime/ApplicationCore/Public/Mac/MacWindow.h b/Engine/Source/Runtime/ApplicationCore/Public/Mac/MacWindow.h
--- a/Engine/Source/Runtime/ApplicationCore/Public/Mac/MacWindow.h
+++ b/Engine/Source/Runtime/ApplicationCore/Public/Mac/MacWindow.h
@@ -75,6 +75,8 @@ public:
 
 	virtual float GetDPIScaleFactor() const override;
 
+	virtual void SetNativeWindowButtonsVisibility(bool bVisible);
+
 	bool IsRegularWindow() const;
 
 	int32 PositionX;
diff --git a/Engine/Source/Runtime/SlateCore/Private/Widgets/SWindow.cpp b/Engine/Source/Runtime/SlateCore/Private/Widgets/SWindow.cpp
--- a/Engine/Source/Runtime/SlateCore/Private/Widgets/SWindow.cpp
+++ b/Engine/Source/Runtime/SlateCore/Private/Widgets/SWindow.cpp
@@ -2026,6 +2026,14 @@ void SWindow::EndFullWindowOverlayTransition()
 	bShouldShowWindowContentDuringOverlay = false;
 }
 
+void SWindow::SetNativeWindowButtonsVisibility(bool bVisible)
+{
+	if (NativeWindow.IsValid())
+	{
+		NativeWindow->SetNativeWindowButtonsVisibility(bVisible);
+	}
+}
+
 EVisibility SWindow::GetWindowContentVisibility() const
 {
 	// The content of the window should be visible unless we have a full window overlay content
diff --git a/Engine/Source/Runtime/SlateCore/Public/Widgets/SWindow.h b/Engine/Source/Runtime/SlateCore/Public/Widgets/SWindow.h
--- a/Engine/Source/Runtime/SlateCore/Public/Widgets/SWindow.h
+++ b/Engine/Source/Runtime/SlateCore/Public/Widgets/SWindow.h
@@ -565,6 +565,9 @@ public:
 	 */
 	bool HasFullWindowOverlayContent() const;
 
+	/** Shows or hides native window buttons on platforms that use them */
+	void SetNativeWindowButtonsVisibility(bool bVisible);
+
 	/** @return should this window show up in the taskbar */
 	bool AppearsInTaskbar() const;
 
