Do not forcefully raise windows when all windows should be above the main UE4 window. This causes focuses stealing with out applications, and unwanted focus lost

Also avoid Setting keyboard grabs on windows that are no longer override redirect

diff --git a/Engine/Source/Runtime/ApplicationCore/Private/Linux/LinuxApplication.cpp b/Engine/Source/Runtime/ApplicationCore/Private/Linux/LinuxApplication.cpp
--- a/Engine/Source/Runtime/ApplicationCore/Private/Linux/LinuxApplication.cpp
+++ b/Engine/Source/Runtime/ApplicationCore/Private/Linux/LinuxApplication.cpp
@@ -364,6 +364,7 @@ void FLinuxApplication::ProcessDeferredMessage( SDL_Event Event )
 			}
 			else
 			{
+				UE_LOG(LogLinuxWindowEvent, Verbose, TEXT("MOUSE_DOWN                                 : %d"), CurrentEventWindow->GetID());
 				// User clicked any button. Is the application active? If not activate it.
 				if (!bActivateApp)
 				{
@@ -395,15 +396,7 @@ void FLinuxApplication::ProcessDeferredMessage( SDL_Event Event )
 					// Check if we have to set the focus.
 					if(CurrentFocusWindow != CurrentEventWindow)
 					{
-						SDL_RaiseWindow(CurrentEventWindow->GetHWnd());
-						if(CurrentEventWindow->IsPopupMenuWindow())
-						{
-							SDL_SetKeyboardGrab(CurrentEventWindow->GetHWnd(), SDL_TRUE);
-						}
-						else
-						{
-							SDL_SetWindowInputFocus(CurrentEventWindow->GetHWnd());
-						}
+						SDL_SetWindowInputFocus(CurrentEventWindow->GetHWnd());
 					}
 
 					MessageHandler->OnMouseDown(CurrentEventWindow, button);
@@ -734,16 +727,7 @@ void FLinuxApplication::ProcessDeferredMessage( SDL_Event Event )
 						// Set focus if the window wants to have a focus when first shown.
 						if (CurrentEventWindow->IsFocusWhenFirstShown())
 						{
-							if (CurrentEventWindow->IsPopupMenuWindow())
-							{
-								// We use grab here because this seems to be a proper way to set focus to an override-redirect window.
-								// This prevents additional window changed highlighting in some WMs.
-								SDL_SetKeyboardGrab(CurrentEventWindow->GetHWnd(), SDL_TRUE);
-							}
-							else
-							{
-								SDL_SetWindowInputFocus(CurrentEventWindow->GetHWnd());
-							}
+							SDL_SetWindowInputFocus(CurrentEventWindow->GetHWnd());
 						}
 					}
 					break;
@@ -1567,8 +1551,6 @@ void FLinuxApplication::RemoveRevertFocusWindow(SDL_HWindow HWnd)
 			if(Window->IsUtilityWindow() || Window->IsDialogWindow())
 			{
 				ActivateWindow(Window->GetParent());
-
-				SDL_RaiseWindow(Window->GetParent()->GetHWnd() );
 				SDL_SetWindowInputFocus(Window->GetParent()->GetHWnd() );
 			}
 			// Was the deleted window a Blueprint, Cascade, Matinee etc. window?
@@ -1591,8 +1573,6 @@ void FLinuxApplication::RemoveRevertFocusWindow(SDL_HWindow HWnd)
 				}
 
 				ActivateWindow(RevertFocusToWindow);
-
-				SDL_RaiseWindow(RevertFocusToWindow->GetHWnd());
 				SDL_SetWindowInputFocus(RevertFocusToWindow->GetHWnd());
 			}
 			// Was the deleted window a top level window and we have still at least one other window in the stack?
@@ -1603,8 +1583,6 @@ void FLinuxApplication::RemoveRevertFocusWindow(SDL_HWindow HWnd)
 				if (TopmostWindow.IsValid())
 				{
 					ActivateWindow(TopmostWindow);
-
-					SDL_RaiseWindow(TopmostWindow->GetHWnd());
 					SDL_SetWindowInputFocus(TopmostWindow->GetHWnd());
 				}
 			}
@@ -1612,16 +1590,8 @@ void FLinuxApplication::RemoveRevertFocusWindow(SDL_HWindow HWnd)
 			else if (Window->IsPopupMenuWindow() && bActivateApp)
 			{
 				ActivateWindow(Window->GetParent());
+				SDL_SetWindowInputFocus(Window->GetParent()->GetHWnd());
 
-				SDL_RaiseWindow(Window->GetParent()->GetHWnd());
-				if(Window->GetParent()->IsPopupMenuWindow())
-				{
-					SDL_SetKeyboardGrab(Window->GetParent()->GetHWnd(), SDL_TRUE);
-				}
-				else
-				{
-					SDL_SetWindowInputFocus(Window->GetParent()->GetHWnd());
-				}
 				UE_LOG(LogLinuxWindowType, Verbose, TEXT("FLinuxWindow::Destroy: Going to revert focus to %d"), Window->GetParent()->GetID());
 			}
 			break;
