Dont force focus

This change was suppose to be part of CL 4770158

diff --git a/Engine/Source/Runtime/ApplicationCore/Private/Linux/LinuxApplication.cpp b/Engine/Source/Runtime/ApplicationCore/Private/Linux/LinuxApplication.cpp
--- a/Engine/Source/Runtime/ApplicationCore/Private/Linux/LinuxApplication.cpp
+++ b/Engine/Source/Runtime/ApplicationCore/Private/Linux/LinuxApplication.cpp
@@ -1551,7 +1551,6 @@ void FLinuxApplication::RemoveRevertFocusWindow(SDL_HWindow HWnd)
 			if(Window->IsUtilityWindow() || Window->IsDialogWindow())
 			{
 				ActivateWindow(Window->GetParent());
-				SDL_SetWindowInputFocus(Window->GetParent()->GetHWnd() );
 			}
 			// Was the deleted window a Blueprint, Cascade, Matinee etc. window?
 			else if (Window->IsNotificationWindow())
@@ -1573,7 +1572,6 @@ void FLinuxApplication::RemoveRevertFocusWindow(SDL_HWindow HWnd)
 				}
 
 				ActivateWindow(RevertFocusToWindow);
-				SDL_SetWindowInputFocus(RevertFocusToWindow->GetHWnd());
 			}
 			// Was the deleted window a top level window and we have still at least one other window in the stack?
 			else if (Window->IsTopLevelWindow() && (RevertFocusStack.Num() > 0))
@@ -1583,14 +1581,12 @@ void FLinuxApplication::RemoveRevertFocusWindow(SDL_HWindow HWnd)
 				if (TopmostWindow.IsValid())
 				{
 					ActivateWindow(TopmostWindow);
-					SDL_SetWindowInputFocus(TopmostWindow->GetHWnd());
 				}
 			}
 			// Was it a popup menu?
 			else if (Window->IsPopupMenuWindow() && bActivateApp)
 			{
 				ActivateWindow(Window->GetParent());
-				SDL_SetWindowInputFocus(Window->GetParent()->GetHWnd());
 
 				UE_LOG(LogLinuxWindowType, Verbose, TEXT("FLinuxWindow::Destroy: Going to revert focus to %d"), Window->GetParent()->GetID());
 			}
