Was leaking the CGDisplayModes list.

diff --git a/Engine/Source/Runtime/ApplicationCore/Private/Mac/MacApplication.cpp b/Engine/Source/Runtime/ApplicationCore/Private/Mac/MacApplication.cpp
--- a/Engine/Source/Runtime/ApplicationCore/Private/Mac/MacApplication.cpp
+++ b/Engine/Source/Runtime/ApplicationCore/Private/Mac/MacApplication.cpp
@@ -2043,6 +2043,8 @@ void FDisplayMetrics::RebuildDisplayMetrics(FDisplayMetrics& OutDisplayMetrics)
 				Info.NativeHeight = Height;
 			}
 		}
+		
+		CFRelease(ArrDisplay);
 
 		if (!Info.NativeWidth || !Info.NativeHeight)
 		{
