Fix mismerge that caused us to over-release a CFArray.

diff --git a/Engine/Source/Runtime/ApplicationCore/Private/Mac/MacApplication.cpp b/Engine/Source/Runtime/ApplicationCore/Private/Mac/MacApplication.cpp
--- a/Engine/Source/Runtime/ApplicationCore/Private/Mac/MacApplication.cpp
+++ b/Engine/Source/Runtime/ApplicationCore/Private/Mac/MacApplication.cpp
@@ -2051,7 +2051,7 @@ void FDisplayMetrics::RebuildDisplayMetrics(FDisplayMetrics& OutDisplayMetrics)
 				Info.NativeHeight = CGDisplayPixelsHigh(DisplayID);
 			}
 		
-		CFRelease(ArrDisplay);
+			CFRelease(ArrDisplay);
 
 			Info.DisplayRect = FPlatformRect(Screen->FramePixels.origin.x, Screen->FramePixels.origin.y, Screen->FramePixels.origin.x + Screen->FramePixels.size.width, Screen->FramePixels.origin.y + Screen->FramePixels.size.height);
 			Info.WorkArea = FPlatformRect(Screen->VisibleFramePixels.origin.x, Screen->VisibleFramePixels.origin.y, Screen->VisibleFramePixels.origin.x + Screen->VisibleFramePixels.size.width, Screen->VisibleFramePixels.origin.y + Screen->VisibleFramePixels.size.height);
@@ -2101,8 +2101,6 @@ void FDisplayMetrics::RebuildDisplayMetrics(FDisplayMetrics& OutDisplayMetrics)
 			}
 
 			OutDisplayMetrics.MonitorInfo.Add(Info);
-
-			CFRelease(ArrDisplay);
 		}
 	}
 
