Defaulting controllers to use hat d-pads.

diff --git a/Engine/Source/Runtime/ApplicationCore/Private/Android/AndroidInputInterface.cpp b/Engine/Source/Runtime/ApplicationCore/Private/Android/AndroidInputInterface.cpp
--- a/Engine/Source/Runtime/ApplicationCore/Private/Android/AndroidInputInterface.cpp
+++ b/Engine/Source/Runtime/ApplicationCore/Private/Android/AndroidInputInterface.cpp
@@ -812,7 +812,7 @@ void FAndroidInputInterface::SendControllerEvents()
 						CurrentDevice.ButtonRemapping = ButtonRemapType::Normal;
 						CurrentDevice.LTAnalogRangeMinimum = 0.0f;
 						CurrentDevice.RTAnalogRangeMinimum = 0.0f;
-						CurrentDevice.bSupportsHat = false;
+						CurrentDevice.bSupportsHat = true;
 						CurrentDevice.bMapL1R1ToTriggers = false;
 						CurrentDevice.bMapZRZToTriggers = false;
 						CurrentDevice.bRightStickZRZ = true;
@@ -827,9 +827,11 @@ void FAndroidInputInterface::SendControllerEvents()
 							}
 							else if (CurrentDevice.DeviceInfo.Name.StartsWith(TEXT("Amazon Fire TV Remote")))
 							{
+								CurrentDevice.bSupportsHat = false;
 							}
 							else
 							{
+								CurrentDevice.bSupportsHat = false;
 							}
 						}
 						else if (CurrentDevice.DeviceInfo.Name.StartsWith(TEXT("NVIDIA Corporation NVIDIA Controller")))
