Move back to using the Popup_Menu window type as we are no longer grabbing the mouse/keyboard here

Make sure if we are a popup, even if we dont take focus we will be taking activation away from the main editor. Put our selfs in the revert list to revert that activation

diff --git a/Engine/Source/Runtime/ApplicationCore/Private/Linux/LinuxApplication.cpp b/Engine/Source/Runtime/ApplicationCore/Private/Linux/LinuxApplication.cpp
--- a/Engine/Source/Runtime/ApplicationCore/Private/Linux/LinuxApplication.cpp
+++ b/Engine/Source/Runtime/ApplicationCore/Private/Linux/LinuxApplication.cpp
@@ -134,8 +134,8 @@ void FLinuxApplication::InitializeWindow(	const TSharedRef< FGenericWindow >& In
 	Window->Initialize( this, InDefinition, ParentWindow, bShowImmediately );
 	Windows.Add(Window);
 
-	// Add the windows into the focus stack.
-	if (Window->IsFocusWhenFirstShown())
+	// Add the windows that should revert activation when closed
+	if (Window->IsFocusWhenFirstShown() || Window->IsPopupMenuWindow())
 	{
 		RevertFocusStack.Add(Window);
 	}
diff --git a/Engine/Source/Runtime/ApplicationCore/Private/Linux/LinuxWindow.cpp b/Engine/Source/Runtime/ApplicationCore/Private/Linux/LinuxWindow.cpp
--- a/Engine/Source/Runtime/ApplicationCore/Private/Linux/LinuxWindow.cpp
+++ b/Engine/Source/Runtime/ApplicationCore/Private/Linux/LinuxWindow.cpp
@@ -171,8 +171,7 @@ void FLinuxWindow::Initialize( FLinuxApplication* const Application, const TShar
 		!Definition->IsModalWindow && !Definition->IsRegularWindow &&
 		bShouldActivate && !Definition->SizeWillChangeOften)
 	{
-		// Popup menus grab the mouse/keyboard which is undesired behaviour. Slate will give the window events.
-		WindowStyle |= SDL_WINDOW_BORDERLESS;
+		WindowStyle |= SDL_WINDOW_POPUP_MENU;
 		bIsPopupWindow = true;
 		UE_LOG(LogLinuxWindowType, Verbose, TEXT("*** New Window is a Popup Menu Window ***"));
 	}
@@ -183,7 +182,7 @@ void FLinuxWindow::Initialize( FLinuxApplication* const Application, const TShar
 		!Definition->IsModalWindow && !Definition->IsRegularWindow &&
 		!bShouldActivate && !Definition->SizeWillChangeOften)
 	{
-		WindowStyle |= SDL_WINDOW_BORDERLESS;
+		WindowStyle |= SDL_WINDOW_POPUP_MENU;
 		bIsConsoleWindow = true;
 		bIsPopupWindow = true;
 		UE_LOG(LogLinuxWindowType, Verbose, TEXT("*** New Window is a Console Window ***"));
