Added a command line option to override the cursor type on Mac for debugging (MacCursorTypeOverride)

diff --git a/Engine/Source/Runtime/ApplicationCore/Private/Mac/MacCursor.cpp b/Engine/Source/Runtime/ApplicationCore/Private/Mac/MacCursor.cpp
--- a/Engine/Source/Runtime/ApplicationCore/Private/Mac/MacCursor.cpp
+++ b/Engine/Source/Runtime/ApplicationCore/Private/Mac/MacCursor.cpp
@@ -7,6 +7,7 @@
 #include "HAL/IConsoleManager.h"
 #include "HAL/PlatformProcess.h"
 #include "Misc/Paths.h"
+#include "Misc/CommandLine.h"
 
 #import <IOKit/hid/IOHIDKeys.h>
 #import <IOKit/hidsystem/IOHIDShared.h>
@@ -26,6 +27,7 @@ static FAutoConsoleVariableRef CVarMacDisableMouseAcceleration(
 FMacCursor::FMacCursor()
 :	bIsVisible(true)
 ,	bUseHighPrecisionMode(false)
+,	CursorTypeOverride(-1)
 ,	CurrentPosition(FVector2D::ZeroVector)
 ,	MouseWarpDelta(FVector2D::ZeroVector)
 ,	bIsPositionInitialised(false)
@@ -129,6 +131,14 @@ FMacCursor::FMacCursor()
 		CursorHandles[CursorIndex] = CursorHandle;
 	}
 
+#if !UE_BUILD_SHIPPING
+	FParse::Value(FCommandLine::Get(), TEXT("MacCursorTypeOverride="), CursorTypeOverride);
+	if (CursorTypeOverride < EMouseCursor::None || CursorTypeOverride >= EMouseCursor::TotalCursorCount)
+	{
+		CursorTypeOverride = -1;
+	}
+#endif
+
 	// Set the default cursor
 	SetType(EMouseCursor::Default);
 
@@ -250,8 +260,8 @@ void FMacCursor::SetType(const EMouseCursor::Type InNewCursor)
 		MacApplication->SetHighPrecisionMouseMode(false, nullptr);
 	}
 	
-	CurrentType = InNewCursor;
-	CurrentCursor = CursorOverrideHandles[InNewCursor] ? CursorOverrideHandles[InNewCursor] : CursorHandles[InNewCursor];
+	CurrentType = CursorTypeOverride > -1 ? (EMouseCursor::Type)CursorTypeOverride : InNewCursor;
+	CurrentCursor = CursorOverrideHandles[CurrentType] ? CursorOverrideHandles[CurrentType] : CursorHandles[CurrentType];
 
 	if (CurrentCursor)
 	{
diff --git a/Engine/Source/Runtime/ApplicationCore/Public/Mac/MacCursor.h b/Engine/Source/Runtime/ApplicationCore/Public/Mac/MacCursor.h
--- a/Engine/Source/Runtime/ApplicationCore/Public/Mac/MacCursor.h
+++ b/Engine/Source/Runtime/ApplicationCore/Public/Mac/MacCursor.h
@@ -61,6 +61,7 @@ private:
 	bool bIsVisible;
 	bool bUseHighPrecisionMode;
 	NSCursor* CurrentCursor;
+	int32 CursorTypeOverride;
 
 	FVector2D CurrentPosition;
 	FVector2D MouseWarpDelta;
