Fixed the way we calculate screen position in pixels on Mac to properly handle a case where more than two monitors are attached and only the one most on the right is low DPI

diff --git a/Engine/Source/Runtime/ApplicationCore/Private/Mac/MacApplication.cpp b/Engine/Source/Runtime/ApplicationCore/Private/Mac/MacApplication.cpp
--- a/Engine/Source/Runtime/ApplicationCore/Private/Mac/MacApplication.cpp
+++ b/Engine/Source/Runtime/ApplicationCore/Private/Mac/MacApplication.cpp
@@ -1580,27 +1580,13 @@ void FMacApplication::UpdateScreensArray()
 		WholeWorkspace = NSUnionRect(WholeWorkspace, CurScreen->Frame);
 	}
 
-	for (TSharedRef<FMacScreen> CurScreen : AllScreens)
-	{
-		CurScreen->Frame.origin.y = WholeWorkspace.origin.y + WholeWorkspace.size.height - CurScreen->Frame.size.height - CurScreen->Frame.origin.y;
-		CurScreen->VisibleFrame.origin.y = WholeWorkspace.origin.y + WholeWorkspace.size.height - CurScreen->VisibleFrame.size.height - CurScreen->VisibleFrame.origin.y;
-	}
-
 	const bool bUseHighDPIMode = FPlatformApplicationMisc::IsHighDPIModeEnabled();
 
 	TArray<TSharedRef<FMacScreen>> SortedScreens;
-
 	for (TSharedRef<FMacScreen> CurScreen : AllScreens)
 	{
-		const float DPIScaleFactor = bUseHighDPIMode ? CurScreen->Screen.backingScaleFactor : 1.0f;
-		CurScreen->FramePixels.origin.x = CurScreen->Frame.origin.x;
-		CurScreen->FramePixels.origin.y = CurScreen->Frame.origin.y;
-		CurScreen->FramePixels.size.width = CurScreen->Frame.size.width * DPIScaleFactor;
-		CurScreen->FramePixels.size.height = CurScreen->Frame.size.height * DPIScaleFactor;
-		CurScreen->VisibleFramePixels.origin.x = CurScreen->Frame.origin.x + (CurScreen->VisibleFrame.origin.x - CurScreen->Frame.origin.x) * DPIScaleFactor;
-		CurScreen->VisibleFramePixels.origin.y = CurScreen->Frame.origin.y + (CurScreen->VisibleFrame.origin.y - CurScreen->Frame.origin.y) * DPIScaleFactor;
-		CurScreen->VisibleFramePixels.size.width = CurScreen->VisibleFrame.size.width * DPIScaleFactor;
-		CurScreen->VisibleFramePixels.size.height = CurScreen->VisibleFrame.size.height * DPIScaleFactor;
+		CurScreen->Frame.origin.y = CurScreen->FramePixels.origin.y = WholeWorkspace.origin.y + WholeWorkspace.size.height - CurScreen->Frame.size.height - CurScreen->Frame.origin.y;
+		CurScreen->VisibleFrame.origin.y = CurScreen->VisibleFramePixels.origin.y = WholeWorkspace.origin.y + WholeWorkspace.size.height - CurScreen->VisibleFrame.size.height - CurScreen->VisibleFrame.origin.y;
 
 		SortedScreens.Add(CurScreen);
 	}
@@ -1613,13 +1599,18 @@ void FMacApplication::UpdateScreensArray()
 		const float DPIScaleFactor = bUseHighDPIMode ? CurScreen->Screen.backingScaleFactor : 1.0f;
 		if (DPIScaleFactor != 1.0f)
 		{
+			CurScreen->FramePixels.size.width = CurScreen->Frame.size.width * DPIScaleFactor;
+			CurScreen->FramePixels.size.height = CurScreen->Frame.size.height * DPIScaleFactor;
+			CurScreen->VisibleFramePixels.size.width = CurScreen->VisibleFrame.size.width * DPIScaleFactor;
+			CurScreen->VisibleFramePixels.size.height = CurScreen->VisibleFrame.size.height * DPIScaleFactor;
+
 			for (int32 OtherIndex = Index + 1; OtherIndex < SortedScreens.Num(); ++OtherIndex)
 			{
 				TSharedRef<FMacScreen> OtherScreen = SortedScreens[OtherIndex];
 				const float DiffFrame = (OtherScreen->Frame.origin.x - CurScreen->Frame.origin.x) * DPIScaleFactor;
 				const float DiffVisibleFrame = (OtherScreen->VisibleFrame.origin.x - CurScreen->VisibleFrame.origin.x) * DPIScaleFactor;
-				OtherScreen->FramePixels.origin.x = CurScreen->Frame.origin.x + DiffFrame;
-				OtherScreen->VisibleFramePixels.origin.x = CurScreen->VisibleFrame.origin.x + DiffVisibleFrame;
+				OtherScreen->FramePixels.origin.x = CurScreen->FramePixels.origin.x + DiffFrame;
+				OtherScreen->VisibleFramePixels.origin.x = CurScreen->VisibleFramePixels.origin.x + DiffVisibleFrame;
 			}
 		}
 	}
@@ -1637,8 +1628,8 @@ void FMacApplication::UpdateScreensArray()
 				TSharedRef<FMacScreen> OtherScreen = SortedScreens[OtherIndex];
 				const float DiffFrame = (OtherScreen->Frame.origin.y - CurScreen->Frame.origin.y) * DPIScaleFactor;
 				const float DiffVisibleFrame = (OtherScreen->VisibleFrame.origin.y - CurScreen->VisibleFrame.origin.y) * DPIScaleFactor;
-				OtherScreen->FramePixels.origin.y = CurScreen->Frame.origin.y + DiffFrame;
-				OtherScreen->VisibleFramePixels.origin.y = CurScreen->VisibleFrame.origin.y + DiffVisibleFrame;
+				OtherScreen->FramePixels.origin.y = CurScreen->FramePixels.origin.y + DiffFrame;
+				OtherScreen->VisibleFramePixels.origin.y = CurScreen->VisibleFramePixels.origin.y + DiffVisibleFrame;
 			}
 		}
 	}
