Queue Slate CloseWindow calls in MacApplication so that they're called in Tick instead of at random times when main thread is doing blocked game thread calls.

diff --git a/Engine/Source/Runtime/ApplicationCore/Private/Mac/MacApplication.cpp b/Engine/Source/Runtime/ApplicationCore/Private/Mac/MacApplication.cpp
--- a/Engine/Source/Runtime/ApplicationCore/Private/Mac/MacApplication.cpp
+++ b/Engine/Source/Runtime/ApplicationCore/Private/Mac/MacApplication.cpp
@@ -351,7 +351,10 @@ void FMacApplication::EndScopedModalEvent()
 
 void FMacApplication::CloseWindow(TSharedRef<FMacWindow> Window)
 {
-	MessageHandler->OnWindowClose(Window);
+	if (!SlateWindowsToClose.Contains(Window))
+	{
+		SlateWindowsToClose.Add(Window);
+	}
 }
 
 void FMacApplication::DeferEvent(NSObject* Object)
@@ -1141,9 +1144,9 @@ bool FMacApplication::OnWindowDestroyed(TSharedRef<FMacWindow> DestroyedWindow)
 
 	Windows.Remove(DestroyedWindow);
 
-	if (!WindowsToClose.Contains(WindowHandle))
+	if (!CocoaWindowsToClose.Contains(WindowHandle))
 	{
-		WindowsToClose.Add(WindowHandle);
+		CocoaWindowsToClose.Add(WindowHandle);
 	}
 
 	TSharedPtr<FMacWindow> WindowToActivate;
@@ -1157,7 +1160,7 @@ bool FMacApplication::OnWindowDestroyed(TSharedRef<FMacWindow> DestroyedWindow)
 		for (int32 Index = 0; Index < Windows.Num(); ++Index)
 		{
 			TSharedRef<FMacWindow> WindowRef = Windows[Index];
-			if (!WindowsToClose.Contains(WindowRef->GetWindowHandle()) && [WindowRef->GetWindowHandle() canBecomeMainWindow] && WindowRef->GetDefinition().Type != EWindowType::Notification)
+			if (!CocoaWindowsToClose.Contains(WindowRef->GetWindowHandle()) && [WindowRef->GetWindowHandle() canBecomeMainWindow] && WindowRef->GetDefinition().Type != EWindowType::Notification)
 			{
 				WindowToActivate = WindowRef;
 				break;
@@ -1939,18 +1942,27 @@ TCHAR FMacApplication::TranslateCharCode(TCHAR CharCode, uint32 KeyCode) const
 
 void FMacApplication::CloseQueuedWindows()
 {
-	if (WindowsToClose.Num() > 0)
+	if (SlateWindowsToClose.Num() > 0)
+	{
+		for (TSharedRef<FMacWindow> Window : SlateWindowsToClose)
+		{
+			MessageHandler->OnWindowClose(Window);
+		}
+		SlateWindowsToClose.Empty();
+	}
+
+	if (CocoaWindowsToClose.Num() > 0)
 	{
 		MainThreadCall(^{
 			SCOPED_AUTORELEASE_POOL;
-			for (FCocoaWindow* Window : WindowsToClose)
+			for (FCocoaWindow* Window : CocoaWindowsToClose)
 			{
 				[Window close];
 				[Window release];
 			}
 		}, UE4CloseEventMode, true);
 
-		WindowsToClose.Empty();
+		CocoaWindowsToClose.Empty();
 	}
 }
 
diff --git a/Engine/Source/Runtime/ApplicationCore/Public/Mac/MacApplication.h b/Engine/Source/Runtime/ApplicationCore/Public/Mac/MacApplication.h
--- a/Engine/Source/Runtime/ApplicationCore/Public/Mac/MacApplication.h
+++ b/Engine/Source/Runtime/ApplicationCore/Public/Mac/MacApplication.h
@@ -360,7 +360,8 @@ private:
 
 	volatile int32 bIgnoreMouseMoveDelta;
 
-	TArray<FCocoaWindow*> WindowsToClose;
+	TArray<FCocoaWindow*> CocoaWindowsToClose;
+	TArray<TSharedRef<FMacWindow>> SlateWindowsToClose;
 
 	TArray<FCocoaWindow*> WindowsRequiringTextInvalidation;
 
