Fixed thread safety in a couple of places in MacApplication where we modify Windows array

diff --git a/Engine/Source/Runtime/ApplicationCore/Private/Mac/MacApplication.cpp b/Engine/Source/Runtime/ApplicationCore/Private/Mac/MacApplication.cpp
--- a/Engine/Source/Runtime/ApplicationCore/Private/Mac/MacApplication.cpp
+++ b/Engine/Source/Runtime/ApplicationCore/Private/Mac/MacApplication.cpp
@@ -257,7 +257,11 @@ void FMacApplication::InitializeWindow(const TSharedRef<FGenericWindow>& InWindo
 	const TSharedRef<FMacWindow> Window = StaticCastSharedRef<FMacWindow >(InWindow);
 	const TSharedPtr<FMacWindow> ParentWindow = StaticCastSharedPtr<FMacWindow>(InParent);
 
-	Windows.Add(Window);
+	{
+		FScopeLock Lock(&WindowsMutex);
+		Windows.Add(Window);
+	}
+
 	Window->Initialize(this, InDefinition, ParentWindow, bShowImmediately);
 }
 
@@ -580,6 +584,7 @@ void FMacApplication::OnDisplayReconfiguration(CGDirectDisplayID Display, CGDisp
 		App->BroadcastDisplayMetricsChanged(DisplayMetrics);
 	}
 
+	FScopeLock Lock(&App->WindowsMutex);
 	for (int32 WindowIndex=0; WindowIndex < App->Windows.Num(); ++WindowIndex)
 	{
 		TSharedRef<FMacWindow> WindowRef = App->Windows[WindowIndex];
@@ -1142,7 +1147,10 @@ bool FMacApplication::OnWindowDestroyed(TSharedRef<FMacWindow> DestroyedWindow)
 		OnWindowActivationChanged(DestroyedWindow, EWindowActivation::Deactivate);
 	}
 
-	Windows.Remove(DestroyedWindow);
+	{
+		FScopeLock Lock(&WindowsMutex);
+		Windows.Remove(DestroyedWindow);
+	}
 
 	if (!CocoaWindowsToClose.Contains(WindowHandle))
 	{
