Allow width and height of MacWindow to be 0

diff --git a/Engine/Source/Runtime/ApplicationCore/Private/Mac/MacWindow.cpp b/Engine/Source/Runtime/ApplicationCore/Private/Mac/MacWindow.cpp
--- a/Engine/Source/Runtime/ApplicationCore/Private/Mac/MacWindow.cpp
+++ b/Engine/Source/Runtime/ApplicationCore/Private/Mac/MacWindow.cpp
@@ -627,7 +627,7 @@ void FMacWindow::ApplySizeAndModeChanges(int32 X, int32 Y, int32 Width, int32 He
 			Height /= DPIScaleFactor;
 
 			const FVector2D CocoaPosition = FMacApplication::ConvertSlatePositionToCocoa(X, Y);
-			NSRect Rect = NSMakeRect(CocoaPosition.X, CocoaPosition.Y - Height + 1, FMath::Max(Width, 1), FMath::Max(Height, 1));
+			NSRect Rect = NSMakeRect(CocoaPosition.X, CocoaPosition.Y - Height + 1, FMath::Max(Width, 0), FMath::Max(Height, 0));
 			if (Definition->HasOSWindowBorder)
 			{
 				Rect = [WindowHandle frameRectForContentRect:Rect];
