Added a mutex in MacApplication to fix thread safety issue with SlateWindowsToClose array

diff --git a/Engine/Source/Runtime/ApplicationCore/Private/Mac/MacApplication.cpp b/Engine/Source/Runtime/ApplicationCore/Private/Mac/MacApplication.cpp
--- a/Engine/Source/Runtime/ApplicationCore/Private/Mac/MacApplication.cpp
+++ b/Engine/Source/Runtime/ApplicationCore/Private/Mac/MacApplication.cpp
@@ -355,6 +355,7 @@ void FMacApplication::EndScopedModalEvent()
 
 void FMacApplication::CloseWindow(TSharedRef<FMacWindow> Window)
 {
+	FScopeLock Lock(&WindowsToCloseMutex);
 	if (!SlateWindowsToClose.Contains(Window))
 	{
 		SlateWindowsToClose.Add(Window);
@@ -1950,13 +1951,17 @@ TCHAR FMacApplication::TranslateCharCode(TCHAR CharCode, uint32 KeyCode) const
 
 void FMacApplication::CloseQueuedWindows()
 {
-	if (SlateWindowsToClose.Num() > 0)
 	{
-		for (TSharedRef<FMacWindow> Window : SlateWindowsToClose)
+		FScopeLock Lock(&WindowsToCloseMutex);
+
+		if (SlateWindowsToClose.Num() > 0)
 		{
-			MessageHandler->OnWindowClose(Window);
+			for (TSharedRef<FMacWindow> Window : SlateWindowsToClose)
+			{
+				MessageHandler->OnWindowClose(Window);
+			}
+			SlateWindowsToClose.Empty();
 		}
-		SlateWindowsToClose.Empty();
 	}
 
 	if (CocoaWindowsToClose.Num() > 0)
diff --git a/Engine/Source/Runtime/ApplicationCore/Public/Mac/MacApplication.h b/Engine/Source/Runtime/ApplicationCore/Public/Mac/MacApplication.h
--- a/Engine/Source/Runtime/ApplicationCore/Public/Mac/MacApplication.h
+++ b/Engine/Source/Runtime/ApplicationCore/Public/Mac/MacApplication.h
@@ -360,6 +360,7 @@ private:
 
 	volatile int32 bIgnoreMouseMoveDelta;
 
+	FCriticalSection WindowsToCloseMutex;
 	TArray<FCocoaWindow*> CocoaWindowsToClose;
 	TArray<TSharedRef<FMacWindow>> SlateWindowsToClose;
 
