Add the ability for WindowApplication to query minimized status. Stub for other platforms.

diff --git a/Engine/Source/Runtime/ApplicationCore/Private/Windows/WindowsApplication.cpp b/Engine/Source/Runtime/ApplicationCore/Private/Windows/WindowsApplication.cpp
--- a/Engine/Source/Runtime/ApplicationCore/Private/Windows/WindowsApplication.cpp
+++ b/Engine/Source/Runtime/ApplicationCore/Private/Windows/WindowsApplication.cpp
@@ -78,6 +78,7 @@ FWindowsApplication* FWindowsApplication::CreateWindowsApplication( const HINSTA
 FWindowsApplication::FWindowsApplication( const HINSTANCE HInstance, const HICON IconHandle )
 	: GenericApplication( MakeShareable( new FWindowsCursor() ) )
 	, InstanceHandle( HInstance )
+	, bMinimized( false )
 	, bUsingHighPrecisionMouseInput( false )
 	, bIsMouseAttached( false )
 	, bForceActivateByMouse( false )
@@ -2065,9 +2066,11 @@ int32 FWindowsApplication::ProcessDeferredMessage( const FDeferredWindowsMessage
 					{
 					case SW_PARENTCLOSING:
 						CurrentNativeEventWindowPtr->OnParentWindowMinimized();
+						bMinimized = true;
 						break;
 					case SW_PARENTOPENING:
 						CurrentNativeEventWindowPtr->OnParentWindowRestored();
+						bMinimized = true;
 						break;
 					default:
 						break;
diff --git a/Engine/Source/Runtime/ApplicationCore/Public/GenericPlatform/GenericApplication.h b/Engine/Source/Runtime/ApplicationCore/Public/GenericPlatform/GenericApplication.h
--- a/Engine/Source/Runtime/ApplicationCore/Public/GenericPlatform/GenericApplication.h
+++ b/Engine/Source/Runtime/ApplicationCore/Public/GenericPlatform/GenericApplication.h
@@ -462,6 +462,8 @@ public:
 	/** @return Native window under the mouse cursor. */
 	virtual TSharedPtr< FGenericWindow > GetWindowUnderCursor() { return TSharedPtr< FGenericWindow >( nullptr ); }
 
+	virtual bool IsMinimized() const { return false; }
+
 	virtual void SetHighPrecisionMouseMode( const bool Enable, const TSharedPtr< FGenericWindow >& InWindow ) { };
 
 	virtual bool IsUsingHighPrecisionMouseMode() const { return false; }
diff --git a/Engine/Source/Runtime/ApplicationCore/Public/Windows/WindowsApplication.h b/Engine/Source/Runtime/ApplicationCore/Public/Windows/WindowsApplication.h
--- a/Engine/Source/Runtime/ApplicationCore/Public/Windows/WindowsApplication.h
+++ b/Engine/Source/Runtime/ApplicationCore/Public/Windows/WindowsApplication.h
@@ -349,6 +349,7 @@ public:
 	virtual void InitializeWindow( const TSharedRef< FGenericWindow >& Window, const TSharedRef< FGenericWindowDefinition >& InDefinition, const TSharedPtr< FGenericWindow >& InParent, const bool bShowImmediately ) override;
 	virtual void SetCapture( const TSharedPtr< FGenericWindow >& InWindow ) override;
 	virtual void* GetCapture( void ) const override;
+	virtual bool IsMinimized() const override { return bMinimized; }
 	virtual void SetHighPrecisionMouseMode( const bool Enable, const TSharedPtr< FGenericWindow >& InWindow ) override;
 	virtual bool IsUsingHighPrecisionMouseMode() const override { return bUsingHighPrecisionMouseInput; }
 	virtual bool IsMouseAttached() const override { return bIsMouseAttached; }
@@ -449,6 +450,8 @@ private:
 
 	HINSTANCE InstanceHandle;
 
+	bool bMinimized;
+
 	bool bUsingHighPrecisionMouseInput;
 
 	bool bIsMouseAttached;
