bring over exit fix from 8.01

diff --git a/Engine/Source/Runtime/ApplicationCore/Private/IOS/IOSAppDelegate.cpp b/Engine/Source/Runtime/ApplicationCore/Private/IOS/IOSAppDelegate.cpp
--- a/Engine/Source/Runtime/ApplicationCore/Private/IOS/IOSAppDelegate.cpp
+++ b/Engine/Source/Runtime/ApplicationCore/Private/IOS/IOSAppDelegate.cpp
@@ -928,8 +928,8 @@ static FAutoConsoleVariableRef CVarGEnableThermalsReport(
 	OSVersion = [[[UIDevice currentDevice] systemVersion] floatValue];
 	if (!FPlatformMisc::IsDebuggerPresent() || GAlwaysReportCrash)
 	{
-//        FPlatformMisc::SetCrashHandler(EngineCrashHandler);
-        InstallSignalHandlers();
+        FPlatformMisc::SetCrashHandler(EngineCrashHandler);
+//        InstallSignalHandlers();
 	}
 
 	self.savedOpenUrlParameters = [[NSMutableArray alloc] init];
