Fix issue with JNI dependency

#ue4
#android

diff --git a/Engine/Source/Runtime/ApplicationCore/Private/Android/AndroidApplication.cpp b/Engine/Source/Runtime/ApplicationCore/Private/Android/AndroidApplication.cpp
--- a/Engine/Source/Runtime/ApplicationCore/Private/Android/AndroidApplication.cpp
+++ b/Engine/Source/Runtime/ApplicationCore/Private/Android/AndroidApplication.cpp
@@ -99,6 +99,7 @@ void FAndroidApplication::PollGameDeviceState( const float TimeDelta )
 
 void FAndroidApplication::HandleDeviceOrientation()
 {
+#if USE_ANDROID_JNI
     JNIEnv* JEnv = AndroidJavaEnv::GetJavaEnv();
     if (JEnv)
     {
@@ -137,6 +138,7 @@ void FAndroidApplication::HandleDeviceOrientation()
             FCoreDelegates::OnSafeFrameChangedEvent.Broadcast();
         }
     }
+#endif
 }
 
 FPlatformRect FAndroidApplication::GetWorkArea( const FPlatformRect& CurrentWindow ) const
