merge from Main corrections to Dev-Mobile - continuation of CL: #5440487

reset emscripten to 1.38.20

#rnx

diff --git a/Engine/Source/Programs/UnrealBuildTool/Platform/HTML5/HTML5SDKInfo.cs b/Engine/Source/Programs/UnrealBuildTool/Platform/HTML5/HTML5SDKInfo.cs
--- a/Engine/Source/Programs/UnrealBuildTool/Platform/HTML5/HTML5SDKInfo.cs
+++ b/Engine/Source/Programs/UnrealBuildTool/Platform/HTML5/HTML5SDKInfo.cs
@@ -18,8 +18,8 @@ namespace UnrealBuildTool
 		static string NODE_VER = "8.9.1_64bit";
 		static string PYTHON_VER = "2.7.13.1_64bit"; // Only used on Windows; other platforms use built-in Python.
 
-		static string LLVM_VER = "e1.38.24_64bit";
-		static string SDKVersion = "1.38.24";
+		static string LLVM_VER = "e1.38.20_64bit";
+		static string SDKVersion = "1.38.20";
 
 //		static string LLVM_VER = "incoming";
 //		static string SDKVersion = "incoming";
diff --git a/Engine/Source/Runtime/ApplicationCore/Private/HTML5/HTML5InputInterface.cpp b/Engine/Source/Runtime/ApplicationCore/Private/HTML5/HTML5InputInterface.cpp
--- a/Engine/Source/Runtime/ApplicationCore/Private/HTML5/HTML5InputInterface.cpp
+++ b/Engine/Source/Runtime/ApplicationCore/Private/HTML5/HTML5InputInterface.cpp
@@ -186,7 +186,7 @@ void FHTML5InputInterface::SendControllerEvents()
 //	/// XXX TODO Restore this, emscripten_get_num_gamepads() currently has a moderately high performance impact, so disabled for local testing
 //	static bool GamepadSupported = false;
 //#else
-	static bool GamepadSupported = emscripten_sample_gamepad_data() == EMSCRIPTEN_RESULT_SUCCESS;
+	static bool GamepadSupported = true;
 //#endif
 
 	const double CurrentTime = FPlatformTime::Seconds();
