UE-71338 //UE4/Main - Compile UE4Game TVOS - 'AVAudioSessionCategoryOptionAllowBluetooth' is unavailable: not available on tvOS

#tvOS
#4.23

Adding  !PLATFORM_TVOS where necessary

diff --git a/Engine/Source/Runtime/ApplicationCore/Private/IOS/IOSAppDelegate.cpp b/Engine/Source/Runtime/ApplicationCore/Private/IOS/IOSAppDelegate.cpp
--- a/Engine/Source/Runtime/ApplicationCore/Private/IOS/IOSAppDelegate.cpp
+++ b/Engine/Source/Runtime/ApplicationCore/Private/IOS/IOSAppDelegate.cpp
@@ -610,10 +610,10 @@ static IOSAppDelegate* CachedDelegate = nil;
 					else
 					{
 						AVAudioSessionCategoryOptions opts =
-							AVAudioSessionCategoryOptionAllowBluetooth |
 							AVAudioSessionCategoryOptionAllowBluetoothA2DP |
 #if !PLATFORM_TVOS
-    						AVAudioSessionCategoryOptionDefaultToSpeaker |
+							AVAudioSessionCategoryOptionAllowBluetooth |
+							AVAudioSessionCategoryOptionDefaultToSpeaker |
 #endif
 							AVAudioSessionCategoryOptionMixWithOthers;
 						
@@ -684,9 +684,9 @@ static IOSAppDelegate* CachedDelegate = nil;
 				else
 				{
 					AVAudioSessionCategoryOptions opts =
-						AVAudioSessionCategoryOptionAllowBluetooth |
 						AVAudioSessionCategoryOptionAllowBluetoothA2DP |
 #if !PLATFORM_TVOS
+						AVAudioSessionCategoryOptionAllowBluetooth |
 						AVAudioSessionCategoryOptionDefaultToSpeaker |
 #endif
 						AVAudioSessionCategoryOptionMixWithOthers;
@@ -714,9 +714,9 @@ static IOSAppDelegate* CachedDelegate = nil;
         if (self.bVoiceChatEnabled)
 		{
 			AVAudioSessionCategoryOptions opts =
-				AVAudioSessionCategoryOptionAllowBluetooth |
 				AVAudioSessionCategoryOptionAllowBluetoothA2DP |
 #if !PLATFORM_TVOS
+				AVAudioSessionCategoryOptionAllowBluetooth |
 				AVAudioSessionCategoryOptionDefaultToSpeaker |
 #endif
 				AVAudioSessionCategoryOptionMixWithOthers;
@@ -881,7 +881,7 @@ bool GIsSuspended = 0;
 		// Don't deadlock here because a msg box may appear super early blocking the game thread and then the app may go into the background
 		double	startTime = FPlatformTime::Seconds();
 
-		// don't wait for FDefaultGameMoviePlayer::WaitForMovieToFinish(), crash with 0x8badf00d if “Wait for Movies to Complete” is checked
+		// don't wait for FDefaultGameMoviePlayer::WaitForMovieToFinish(), crash with 0x8badf00d if Wait for Movies to Complete is checked
 		while(!self.bHasSuspended && !FAppEntry::IsStartupMoviePlaying() &&  (FPlatformTime::Seconds() - startTime) < cMaxThreadWaitTime)
 		{
             FIOSPlatformRHIFramePacer::Suspend();
