Add version guard around Win10 accessibility code

diff --git a/Engine/Source/Runtime/ApplicationCore/Private/Windows/Accessibility/WindowsUIAManager.cpp b/Engine/Source/Runtime/ApplicationCore/Private/Windows/Accessibility/WindowsUIAManager.cpp
--- a/Engine/Source/Runtime/ApplicationCore/Private/Windows/Accessibility/WindowsUIAManager.cpp
+++ b/Engine/Source/Runtime/ApplicationCore/Private/Windows/Accessibility/WindowsUIAManager.cpp
@@ -167,6 +167,7 @@ void FWindowsUIAManager::OnEventRaised(TSharedRef<IAccessibleWidget> Widget, EAc
 				UiaRaiseAutomationEvent(&ScopedProvider.Provider, UIA_Invoke_InvokedEventId);
 			}
 			break;
+#if WINVER >= 0x0A00
 		case EAccessibleEvent::Notification:
 		{
 			typedef HRESULT(WINAPI* UiaRaiseNotificationEventFunc)(IRawElementProviderSimple*, NotificationKind, NotificationProcessing, BSTR, BSTR);
@@ -188,6 +189,7 @@ void FWindowsUIAManager::OnEventRaised(TSharedRef<IAccessibleWidget> Widget, EAc
 			}
 			break;
 		}
+#endif
 		// IMPORTANT: Calling UiaRaiseStructureChangedEvent seems to raise our per-frame timing for accessibility by
 		// over 10x (.04ms to .7ms, tested by clicking on the "All Classes" button in the modes panel of the Editor).
 		// For now, I'm disabling this until we figure out if it's absolutely necessary.
