Fix missing !shipping guard and includes

diff --git a/Engine/Source/Runtime/ApplicationCore/Private/IOS/Accessibility/IOSAccessibilityCache.cpp b/Engine/Source/Runtime/ApplicationCore/Private/IOS/Accessibility/IOSAccessibilityCache.cpp
--- a/Engine/Source/Runtime/ApplicationCore/Private/IOS/Accessibility/IOSAccessibilityCache.cpp
+++ b/Engine/Source/Runtime/ApplicationCore/Private/IOS/Accessibility/IOSAccessibilityCache.cpp
@@ -8,13 +8,17 @@
 #include "IOS/IOSAppDelegate.h"
 #include "IOS/IOSAsyncTask.h"
 #include "IOS/IOSView.h"
+#include "HAL/IConsoleManager.h"
+#include "Async/TaskGraphInterfaces.h"
 
 @implementation FIOSAccessibilityCache
 
+#if !UE_BUILD_SHIPPING
 static void DumpAccessibilityStatsForwarder()
 {
 	[[FIOSAccessibilityCache AccessibilityElementCache] DumpAccessibilityStats];
 }
+#endif
 
 - (id)init
 {
