Setting the front microphone as InputDataSource for AVAudioSession if available

diff --git a/Engine/Source/Runtime/ApplicationCore/Private/IOS/IOSAppDelegate.cpp b/Engine/Source/Runtime/ApplicationCore/Private/IOS/IOSAppDelegate.cpp
--- a/Engine/Source/Runtime/ApplicationCore/Private/IOS/IOSAppDelegate.cpp
+++ b/Engine/Source/Runtime/ApplicationCore/Private/IOS/IOSAppDelegate.cpp
@@ -262,6 +262,7 @@ static IOSAppDelegate* CachedDelegate = nil;
 	FAppEntry::Init();
 	
 	[self InitIdleTimerSettings];
+	[self InitSpeakerPhoneSettings];
 
 	bEngineInit = true;
     
@@ -413,6 +414,13 @@ static IOSAppDelegate* CachedDelegate = nil;
 	[self EnableIdleTimer : bEnableTimer];
 }
 
+-(void)InitSpeakerPhoneSettings
+{
+	bool bEnableSpeakerPhone = NO;
+	GConfig->GetBool(TEXT("/Script/IOSRuntimeSettings.IOSRuntimeSettings"), TEXT("bEnableSpeakerPhone"), bEnableSpeakerPhone, GEngineIni);
+	[self EnableSpeakerPhone : bEnableSpeakerPhone];
+}
+
 -(bool)IsIdleTimerEnabled
 {
 	return ([UIApplication sharedApplication].idleTimerDisabled == NO);
@@ -635,6 +643,19 @@ static IOSAppDelegate* CachedDelegate = nil;
 					}
 					ActiveError = nil;
 	
+					// We will prioritise use of the front microphone when speakerphone mode is enabled
+					if (self.bVoiceChatEnabled)
+					{
+						AVAudioSessionLocation Orientation = self.bSpeakerPhoneEnabled ? AVAudioSessionOrientationFront : AVAudioSessionOrientationBottom;
+
+						bool bSuccess = [self SetInputDataSourceWithOrientation : Orientation error : ActiveError];
+						if (!bSuccess)
+						{
+							UE_LOG(LogIOSAudioSession, Error, TEXT("Failed to set InputDataSource for AVAudioSession [Error = %s]"), *FString([ActiveError description]));
+						}
+						ActiveError = nil;
+					}
+                    
                     if (!self.bVoiceChatEnabled)
                     {
 						if (!GAudio_ForceAmbientCategory)
@@ -709,6 +730,19 @@ static IOSAppDelegate* CachedDelegate = nil;
 				}
 				ActiveError = nil;
                 
+				// We will prioritise use of the front microphone when speakerphone mode is enabled
+				if (self.bVoiceChatEnabled)
+				{
+					AVAudioSessionLocation Orientation = self.bSpeakerPhoneEnabled ? AVAudioSessionOrientationFront : AVAudioSessionOrientationBottom;
+
+					bool bSuccess = [self SetInputDataSourceWithOrientation : Orientation error : ActiveError];
+					if (!bSuccess)
+					{
+						UE_LOG(LogIOSAudioSession, Error, TEXT("Failed to set InputDataSource for AVAudioSession [Error = %s]"), *FString([ActiveError description]));
+					}
+					ActiveError = nil;
+				}
+
                 if (!self.bVoiceChatEnabled)
                 {
 					if (!GAudio_ForceAmbientCategory)
@@ -798,6 +832,28 @@ static IOSAppDelegate* CachedDelegate = nil;
     self.bAudioActive = bActive;
 }
 
+- (bool)SetInputDataSourceWithOrientation:(AVAudioSessionOrientation)PreferredOrientation error:(NSError*)ActiveError
+{
+    bool bResult = NO;
+    
+	NSArray<AVAudioSessionDataSourceDescription*>* InputDataSources = [[AVAudioSession sharedInstance] inputDataSources];
+	if (InputDataSources != nil)
+	{
+		for (AVAudioSessionDataSourceDescription* InputDataSource in InputDataSources)
+		{
+			AVAudioSessionOrientation Orientation = [InputDataSource orientation];
+			if (Orientation == PreferredOrientation)
+			{				
+				bResult = [[AVAudioSession sharedInstance] setInputDataSource:InputDataSource error : &ActiveError];
+                
+                break;
+			}
+		}
+	}
+    
+    return bResult;
+}
+
 - (bool)IsBackgroundAudioPlaying
 {
 	AVAudioSession* Session = [AVAudioSession sharedInstance];
@@ -839,6 +895,23 @@ static IOSAppDelegate* CachedDelegate = nil;
 	//return [self IsFeatureActive:EAudioFeature::VoiceChat];
 }
 
+- (void)EnableSpeakerPhone:(bool)bEnable
+{
+	self.bSpeakerPhoneEnabled = false;
+
+	// mobile will prompt for microphone access
+	if (FApp::IsUnattended())
+	{
+		return;
+	}
+	self.bSpeakerPhoneEnabled = bEnable;
+	[self ToggleAudioSession:self.bAudioActive force:true];
+}
+
+- (bool)IsSpeakerPhoneEnabled
+{
+	return self.bSpeakerPhoneEnabled;
+}
 
 - (void)SetFeature:(EAudioFeature)Feature Active:(bool)bIsActive
 {
diff --git a/Engine/Source/Runtime/ApplicationCore/Public/IOS/IOSAppDelegate.h b/Engine/Source/Runtime/ApplicationCore/Public/IOS/IOSAppDelegate.h
--- a/Engine/Source/Runtime/ApplicationCore/Public/IOS/IOSAppDelegate.h
+++ b/Engine/Source/Runtime/ApplicationCore/Public/IOS/IOSAppDelegate.h
@@ -215,6 +215,10 @@ APPLICATIONCORE_API
 - (void)EnableVoiceChat:(bool)bEnable;
 - (void)EnableHighQualityVoiceChat:(bool)bEnable;
 - (bool)IsVoiceChatEnabled;
+- (void)EnableSpeakerPhone:(bool)bEnable;
+- (bool)IsSpeakerPhoneEnabled;
+
+- (bool)SetInputDataSourceWithOrientation:(AVAudioSessionOrientation)PreferredOrientation error:(NSError*)ActiveError;
 
 - (void)SetFeature:(EAudioFeature)Feature Active:(bool)bIsActive;
 - (bool)IsFeatureActive:(EAudioFeature)Mode;
@@ -222,6 +226,7 @@ APPLICATIONCORE_API
 @property (atomic) bool bAudioActive;
 @property (atomic) bool bVoiceChatEnabled;
 @property (atomic) bool bHighQualityVoiceChatEnabled;
+@property (atomic) bool bSpeakerPhoneEnabled;
 
 @property (atomic) bool bIsSuspended;
 @property (atomic) bool bHasSuspended;
