Fixed UE4Game Android error

diff --git a/Engine/Source/Runtime/ApplicationCore/Private/Android/AndroidWindow.cpp b/Engine/Source/Runtime/ApplicationCore/Private/Android/AndroidWindow.cpp
--- a/Engine/Source/Runtime/ApplicationCore/Private/Android/AndroidWindow.cpp
+++ b/Engine/Source/Runtime/ApplicationCore/Private/Android/AndroidWindow.cpp
@@ -164,9 +164,6 @@ void FAndroidWindow::ReleaseWindowRef(ANativeWindow* InWindow)
 
  void FAndroidWindow::SetHardwareWindow_EventThread(void* InWindow)
 {
-#if USE_ANDROID_EVENTS
-	check(IsInAndroidEventThread());
-#endif
 	NativeWindow = InWindow; //using raw native window handle for now. Could be changed to use AndroidWindow later if needed
 }
 
