Use FString::GetNSString() instead of FPlatformString::TCHARToCFString in Mac ClipboardCopy() as it should be more reliable. Also, handle the case where a string object wasn't created and print a warning to the log instead of crashing.

diff --git a/Engine/Source/Runtime/ApplicationCore/Private/Mac/MacPlatformApplicationMisc.cpp b/Engine/Source/Runtime/ApplicationCore/Private/Mac/MacPlatformApplicationMisc.cpp
--- a/Engine/Source/Runtime/ApplicationCore/Private/Mac/MacPlatformApplicationMisc.cpp
+++ b/Engine/Source/Runtime/ApplicationCore/Private/Mac/MacPlatformApplicationMisc.cpp
@@ -517,13 +517,19 @@ void FMacPlatformApplicationMisc::ClipboardCopy(const TCHAR* Str)
 	{
 		SCOPED_AUTORELEASE_POOL;
 
-		CFStringRef CocoaString = FPlatformString::TCHARToCFString(Str);
-		NSPasteboard *Pasteboard = [NSPasteboard generalPasteboard];
-		[Pasteboard clearContents];
-		NSPasteboardItem *Item = [[[NSPasteboardItem alloc] init] autorelease];
-		[Item setString: (NSString *)CocoaString forType: NSPasteboardTypeString];
-		[Pasteboard writeObjects:[NSArray arrayWithObject:Item]];
-		CFRelease(CocoaString);
+		NSString* CocoaString = FString(Str).GetNSString();
+		if (CocoaString)
+		{
+			NSPasteboard *Pasteboard = [NSPasteboard generalPasteboard];
+			[Pasteboard clearContents];
+			NSPasteboardItem *Item = [[[NSPasteboardItem alloc] init] autorelease];
+			[Item setString:CocoaString forType:NSPasteboardTypeString];
+			[Pasteboard writeObjects:[NSArray arrayWithObject:Item]];
+		}
+		else
+		{
+			UE_LOG(LogMac, Warning, TEXT("Failed to create NSString to copy text to the pasteboard."));
+		}
 	}
 }
 
