Add controller lighting and force feedback support to android input interface

#android

diff --git a/Engine/Source/Runtime/ApplicationCore/Private/Android/AndroidInputInterface.cpp b/Engine/Source/Runtime/ApplicationCore/Private/Android/AndroidInputInterface.cpp
--- a/Engine/Source/Runtime/ApplicationCore/Private/Android/AndroidInputInterface.cpp
+++ b/Engine/Source/Runtime/ApplicationCore/Private/Android/AndroidInputInterface.cpp
@@ -196,11 +196,38 @@ void FAndroidInputInterface::Tick(float DeltaTime)
 	}
 }
 
+void FAndroidInputInterface::SetLightColor(int32 ControllerId, FColor Color)
+{
+	for (auto DeviceIt = ExternalInputDevices.CreateIterator(); DeviceIt; ++DeviceIt)
+	{
+		(*DeviceIt)->SetLightColor(ControllerId, Color);
+	}
+}
+
+void FAndroidInputInterface::ResetLightColor(int32 ControllerId)
+{
+	for (auto DeviceIt = ExternalInputDevices.CreateIterator(); DeviceIt; ++DeviceIt)
+	{
+		(*DeviceIt)->ResetLightColor(ControllerId);
+	}
+}
+
 void FAndroidInputInterface::SetForceFeedbackChannelValue(int32 ControllerId, FForceFeedbackChannelType ChannelType, float Value)
 {
+	bool bDidFeedback = false;
 	for (auto DeviceIt = ExternalInputDevices.CreateIterator(); DeviceIt; ++DeviceIt)
 	{
-		(*DeviceIt)->SetChannelValue(ControllerId, ChannelType, Value);
+		if ((*DeviceIt)->SupportsForceFeedback(ControllerId))
+		{
+			bDidFeedback = true;
+			(*DeviceIt)->SetChannelValue(ControllerId, ChannelType, Value);
+		}
+	}
+
+	// If controller handled force feedback don't do it on the phone
+	if (bDidFeedback)
+	{
+		return;
 	}
 
 	// Note: only one motor on Android at the moment, but remember all the settings
@@ -236,9 +263,20 @@ void FAndroidInputInterface::SetForceFeedbackChannelValue(int32 ControllerId, FF
 
 void FAndroidInputInterface::SetForceFeedbackChannelValues(int32 ControllerId, const FForceFeedbackValues &Values)
 {
+	bool bDidFeedback = false;
 	for (auto DeviceIt = ExternalInputDevices.CreateIterator(); DeviceIt; ++DeviceIt)
 	{
-		(*DeviceIt)->SetChannelValues(ControllerId, Values);
+		if ((*DeviceIt)->SupportsForceFeedback(ControllerId))
+		{
+			bDidFeedback = true;
+			(*DeviceIt)->SetChannelValues(ControllerId, Values);
+		}
+	}
+
+	// If controller handled force feedback don't do it on the phone
+	if (bDidFeedback)
+	{
+		return;
 	}
 
 	// Note: only one motor on Android at the moment, but remember all the settings
@@ -812,6 +850,7 @@ void FAndroidInputInterface::SendControllerEvents()
 						CurrentDevice.DeviceState = MappingState::Valid;
 
 						// Generic mappings
+						CurrentDevice.ControllerClass = ControllerClassType::Generic;
 						CurrentDevice.ButtonRemapping = ButtonRemapType::Normal;
 						CurrentDevice.LTAnalogRangeMinimum = 0.0f;
 						CurrentDevice.RTAnalogRangeMinimum = 0.0f;
@@ -854,16 +893,18 @@ void FAndroidInputInterface::SendControllerEvents()
 						}
 						else if (CurrentDevice.DeviceInfo.Name.StartsWith(TEXT("Xbox Wired Controller")))
 						{
+							CurrentDevice.ControllerClass = ControllerClassType::XBoxWired;
 							CurrentDevice.bSupportsHat = true;
 						}
 						else if (CurrentDevice.DeviceInfo.Name.StartsWith(TEXT("Xbox Wireless Controller")))
 						{
+							CurrentDevice.ControllerClass = ControllerClassType::XBoxWireless;
 							CurrentDevice.bSupportsHat = true;
 
 							if (GAndroidOldXBoxWirelessFirmware == 1)
 							{
 								// Apply mappings for older firmware before 3.1.1221.0
-								CurrentDevice.ButtonRemapping = ButtonRemapType::XBoxWireless;
+								CurrentDevice.ButtonRemapping = ButtonRemapType::XBox;
 								CurrentDevice.bMapL1R1ToTriggers = false;
 								CurrentDevice.bMapZRZToTriggers = true;
 								CurrentDevice.bRightStickZRZ = false;
@@ -882,6 +923,7 @@ void FAndroidInputInterface::SendControllerEvents()
 						// comparison. Instead we check the product and vendor IDs to ensure it's the correct one.
 						else if (CurrentDevice.DeviceInfo.Name.StartsWith(TEXT("PS4 Wireless Controller")))
 						{
+							CurrentDevice.ControllerClass = ControllerClassType::PS4Wireless;
 							if (CurrentDevice.DeviceInfo.Name.EndsWith(TEXT(" (v2)")))
 							{
 								CurrentDevice.ButtonRemapping = ButtonRemapType::PS4;
@@ -1170,6 +1212,15 @@ int32 FAndroidInputInterface::FindExistingDevice(int32 deviceId)
 	return -1;
 }
 
+FAndroidGamepadDeviceMapping* FAndroidInputInterface::GetDeviceMapping(int32 ControllerId)
+{
+	if (ControllerId < 0 || ControllerId >= MAX_NUM_CONTROLLERS)
+	{
+		return nullptr;
+	}
+	return &DeviceMapping[ControllerId];
+}
+
 int32 FAndroidInputInterface::GetControllerIndex(int32 deviceId)
 {
 	if (!bAllowControllers)
@@ -1336,6 +1387,17 @@ void FAndroidInputInterface::JoystickButtonEvent(int32 deviceId, int32 buttonId,
 	if (deviceId == -1)
 		return;
 
+	//FPlatformMisc::LowLevelOutputDebugStringf(TEXT("JoystickButtonEvent[%d]: %d"), (int)DeviceMapping[deviceId].ButtonRemapping, buttonId);
+
+	if (DeviceMapping[deviceId].ControllerClass == ControllerClassType::PS4Wireless)
+	{
+		if (buttonId == 3002)
+		{
+			NewControllerData[deviceId].ButtonStates[7] = buttonDown;  // Touchpad = Special Left
+			return;
+		}
+	}
+
 	// Deal with button remapping
 	switch (DeviceMapping[deviceId].ButtonRemapping)
 	{
@@ -1384,7 +1446,7 @@ void FAndroidInputInterface::JoystickButtonEvent(int32 deviceId, int32 buttonId,
 			}
 			break;
 
-		case ButtonRemapType::XBoxWireless:
+		case ButtonRemapType::XBox:
 			switch (buttonId)
 			{
 				case AKEYCODE_BUTTON_A:      NewControllerData[deviceId].ButtonStates[0] = buttonDown; break; // A
diff --git a/Engine/Source/Runtime/ApplicationCore/Public/Android/AndroidInputInterface.h b/Engine/Source/Runtime/ApplicationCore/Public/Android/AndroidInputInterface.h
--- a/Engine/Source/Runtime/ApplicationCore/Public/Android/AndroidInputInterface.h
+++ b/Engine/Source/Runtime/ApplicationCore/Public/Android/AndroidInputInterface.h
@@ -96,10 +96,18 @@ enum MappingState
 	Valid
 };
 
+enum ControllerClassType
+{
+	Generic,
+	XBoxWired,
+	XBoxWireless,
+	PS4Wireless
+};
+
 enum ButtonRemapType
 {
 	Normal,
-	XBoxWireless,
+	XBox,
 	PS4
 };
 
@@ -120,6 +128,9 @@ struct FAndroidGamepadDeviceMapping
 	// State of mapping
 	MappingState DeviceState;
 
+	// Type of controller
+	ControllerClassType ControllerClass;
+
 	// Type of button remapping to use
 	ButtonRemapType ButtonRemapping;
 
@@ -249,8 +260,8 @@ public:
 	virtual void SetForceFeedbackChannelValue(int32 ControllerId, FForceFeedbackChannelType ChannelType, float Value) override;
 	virtual void SetForceFeedbackChannelValues(int32 ControllerId, const FForceFeedbackValues &values) override;
 	virtual void SetHapticFeedbackValues(int32 ControllerId, int32 Hand, const FHapticFeedbackValues& Values) override;
-	virtual void SetLightColor(int32 ControllerId, FColor Color) override {}
-	virtual void ResetLightColor(int32 ControllerId) override {}
+	virtual void SetLightColor(int32 ControllerId, FColor Color) override;
+	virtual void ResetLightColor(int32 ControllerId) override;
 
 	void SetGamepadsAllowed(bool bAllowed) { bAllowControllers = bAllowed; }
 	virtual bool IsGamepadAttached() const;
@@ -264,6 +275,9 @@ private:
 
 	FAndroidInputInterface( const TSharedRef< FGenericApplicationMessageHandler >& InMessageHandler, const TSharedPtr< ICursor >& InCursor);
 
+public:
+
+	FAndroidGamepadDeviceMapping* GetDeviceMapping(int32 ControllerId);
 
 private:
 
diff --git a/Engine/Source/Runtime/InputDevice/Public/IInputDevice.h b/Engine/Source/Runtime/InputDevice/Public/IInputDevice.h
--- a/Engine/Source/Runtime/InputDevice/Public/IInputDevice.h
+++ b/Engine/Source/Runtime/InputDevice/Public/IInputDevice.h
@@ -34,6 +34,13 @@ public:
 	 */
 	virtual void SetChannelValue (int32 ControllerId, FForceFeedbackChannelType ChannelType, float Value) = 0;
 	virtual void SetChannelValues (int32 ControllerId, const FForceFeedbackValues &values) = 0;
+	virtual bool SupportsForceFeedback(int32 ControllerId) { return true; }
+
+	/**
+	 * Pass though functions for light color
+	 */
+	virtual void SetLightColor(int32 ControllerId, FColor Color) { };
+	virtual void ResetLightColor(int32 ControllerId) { };
 
 	/** If this device supports a haptic interface, implement this, and inherit the IHapticDevice interface */
 	virtual class IHapticDevice* GetHapticDevice()
diff --git a/Engine/Source/Runtime/Launch/Private/Android/LaunchAndroid.cpp b/Engine/Source/Runtime/Launch/Private/Android/LaunchAndroid.cpp
--- a/Engine/Source/Runtime/Launch/Private/Android/LaunchAndroid.cpp
+++ b/Engine/Source/Runtime/Launch/Private/Android/LaunchAndroid.cpp
@@ -109,7 +109,8 @@ static const uint32 ValidGamepadKeyCodesList[] =
 	AKEYCODE_DPAD_UP,
 	AKEYCODE_DPAD_DOWN,
 	AKEYCODE_DPAD_LEFT,
-	AKEYCODE_DPAD_RIGHT
+	AKEYCODE_DPAD_RIGHT,
+	3002  // touchpad
 };
 
 // map of gamepad keycodes that should be passed forward
