Should be popup here not borderless

diff --git a/Engine/Source/Runtime/ApplicationCore/Private/Linux/LinuxWindow.cpp b/Engine/Source/Runtime/ApplicationCore/Private/Linux/LinuxWindow.cpp
--- a/Engine/Source/Runtime/ApplicationCore/Private/Linux/LinuxWindow.cpp
+++ b/Engine/Source/Runtime/ApplicationCore/Private/Linux/LinuxWindow.cpp
@@ -182,7 +182,7 @@ void FLinuxWindow::Initialize( FLinuxApplication* const Application, const TShar
 		!Definition->IsModalWindow && !Definition->IsRegularWindow &&
 		!bShouldActivate && !Definition->SizeWillChangeOften)
 	{
-		WindowStyle |= SDL_WINDOW_BORDERLESS;
+		WindowStyle |= SDL_WINDOW_POPUP_MENU;
 		bIsConsoleWindow = true;
 		bIsPopupWindow = true;
 		UE_LOG(LogLinuxWindowType, Verbose, TEXT("*** New Window is a Console Window ***"));
