fix for OOM file remaining when the user exits the application while the game thread is doing a bit of work (GC or audio thread wait)

diff --git a/Engine/Source/Runtime/ApplicationCore/Private/IOS/IOSAppDelegate.cpp b/Engine/Source/Runtime/ApplicationCore/Private/IOS/IOSAppDelegate.cpp
--- a/Engine/Source/Runtime/ApplicationCore/Private/IOS/IOSAppDelegate.cpp
+++ b/Engine/Source/Runtime/ApplicationCore/Private/IOS/IOSAppDelegate.cpp
@@ -59,6 +59,7 @@ int GAudio_ForceAmbientCategory = 1;
 extern bool GShowSplashScreen;
 
 FIOSCoreDelegates::FOnOpenURL FIOSCoreDelegates::OnOpenURL;
+FIOSCoreDelegates::FOnWillResignActive FIOSCoreDelegates::OnWillResignActive;
 TArray<FIOSCoreDelegates::FFilterDelegateAndHandle> FIOSCoreDelegates::PushNotificationFilters;
 
 static bool GEnabledAudioFeatures[(uint8)EAudioFeature::NumFeatures];
@@ -1479,7 +1480,8 @@ extern EDeviceScreenOrientation ConvertFromUIInterfaceOrientation(UIInterfaceOri
 FCriticalSection RenderSuspend;
 - (void)applicationWillResignActive:(UIApplication *)application
 {
-    
+	FIOSCoreDelegates::OnWillResignActive.Broadcast();
+	
     FIOSPlatformMisc::ResetBrightness();
     
     /*
@@ -1497,7 +1499,9 @@ FCriticalSection RenderSuspend;
  		FEmbeddedCommunication::KeepAwake(TEXT("Background"), false);
         FGraphEventRef ResignTask = FFunctionGraphTask::CreateAndDispatchWhenReady([]()
         {
-            FCoreDelegates::ApplicationWillDeactivateDelegate.Broadcast();
+			UE_LOG(LogTemp, Display, TEXT("Calling Delegate"));
+
+			FCoreDelegates::ApplicationWillDeactivateDelegate.Broadcast();
 
 			FEmbeddedCommunication::AllowSleep(TEXT("Background"));
         }, TStatId(), NULL, ENamedThreads::GameThread);
@@ -1509,9 +1513,11 @@ FCriticalSection RenderSuspend;
 			FPlatformProcess::Sleep(0.05f);
 			if(ResignTask->IsComplete())
 			{
+				UE_LOG(LogTemp, Display, TEXT("Task was completed before time."));
 				break;
 			}
 		}
+		UE_LOG(LogTemp, Display, TEXT("Done with entering background tasks time."));
     }
 	[self ToggleSuspend:true];
 	[self ToggleAudioSession:false force:true];
@@ -1601,6 +1607,7 @@ extern double GCStartTime;
 
        FGraphEventRef ResignTask = FFunctionGraphTask::CreateAndDispatchWhenReady([]()
        {
+		   
             FCoreDelegates::ApplicationHasReactivatedDelegate.Broadcast();
 		
 			FEmbeddedCommunication::AllowSleep(TEXT("Background"));
diff --git a/Engine/Source/Runtime/ApplicationCore/Public/IOS/IOSAppDelegate.h b/Engine/Source/Runtime/ApplicationCore/Public/IOS/IOSAppDelegate.h
--- a/Engine/Source/Runtime/ApplicationCore/Public/IOS/IOSAppDelegate.h
+++ b/Engine/Source/Runtime/ApplicationCore/Public/IOS/IOSAppDelegate.h
@@ -46,6 +46,10 @@ public:
 	/** INTERNAL - check if a push notification payload passes all registered filters */
 	static bool PassesPushNotificationFilters(NSDictionary* Payload);
 
+	/** INTERNAL - called when entering the background - this is not thread-safe with the game thread or render thread as it is called from the app's Main thread */
+	DECLARE_MULTICAST_DELEGATE(FOnWillResignActive);
+	static FOnWillResignActive OnWillResignActive;
+
 private:
 	struct FFilterDelegateAndHandle
 	{
