With some specific DPI settings the barriers for confined cursor can end up touching the border
Allowing for OS resize to happen while in some locked game. Add a percentage to make the barrier
smaller around the edges

diff --git a/Engine/Source/Runtime/ApplicationCore/Private/Linux/LinuxCursor.cpp b/Engine/Source/Runtime/ApplicationCore/Private/Linux/LinuxCursor.cpp
--- a/Engine/Source/Runtime/ApplicationCore/Private/Linux/LinuxCursor.cpp
+++ b/Engine/Source/Runtime/ApplicationCore/Private/Linux/LinuxCursor.cpp
@@ -306,6 +306,33 @@ void FLinuxCursor::Show( bool bShow )
 	}
 }
 
+static SDL_Rect GetConfinementRect(const RECT* const Bounds, const TSharedPtr<FLinuxWindow>& CurrentFocusWindow)
+{
+	SDL_Rect ConfineRect;
+
+	int Left = 0;
+	int Top = 0;
+	int Right = 0;
+	int Bottom = 0;
+	SDL_GetWindowBordersSize(CurrentFocusWindow->GetHWnd(), &Top, &Left, &Bottom, &Right);
+
+	ConfineRect.x = FMath::TruncToInt(Bounds->left) + Left;
+	ConfineRect.y = FMath::TruncToInt(Bounds->top) - Top;
+
+	ConfineRect.w = FMath::TruncToInt(Bounds->right) - FMath::TruncToInt(Bounds->left) - 1;
+	ConfineRect.h = FMath::TruncToInt(Bounds->bottom) - FMath::TruncToInt(Bounds->top) - 1;
+
+	// TODO. For some reason the values from Bounds seem not to cover the SDL's window area.
+	// The cursor exceeds the right and bottom border of the window. This can be troublesome
+	// when the cursor is outside the window and the user clicks the mouse button even the
+	// cursor is confined. To prevent that the workaround is to add -amount additionally to the 
+	// values used before.
+	ConfineRect.w -= FMath::CeilToInt(ConfineRect.w * 0.002);
+	ConfineRect.h -= FMath::CeilToInt(ConfineRect.h * 0.002);
+
+	return ConfineRect;
+}
+
 void FLinuxCursor::Lock( const RECT* const Bounds )
 {
 	if (!LinuxApplication)
@@ -326,23 +353,7 @@ void FLinuxCursor::Lock( const RECT* const Bounds )
 	}
 	else
 	{
-		int Left = 0;
-		int Top = 0;
-		int Right = 0;
-		int Bottom = 0;
-		SDL_GetWindowBordersSize(CurrentFocusWindow->GetHWnd(), &Top, &Left, &Bottom, &Right);
-
-		CursorClipRect.x = FMath::TruncToInt(Bounds->left) + Left;
-		CursorClipRect.y = FMath::TruncToInt(Bounds->top) - Top;
-
-		// TODO. For some reason the values from Bounds seem not to cover the SDL's window area.
-		// The cursor exceeds the right and bottom border of the window. This can be troublesome
-		// when the cursor is outside the window and the user clicks the mouse button even the
-		// cursor is confined. To prevent that the workaround is to add -1 additionally to the 
-		// values used before.
-		const int ExceedingWindowPrevention = 1;
-		CursorClipRect.w = FMath::TruncToInt(Bounds->right) - FMath::TruncToInt(Bounds->left) - 1 - ExceedingWindowPrevention;
-		CursorClipRect.h = FMath::TruncToInt(Bounds->bottom) - FMath::TruncToInt(Bounds->top) - 1 - ExceedingWindowPrevention;
+		CursorClipRect = GetConfinementRect(Bounds, CurrentFocusWindow);
 
 		// We dont want to set a negative bounding region. If Top, Left, Bottom, Right are all 0
 		if (CursorClipRect.x >= 0 && CursorClipRect.y >= 0 && CursorClipRect.w > 0 && CursorClipRect.h > 0)
