Restructure how we set AVAudioSession to give the application more control over it
Workaround an issue in the vivox-sdk that causes the iOS recording notification to be displayed if the application is backgrounded while the sdk is initialized and the AVAudioSessionCategory is set to Ambient

diff --git a/Engine/Plugins/Online/VoiceChat/VivoxVoiceChat/Source/Private/IOS/IOSVivoxVoiceChat.cpp b/Engine/Plugins/Online/VoiceChat/VivoxVoiceChat/Source/Private/IOS/IOSVivoxVoiceChat.cpp
--- a/Engine/Plugins/Online/VoiceChat/VivoxVoiceChat/Source/Private/IOS/IOSVivoxVoiceChat.cpp
+++ b/Engine/Plugins/Online/VoiceChat/VivoxVoiceChat/Source/Private/IOS/IOSVivoxVoiceChat.cpp
@@ -29,6 +29,9 @@ FIOSVivoxVoiceChat::~FIOSVivoxVoiceChat()
 
 bool FIOSVivoxVoiceChat::Initialize()
 {
+	// Vivox sdk opens the microphone if we are in one of the Ambient modes and the user backgrounds the application. Set the mode to Playback to work around this
+	[[IOSAppDelegate GetDelegate] SetFeature:EAudioFeature::Playback Active:true];
+
 	bool bResult = FVivoxVoiceChat::Initialize();
 
 	if (bResult)
@@ -76,7 +79,16 @@ bool FIOSVivoxVoiceChat::Uninitialize()
 		AudioRouteChangedHandle.Reset();
 	}
 
-	return FVivoxVoiceChat::Uninitialize();
+	bool bReturn = FVivoxVoiceChat::Uninitialize();
+
+	while (VoiceChatEnableCount > 0)
+	{
+		EnableVoiceChat(false);
+	}
+
+	[[IOSAppDelegate GetDelegate] SetFeature:EAudioFeature::Playback Active:false];
+
+	return bReturn;
 }
 
 void FIOSVivoxVoiceChat::SetSetting(const FString& Name, const FString& Value)
@@ -103,6 +115,7 @@ FString FIOSVivoxVoiceChat::GetSetting(const FString& Name)
 
 FDelegateHandle FIOSVivoxVoiceChat::StartRecording(const FOnVoiceChatRecordSamplesAvailableDelegate::FDelegate& Delegate)
 {
+	bIsRecording = true;
 	EnableVoiceChat(true);
 	return FVivoxVoiceChat::StartRecording(Delegate);
 }
@@ -110,10 +123,8 @@ FDelegateHandle FIOSVivoxVoiceChat::StartRecording(const FOnVoiceChatRecordSampl
 void FIOSVivoxVoiceChat::StopRecording(FDelegateHandle Handle)
 {
 	FVivoxVoiceChat::StopRecording(Handle);
-	if (ConnectionState < EConnectionState::Connecting)
-	{
-		EnableVoiceChat(false);
-	}
+	EnableVoiceChat(false);
+	bIsRecording = false;
 }
 
 void FIOSVivoxVoiceChat::InvokeOnUIThread(void (Func)(void* Arg0), void* Arg0)
@@ -138,29 +149,18 @@ void FIOSVivoxVoiceChat::onChannelJoined(const VivoxClientApi::AccountName& Acco
 
 void FIOSVivoxVoiceChat::onChannelExited(const VivoxClientApi::AccountName& AccountName, const VivoxClientApi::Uri& ChannelUri, const VivoxClientApi::VCSStatus& Status)
 {
+	EnableVoiceChat(false);
 	FVivoxVoiceChat::onChannelExited(AccountName, ChannelUri, Status);
-	
-	bool bIsInChannel = false;
-	for (const TPair<FString,FChannelSession>& ChannelSession : LoginSession.ChannelSessions)
-	{
-		if (ChannelSession.Value.State == FChannelSession::EState::Connected)
-		{
-			bIsInChannel = true;
-		}
-	}
-	if (!bIsRecording && !bIsInChannel)
-	{
-		EnableVoiceChat(false);
-	}
 }
 
 void FIOSVivoxVoiceChat::onDisconnected(const VivoxClientApi::Uri& Server, const VivoxClientApi::VCSStatus& Status)
 {
-	FVivoxVoiceChat::onDisconnected(Server, Status);
-	if (!bIsRecording)
+	while (VoiceChatEnableCount > (bIsRecording ? 1 : 0))
 	{
+		// call once for every channel we were in
 		EnableVoiceChat(false);
 	}
+	FVivoxVoiceChat::onDisconnected(Server, Status);
 }
 
 void FIOSVivoxVoiceChat::OnVoiceChatConnectComplete(const FVoiceChatResult& Result)
@@ -297,28 +297,46 @@ bool FIOSVivoxVoiceChat::IsHardwareAECEnabled() const
 
 void FIOSVivoxVoiceChat::EnableVoiceChat(bool bEnable)
 {
-	if (FPlatformMisc::IsVoiceChatEnabled() != bEnable)
+	if (bEnable)
+	{
+		++VoiceChatEnableCount;
+		if (VoiceChatEnableCount == 1)
+		{
+			const bool bEnableAEC = IsHardwareAECEnabled() && !IsBluetoothA2DPInUse();
+			[[IOSAppDelegate GetDelegate] SetFeature:EAudioFeature::Playback Active:true];
+			[[IOSAppDelegate GetDelegate] SetFeature:EAudioFeature::Record Active:true];
+			vx_set_platform_aec_enabled(bEnableAEC ? 1 : 0);
+		}
+	}
+	else
 	{
-		if (IsHardwareAECEnabled() && !IsBluetoothA2DPInUse())
+		if (ensureMsgf(VoiceChatEnableCount > 0, TEXT("Attempted to disable voice chat when it was already disabled")))
 		{
-			[[IOSAppDelegate GetDelegate] EnableHighQualityVoiceChat:bEnable];
-			vx_set_platform_aec_enabled(bEnable ? 1 : 0);
+			--VoiceChatEnableCount;
+			if (VoiceChatEnableCount == 0)
+			{
+				vx_set_platform_aec_enabled(0);
+				[[IOSAppDelegate GetDelegate] SetFeature:EAudioFeature::Record Active:false];
+				[[IOSAppDelegate GetDelegate] SetFeature:EAudioFeature::Playback Active:false];
+			}
 		}
-		FPlatformMisc::EnableVoiceChat(bEnable);
 	}
 }
 
 void FIOSVivoxVoiceChat::UpdateVoiceChatSettings()
 {
-	if (FPlatformMisc::IsVoiceChatEnabled())
+	const bool bEnableAEC = IsHardwareAECEnabled() && !IsBluetoothA2DPInUse();
+
+	if (bVoiceChatModeEnabled != bEnableAEC)
+	{
+		[[IOSAppDelegate GetDelegate] SetFeature:EAudioFeature::VoiceChat Active:bEnableAEC];
+		bVoiceChatModeEnabled = bEnableAEC;
+	}
+
+	if (VoiceChatEnableCount > 0)
 	{
-		// update the aec settings
-		const bool bEnableAEC = IsHardwareAECEnabled() && !IsBluetoothA2DPInUse();
 		UE_LOG(LogVivoxVoiceChat, Verbose, TEXT("%s AEC"), bEnableAEC ? TEXT("Enabling") : TEXT("Disabling"))
-		[[IOSAppDelegate GetDelegate] EnableHighQualityVoiceChat:bEnableAEC];
 		vx_set_platform_aec_enabled(bEnableAEC ? 1 : 0);
-		// reenable voice chat to apply the changes
-		FPlatformMisc::EnableVoiceChat(true);
 	}
 }
 
diff --git a/Engine/Plugins/Online/VoiceChat/VivoxVoiceChat/Source/Private/IOS/IOSVivoxVoiceChat.h b/Engine/Plugins/Online/VoiceChat/VivoxVoiceChat/Source/Private/IOS/IOSVivoxVoiceChat.h
--- a/Engine/Plugins/Online/VoiceChat/VivoxVoiceChat/Source/Private/IOS/IOSVivoxVoiceChat.h
+++ b/Engine/Plugins/Online/VoiceChat/VivoxVoiceChat/Source/Private/IOS/IOSVivoxVoiceChat.h
@@ -57,6 +57,8 @@ private:
 	bool IsHardwareAECEnabled() const;
 	bool bEnableHardwareAEC = false;
 	TOptional<bool> OverrideEnableHardwareAEC;
+	bool bVoiceChatModeEnabled = false;
+	uint VoiceChatEnableCount = 0;
 
 	void EnableVoiceChat(bool bEnable);
 	void UpdateVoiceChatSettings();
diff --git a/Engine/Source/Runtime/ApplicationCore/Private/IOS/IOSAppDelegate.cpp b/Engine/Source/Runtime/ApplicationCore/Private/IOS/IOSAppDelegate.cpp
--- a/Engine/Source/Runtime/ApplicationCore/Private/IOS/IOSAppDelegate.cpp
+++ b/Engine/Source/Runtime/ApplicationCore/Private/IOS/IOSAppDelegate.cpp
@@ -62,7 +62,7 @@ FIOSCoreDelegates::FOnOpenURL FIOSCoreDelegates::OnOpenURL;
 FIOSCoreDelegates::FOnWillResignActive FIOSCoreDelegates::OnWillResignActive;
 TArray<FIOSCoreDelegates::FFilterDelegateAndHandle> FIOSCoreDelegates::PushNotificationFilters;
 
-static bool GEnabledAudioFeatures[(uint8)EAudioFeature::NumFeatures];
+static uint GEnabledAudioFeatures[(uint8)EAudioFeature::NumFeatures];
 
 /*
 	From: https://developer.apple.com/library/ios/documentation/UIKit/Reference/UIApplicationDelegate_Protocol/#//apple_ref/occ/intfm/UIApplicationDelegate/applicationDidEnterBackground:
@@ -116,6 +116,59 @@ void InstallSignalHandlers()
 	sigaction(SIGSYS, &Action, NULL);
 }
 
+void LexFromString(EAudioFeature& OutFeature, const TCHAR* String)
+{
+	OutFeature = EAudioFeature::NumFeatures;
+	if (FCString::Stricmp(String, TEXT("Playback")) == 0)
+	{
+		OutFeature = EAudioFeature::Playback;
+	}
+	else if (FCString::Stricmp(String, TEXT("Record")) == 0)
+	{
+		OutFeature = EAudioFeature::Record;
+	}
+	else if (FCString::Stricmp(String, TEXT("DoNotMixWithOthers")) == 0)
+	{
+		OutFeature = EAudioFeature::DoNotMixWithOthers;
+	}
+	else if (FCString::Stricmp(String, TEXT("VoiceChat")) == 0)
+	{
+		OutFeature = EAudioFeature::VoiceChat;
+	}
+	else if (FCString::Stricmp(String, TEXT("UseReceiver")) == 0)
+	{
+		OutFeature = EAudioFeature::UseReceiver;
+	}
+	else if (FCString::Stricmp(String, TEXT("DisableBluetoothSpeaker")) == 0)
+	{
+		OutFeature = EAudioFeature::DisableBluetoothSpeaker;
+	}
+	else if (FCString::Stricmp(String, TEXT("BluetoothMicrophone")) == 0)
+	{
+		OutFeature = EAudioFeature::BluetoothMicrophone;
+	}
+	else if (FCString::Stricmp(String, TEXT("BackgroundAudio")) == 0)
+	{
+		OutFeature = EAudioFeature::BackgroundAudio;
+	}
+}
+
+FString LexToString(EAudioFeature Feature)
+{
+	switch (Feature)
+	{
+	case EAudioFeature::Playback: return TEXT("Playback");
+	case EAudioFeature::Record: return TEXT("Record");
+	case EAudioFeature::DoNotMixWithOthers: return TEXT("DoNotMixWithOthers");
+	case EAudioFeature::VoiceChat: return TEXT("VoiceChat");
+	case EAudioFeature::UseReceiver: return TEXT("UseReceiver");
+	case EAudioFeature::DisableBluetoothSpeaker: return TEXT("DisableBluetoothSpeaker");
+	case EAudioFeature::BluetoothMicrophone: return TEXT("BluetoothMicrophone");
+	case EAudioFeature::BackgroundAudio: return TEXT("BackgroundAudio");
+	default: return FString();
+	}
+}
+
 FDelegateHandle FIOSCoreDelegates::AddPushNotificationFilter(const FPushNotificationFilter& FilterDel)
 {
 	FDelegateHandle NewHandle(FDelegateHandle::EGenerateNewHandleType::GenerateNewHandle);
@@ -209,9 +262,7 @@ static IOSAppDelegate* CachedDelegate = nil;
 {
 	self = [super init];
 	CachedDelegate = self;
-	// default to old style
 	memset(GEnabledAudioFeatures, 0, sizeof(GEnabledAudioFeatures));
-	self.bHighQualityVoiceChatEnabled = false;
 	return self;
 }
 
@@ -261,7 +312,21 @@ static IOSAppDelegate* CachedDelegate = nil;
 	}
 
 	FAppEntry::Init();
-	
+
+	// now that GConfig has been loaded, load the EnabledAudioFeatures from ini
+	TArray<FString> EnabledAudioFeatures;
+	GConfig->GetArray(TEXT("Audio"), TEXT("EnabledAudioFeatures"), EnabledAudioFeatures, GEngineIni);
+	for (const FString& EnabledAudioFeature : EnabledAudioFeatures)
+	{
+		EAudioFeature Feature = EAudioFeature::NumFeatures;
+		LexFromString(Feature, *EnabledAudioFeature);
+		if (Feature != EAudioFeature::NumFeatures)
+		{
+			GEnabledAudioFeatures[static_cast<uint8>(Feature)] = 1;
+		}
+	}
+	[self ToggleAudioSession:true];
+
 	[self InitIdleTimerSettings];
 
 	bEngineInit = true;
@@ -485,7 +550,7 @@ static IOSAppDelegate* CachedDelegate = nil;
 				}
 
 				FAppEntry::Resume(true);
-				[self ToggleAudioSession:true force:true];
+				[self ToggleAudioSession:true];
 				break;
 		}
 	}];
@@ -538,269 +603,105 @@ static IOSAppDelegate* CachedDelegate = nil;
 			self.bForceEmitVolume = true;
 		});
 	
-	[self ToggleAudioSession:true force:true];
+	[self ToggleAudioSession:true];
 }
 
-- (void)ToggleAudioSession:(bool)bActive force:(bool)bForce
+- (void)ToggleAudioSession:(bool)bActive
 {
 	if(!self.bAudioSessionInitialized)
 	{
 		return;
 	}
 	
-	// @todo kairos: resolve old vs new before we go to main
-	if (false)
+	// set the AVAudioSession active if necessary
+	NSError* ActiveError = nil;
+	if (bActive)
 	{
-		// we can actually override bActive based on backgrounding behavior, as that's the only time we actually deactivate the session
-		// @todo kairos: is this a valid check?
-		bool bIsBackground = GIsSuspended;
-		bActive = !bIsBackground || [self IsFeatureActive:EAudioFeature::BackgroundAudio];
-		
-		NSError* ActiveError = nil;
+		[[AVAudioSession sharedInstance] setActive:bActive error:&ActiveError];
+		if (ActiveError)
+		{
+			UE_LOG(LogIOSAudioSession, Error, TEXT("Failed to set audio session to active = %d [Error = %s]"), bActive, *FString([ActiveError description]));
+		}
+	}
 
-		// @todo maybe check the active states, not bForce?
-		if (self.bAudioActive != bActive || bForce)
+	self.bAudioActive = bActive;
+	
+	// get the category and settings to use
+	AVAudioSessionCategory Category = [self IsFeatureActive:EAudioFeature::DoNotMixWithOthers] ? AVAudioSessionCategorySoloAmbient : AVAudioSessionCategoryAmbient;
+	NSString* Mode = AVAudioSessionModeDefault;
+	AVAudioSessionCategoryOptions Options = 0;
+	if (self.bAudioActive || [self IsBackgroundAudioPlaying] || [self IsFeatureActive:EAudioFeature::BackgroundAudio])
+	{
+		// if we are running unattended, do not enable record as this would bring up a mic dialog
+		const bool bRecordActive = [self IsFeatureActive:EAudioFeature::Record] && !FApp::IsUnattended();
+		if ([self IsFeatureActive:EAudioFeature::Playback] && bRecordActive)
 		{
-			if (bActive || bForce)
+			Category = AVAudioSessionCategoryPlayAndRecord;
+
+			if ([self IsFeatureActive:EAudioFeature::VoiceChat])
 			{
-				// get the category and settings to use
-				/*AVAudioSessionCategory*/NSString* Category = self.bAllowExternalAudio ? AVAudioSessionCategoryAmbient : AVAudioSessionCategorySoloAmbient;
-				/*AVAudioSessionMode*/NSString* Mode = AVAudioSessionModeDefault;
-				AVAudioSessionCategoryOptions Options = 0;
+				Mode = AVAudioSessionModeVoiceChat;
+			}
 
-				// attempt to mix with other apps if desired
-				if ([self IsFeatureActive:EAudioFeature::BackgroundAudio] || [self IsFeatureActive:EAudioFeature::ExternalAudio])
-				{
-					Options |= AVAudioSessionCategoryOptionMixWithOthers;
-				}
+			if (![self IsFeatureActive:EAudioFeature::DoNotMixWithOthers])
+			{
+				Options |= AVAudioSessionCategoryOptionMixWithOthers;
+			}
 
-				if ([self IsFeatureActive:EAudioFeature::VoiceChat] || ([self IsFeatureActive:EAudioFeature::Playback] && [self IsFeatureActive:EAudioFeature::Record]))
-				{
-					Category = AVAudioSessionCategoryPlayAndRecord;
-					
-					if ([self IsFeatureActive:EAudioFeature::VoiceChat])
-					{
-						Mode = self.bHighQualityVoiceChatEnabled ? AVAudioSessionModeVoiceChat : AVAudioSessionModeDefault;
 #if !PLATFORM_TVOS
-						// allow bluetooth for chatting and such
-						Options |= AVAudioSessionCategoryOptionDefaultToSpeaker | AVAudioSessionCategoryOptionAllowBluetooth | AVAudioSessionCategoryOptionAllowBluetoothA2DP;
-#endif
-					}
-				}
-				else if ([self IsFeatureActive:EAudioFeature::Playback])
-				{
-					Category = AVAudioSessionCategoryPlayback;
-				}
-				else if ([self IsFeatureActive:EAudioFeature::Record])
-				{
-					Category = AVAudioSessionCategoryRecord;
-				}
-				else if ([self IsFeatureActive:EAudioFeature::BackgroundAudio])
-				{
-					Category = AVAudioSessionCategoryRecord;
-				}
-				
-				// set the category (the most important part here)
-				[[AVAudioSession sharedInstance] setCategory:Category mode:Mode options:Options error:&ActiveError];
-				if (ActiveError)
-				{
-					UE_LOG(LogIOSAudioSession, Error, TEXT("Failed to set audio session category to %s! [Error = %s]"), *FString(Category), *FString([ActiveError description]));
-				}
-			}
-			
-			// enable or disable audio
-			[[AVAudioSession sharedInstance] setActive:bActive error:&ActiveError];
-			if (ActiveError)
+			if (![self IsFeatureActive:EAudioFeature::UseReceiver])
 			{
-				UE_LOG(LogIOSAudioSession, Error, TEXT("Failed to set audio session to active = %d [Error = %s]"), bActive, *FString([ActiveError description]));
+				Options |= AVAudioSessionCategoryOptionDefaultToSpeaker;
 			}
-			else
+
+			if ([self IsFeatureActive:EAudioFeature::BluetoothMicrophone])
 			{
-				self.bAudioActive = bActive;
+				Options |= AVAudioSessionCategoryOptionAllowBluetooth;
 			}
-		}
-	}
-	// old style below
-	else if (bActive)
-	{
-        if (bForce || !self.bAudioActive)
-        {
-			bool bWasUsingBackgroundMusic = self.bUsingBackgroundMusic;
-			self.bUsingBackgroundMusic = [self IsBackgroundAudioPlaying];
-	
-			if (bWasUsingBackgroundMusic != self.bUsingBackgroundMusic || GAudio_ForceAmbientCategory)
-			{
-				if (!self.bUsingBackgroundMusic || GAudio_ForceAmbientCategory)
-				{
-					NSError* ActiveError = nil;
-					[[AVAudioSession sharedInstance] setActive:YES error:&ActiveError];
-					if (ActiveError)
-					{
-						UE_LOG(LogIOSAudioSession, Error, TEXT("Failed to set audio session as active! [Error = %s]"), *FString([ActiveError description]));
-					}
-					ActiveError = nil;
-	
-                    if (!self.bVoiceChatEnabled)
-                    {
-						if (!GAudio_ForceAmbientCategory)
-						{
-        					[[AVAudioSession sharedInstance] setCategory:AVAudioSessionCategorySoloAmbient error:&ActiveError];
-        					if (ActiveError)
-        					{
-        						UE_LOG(LogIOSAudioSession, Error, TEXT("Failed to set audio session category to AVAudioSessionCategorySoloAmbient! [Error = %s]"), *FString([ActiveError description]));
-        					}
-						}
-						else
-						{
-							[[AVAudioSession sharedInstance] setCategory:AVAudioSessionCategoryAmbient error:&ActiveError];
-							if (ActiveError)
-							{
-								UE_LOG(LogIOSAudioSession, Error, TEXT("Failed to set audio session category to AVAudioSessionCategoryAmbient! [Error = %s]"), *FString([ActiveError description]));
-							}
-						}
-    					ActiveError = nil;
-                    }
-					else
-					{
-						AVAudioSessionCategoryOptions opts =
-							AVAudioSessionCategoryOptionAllowBluetoothA2DP |
-#if !PLATFORM_TVOS
-							AVAudioSessionCategoryOptionDefaultToSpeaker |
 #endif
-							AVAudioSessionCategoryOptionMixWithOthers;
-						
-						if (@available(iOS 10, *))
-						{
-							NSString* VoiceChatMode = self.bHighQualityVoiceChatEnabled ? AVAudioSessionModeVoiceChat : AVAudioSessionModeDefault;
-							[[AVAudioSession sharedInstance] setCategory:AVAudioSessionCategoryPlayAndRecord mode:VoiceChatMode options:opts error:&ActiveError];
-						}
-						else
-						{
-							[[AVAudioSession sharedInstance] setCategory:AVAudioSessionCategoryPlayAndRecord withOptions:opts error:&ActiveError];
-						}
-						if (ActiveError)
-						{
-							UE_LOG(LogIOSAudioSession, Error, TEXT("Failed to set audio session category!"));
-						}
-						ActiveError = nil;
-					}
 
-					/* TODO::JTM - Jan 16, 2013 05:05PM - Music player support */
-				}
-				else
-				{
-					/* TODO::JTM - Jan 16, 2013 05:05PM - Music player support */
-	
-                    if (!self.bVoiceChatEnabled)
-                    {
-    					// Allow iPod music to continue playing in the background
-    					NSError* ActiveError = nil;
-    					[[AVAudioSession sharedInstance] setCategory:AVAudioSessionCategoryAmbient error:&ActiveError];
-    					if (ActiveError)
-    					{
-    						UE_LOG(LogIOSAudioSession, Error, TEXT("Failed to set audio session category to AVAudioSessionCategoryAmbient! [Error = %s]"), *FString([ActiveError description]));
-    					}
-    					ActiveError = nil;
-                    }
-				}
-			}
-			else if (!self.bUsingBackgroundMusic)
+			if (![self IsFeatureActive:EAudioFeature::DisableBluetoothSpeaker])
 			{
-				NSError* ActiveError = nil;
-				[[AVAudioSession sharedInstance] setActive:YES error:&ActiveError];
-				if (ActiveError)
-				{
-					UE_LOG(LogIOSAudioSession, Error, TEXT("Failed to set audio session as active! [Error = %s]"), *FString([ActiveError description]));
-				}
-				ActiveError = nil;
-                
-                if (!self.bVoiceChatEnabled)
-                {
-					if (!GAudio_ForceAmbientCategory)
-					{
-						[[AVAudioSession sharedInstance] setCategory:AVAudioSessionCategorySoloAmbient error:&ActiveError];
-						if (ActiveError)
-						{
-							UE_LOG(LogIOSAudioSession, Error, TEXT("Failed to set audio session category to AVAudioSessionCategorySoloAmbient! [Error = %s]"), *FString([ActiveError description]));
-						}
-					}
-					else
-					{
-						[[AVAudioSession sharedInstance] setCategory:AVAudioSessionCategoryAmbient error:&ActiveError];
-						if (ActiveError)
-						{
-							UE_LOG(LogIOSAudioSession, Error, TEXT("Failed to set audio session category to AVAudioSessionCategoryAmbient! [Error = %s]"), *FString([ActiveError description]));
-						}
-					}
-    				ActiveError = nil;
-                }
-				else
-				{
-					AVAudioSessionCategoryOptions opts =
-						AVAudioSessionCategoryOptionAllowBluetoothA2DP |
-#if !PLATFORM_TVOS
-						AVAudioSessionCategoryOptionDefaultToSpeaker |
-#endif
-						AVAudioSessionCategoryOptionMixWithOthers;
-					
-					if (@available(iOS 10, *))
-					{
-						NSString* VoiceChatMode = self.bHighQualityVoiceChatEnabled ? AVAudioSessionModeVoiceChat : AVAudioSessionModeDefault;
-						[[AVAudioSession sharedInstance] setCategory:AVAudioSessionCategoryPlayAndRecord mode:VoiceChatMode options:opts error:&ActiveError];
-					}
-					else
-					{
-						[[AVAudioSession sharedInstance] setCategory:AVAudioSessionCategoryPlayAndRecord withOptions:opts error:&ActiveError];
-					}
-					if (ActiveError)
-					{
-						UE_LOG(LogIOSAudioSession, Error, TEXT("Failed to set audio session category!"));
-					}
-					ActiveError = nil;
-				}
+				Options |= AVAudioSessionCategoryOptionAllowBluetoothA2DP;
 			}
-        }
-	}
-	else if ((bForce || self.bAudioActive) && !self.bUsingBackgroundMusic)
-	{
-		NSError* ActiveError = nil;
-        if (self.bVoiceChatEnabled)
+		}
+		else if ([self IsFeatureActive:EAudioFeature::Playback])
 		{
-			AVAudioSessionCategoryOptions opts =
-				AVAudioSessionCategoryOptionAllowBluetoothA2DP |
-#if !PLATFORM_TVOS
-				AVAudioSessionCategoryOptionDefaultToSpeaker |
-#endif
-				AVAudioSessionCategoryOptionMixWithOthers;
+			Category = AVAudioSessionCategoryPlayback;
 			
-			// Necessary for voice chat if audio is not active
-			if (@available(iOS 10, *))
-			{
-				NSString* VoiceChatMode = self.bHighQualityVoiceChatEnabled ? AVAudioSessionModeVoiceChat : AVAudioSessionModeDefault;
-				[[AVAudioSession sharedInstance] setCategory:AVAudioSessionCategoryPlayAndRecord mode:VoiceChatMode options:opts error:&ActiveError];
-			}
-			else
+			if (![self IsFeatureActive:EAudioFeature::DoNotMixWithOthers])
 			{
-				[[AVAudioSession sharedInstance] setCategory:AVAudioSessionCategoryPlayAndRecord withOptions:opts error:&ActiveError];
+				Options |= AVAudioSessionCategoryOptionMixWithOthers;
 			}
-			if (ActiveError)
+		}
+		else if (bRecordActive)
+		{
+			Category = AVAudioSessionCategoryRecord;
+
+#if !PLATFORM_TVOS
+			if ([self IsFeatureActive:EAudioFeature::BluetoothMicrophone])
 			{
-				UE_LOG(LogIOSAudioSession, Error, TEXT("Failed to set audio session category!"));
+				Options |= AVAudioSessionCategoryOptionAllowBluetooth;
 			}
-			ActiveError = nil;
+#endif
 		}
-		else
-        {
-    		// Necessary to prevent audio from getting killing when setup for background iPod audio playback
-    		[[AVAudioSession sharedInstance] setCategory:AVAudioSessionCategoryAmbient error:&ActiveError];
-    		if (ActiveError)
-    		{
-    			UE_LOG(LogIOSAudioSession, Error, TEXT("Failed to set audio session category to AVAudioSessionCategoryAmbient! [Error = %s]"), *FString([ActiveError description]));
-    		}
-    		ActiveError = nil;
-        }
- 	}
-    self.bAudioActive = bActive;
+	}
+	else
+	{
+		Category = AVAudioSessionCategoryAmbient;
+	}
+
+	// set the category if anything has changed
+	if ([[[AVAudioSession sharedInstance] category] compare:Category] != NSOrderedSame ||
+		[[[AVAudioSession sharedInstance] mode] compare:Mode] != NSOrderedSame ||
+		[[AVAudioSession sharedInstance] categoryOptions] != Options)
+	{
+		[[AVAudioSession sharedInstance] setCategory:Category mode:Mode options:Options error:&ActiveError];
+		if (ActiveError)
+		{
+			UE_LOG(LogIOSAudioSession, Error, TEXT("Failed to set audio session category to Category:%s Mode:%s Options:%x! [Error = %s]"), *FString(Category), *FString(Mode), Options, *FString([ActiveError description]));
+		}
+	}
 }
 
 - (bool)IsBackgroundAudioPlaying
@@ -821,44 +722,62 @@ static IOSAppDelegate* CachedDelegate = nil;
 
 -(void)EnableHighQualityVoiceChat:(bool)bEnable
 {
-	self.bHighQualityVoiceChatEnabled = bEnable;
+	[self SetFeature:EAudioFeature::VoiceChat Active:bEnable];
 }
 
 - (void)EnableVoiceChat:(bool)bEnable
 {
-	self.bVoiceChatEnabled = false;
-	
     // mobile will prompt for microphone access
     if (FApp::IsUnattended())
 	{
 		return;
 	}
-	self.bVoiceChatEnabled = bEnable;
-	[self ToggleAudioSession:self.bAudioActive force:true];
-	//[self SetFeature:EAudioFeature::VoiceChat Active:bEnable];
+	if (bEnable != self.bVoiceChatEnabled)
+	{
+		[self SetFeature:EAudioFeature::Playback Active:bEnable];
+		[self SetFeature:EAudioFeature::Record Active:bEnable];
+		self.bVoiceChatEnabled = bEnable;
+	}
 }
 
 - (bool)IsVoiceChatEnabled
 {
 	return self.bVoiceChatEnabled;
-	//return [self IsFeatureActive:EAudioFeature::VoiceChat];
 }
 
 
 - (void)SetFeature:(EAudioFeature)Feature Active:(bool)bIsActive
 {
-	if (GEnabledAudioFeatures[(uint8)Feature] != bIsActive)
+	if (bIsActive)
 	{
-		GEnabledAudioFeatures[(uint8)Feature] = bIsActive;
-		
-		// actually set the session
-		[self ToggleAudioSession:self.bAudioActive force:true];
+		++GEnabledAudioFeatures[static_cast<uint8>(Feature)];
+		if (GEnabledAudioFeatures[static_cast<uint8>(Feature)] == 1)
+		{
+			// Setting changed from disabled to enabled, apply the change
+			[self ToggleAudioSession:self.bAudioActive];
+		}
+	}
+	else
+	{
+		if (GEnabledAudioFeatures[static_cast<uint8>(Feature)] == 0)
+		{
+			UE_LOG(LogIOSAudioSession, Warning, TEXT("Attempted to disable audio feature that is already disabled"));
+		}
+		else
+		{
+			--GEnabledAudioFeatures[static_cast<uint8>(Feature)];
+			if (GEnabledAudioFeatures[static_cast<uint8>(Feature)] == 0)
+			{
+				// Setting changed from enabled to disabled, apply the change
+				[self ToggleAudioSession:self.bAudioActive];
+			}
+		}
 	}
 }
 
 - (bool)IsFeatureActive:(EAudioFeature)Feature
 {
-	return GEnabledAudioFeatures[(uint8)Feature];
+	return GEnabledAudioFeatures[static_cast<uint8>(Feature)] > 0;
 }
 
 
@@ -1520,7 +1439,7 @@ FCriticalSection RenderSuspend;
 		UE_LOG(LogTemp, Display, TEXT("Done with entering background tasks time."));
     }
 	[self ToggleSuspend:true];
-	[self ToggleAudioSession:false force:true];
+	[self ToggleAudioSession:false];
     
     RenderSuspend.TryLock();
     if (FTaskGraphInterface::IsRunning())
@@ -1599,7 +1518,7 @@ extern double GCStartTime;
 	 */
 	RenderSuspend.Unlock();
 	[self ToggleSuspend : false];
-    [self ToggleAudioSession:true force:true];
+    [self ToggleAudioSession:true];
 
     if (bEngineInit)
     {
diff --git a/Engine/Source/Runtime/ApplicationCore/Public/IOS/IOSAppDelegate.h b/Engine/Source/Runtime/ApplicationCore/Public/IOS/IOSAppDelegate.h
--- a/Engine/Source/Runtime/ApplicationCore/Public/IOS/IOSAppDelegate.h
+++ b/Engine/Source/Runtime/ApplicationCore/Public/IOS/IOSAppDelegate.h
@@ -17,15 +17,24 @@
 
 enum class EAudioFeature : uint8
 {
-	ExternalAudio, // background music, music app, etc
-	VoiceChat,
-	Playback,
-	Record,
-	BackgroundAudio,
-	
+	Playback, // Audio Not affected by the ringer switch
+	Record, // Recording only, unless Playback or VoiceChat is also set
+
+	DoNotMixWithOthers, // Do not mix audio with other applications
+
+	VoiceChat, // set AVAudioSessionModeVoiceChat when both Playback and Record are enabled
+	UseReceiver, // use receiver instead of speaker when both Playback and Record are enabled. Headsets will still be preferred if they are present
+	DisableBluetoothSpeaker, // disable the use of Bluetooth A2DP speakers when both Playback and Record are enabled
+
+	BluetoothMicrophone, // enable the use of Bluetooth HFP headsets when Record is enabled
+
+	BackgroundAudio, // continue to play audio in the background. Requires an appropriate background mode to be set in Info.plist
+
 	NumFeatures,
 };
 
+void LexFromString(EAudioFeature& OutFeature, const TCHAR* String);
+FString LexToString(EAudioFeature Feature);
 
 // Predicate to decide whether a push notification message should be processed
 DECLARE_DELEGATE_RetVal_OneParam(bool, FPushNotificationFilter, NSDictionary*);
@@ -213,19 +222,19 @@ APPLICATIONCORE_API
 @property (assign) bool bForceEmitVolume;
 
 - (void)InitializeAudioSession;
-- (void)ToggleAudioSession:(bool)bActive force:(bool)bForce;
+- (void)ToggleAudioSession:(bool)bActive;
 - (bool)IsBackgroundAudioPlaying;
 - (bool)HasRecordPermission;
 - (void)EnableVoiceChat:(bool)bEnable;
 - (void)EnableHighQualityVoiceChat:(bool)bEnable;
 - (bool)IsVoiceChatEnabled;
 
+/** Enable/Disable an EAudioFeature. This is reference counted, so a feature must be disabled as many times as it has been enabled to actually be disabled. */
 - (void)SetFeature:(EAudioFeature)Feature Active:(bool)bIsActive;
 - (bool)IsFeatureActive:(EAudioFeature)Mode;
 
 @property (atomic) bool bAudioActive;
 @property (atomic) bool bVoiceChatEnabled;
-@property (atomic) bool bHighQualityVoiceChatEnabled;
 
 @property (atomic) bool bIsSuspended;
 @property (atomic) bool bHasSuspended;
