Fix iOS build

diff --git a/Engine/Source/Runtime/ApplicationCore/Private/IOS/IOSAppDelegate.cpp b/Engine/Source/Runtime/ApplicationCore/Private/IOS/IOSAppDelegate.cpp
--- a/Engine/Source/Runtime/ApplicationCore/Private/IOS/IOSAppDelegate.cpp
+++ b/Engine/Source/Runtime/ApplicationCore/Private/IOS/IOSAppDelegate.cpp
@@ -575,17 +575,7 @@ static IOSAppDelegate* CachedDelegate = nil;
 
 	self.bUsingBackgroundMusic = [self IsBackgroundAudioPlaying];
 	self.bForceEmitOtherAudioPlaying = true;
-	
-	self.bAllowExternalAudio = false;
-	bool bAllowExternalAudio = false;
-	GConfig->GetBool(TEXT("/Script/IOSRuntimeSettings.IOSRuntimeSettings"), TEXT("bAudioAllowExternalAudio"), bAllowExternalAudio, GEngineIni);
-	self.bAllowExternalAudio = bAllowExternalAudio;
 
-	if(bAllowExternalAudio)
-	{
-		[self SetFeature:EAudioFeature::ExternalAudio Active:true];
-	}
-	
 #if USE_MUTE_SWITCH_DETECTION
 	// Initialize the mute switch detector.
 	[SharkfoodMuteSwitchDetector shared];
diff --git a/Engine/Source/Runtime/ApplicationCore/Public/IOS/IOSAppDelegate.h b/Engine/Source/Runtime/ApplicationCore/Public/IOS/IOSAppDelegate.h
--- a/Engine/Source/Runtime/ApplicationCore/Public/IOS/IOSAppDelegate.h
+++ b/Engine/Source/Runtime/ApplicationCore/Public/IOS/IOSAppDelegate.h
@@ -210,9 +210,6 @@ APPLICATIONCORE_API
 @property (assign) bool bLastOtherAudioPlaying;
 @property (assign) bool bForceEmitOtherAudioPlaying;
 
-/** TRUE if audio from other apps will be allowed to play while this app is active */
-@property (assign) bool bAllowExternalAudio;
-
 #if USE_MUTE_SWITCH_DETECTION
 @property (assign) bool bLastMutedState;
 @property (assign) bool bForceEmitMutedState;
