Use NSString* instead of AVAudioSessionCategory

diff --git a/Engine/Source/Runtime/ApplicationCore/Private/IOS/IOSAppDelegate.cpp b/Engine/Source/Runtime/ApplicationCore/Private/IOS/IOSAppDelegate.cpp
--- a/Engine/Source/Runtime/ApplicationCore/Private/IOS/IOSAppDelegate.cpp
+++ b/Engine/Source/Runtime/ApplicationCore/Private/IOS/IOSAppDelegate.cpp
@@ -617,7 +617,7 @@ static IOSAppDelegate* CachedDelegate = nil;
 	self.bAudioActive = bActive;
 	
 	// get the category and settings to use
-	AVAudioSessionCategory Category = [self IsFeatureActive:EAudioFeature::DoNotMixWithOthers] ? AVAudioSessionCategorySoloAmbient : AVAudioSessionCategoryAmbient;
+	NSString* Category = [self IsFeatureActive:EAudioFeature::DoNotMixWithOthers] ? AVAudioSessionCategorySoloAmbient : AVAudioSessionCategoryAmbient;
 	NSString* Mode = AVAudioSessionModeDefault;
 	AVAudioSessionCategoryOptions Options = 0;
 	if (self.bAudioActive || [self IsBackgroundAudioPlaying] || [self IsFeatureActive:EAudioFeature::BackgroundAudio])
