Fix for background crash. Reverting part of CL 5959301 to change the ApplicationWillEnterBackgroundDelegate delegate to always called on GameThread on iOS again.
#tests iOS Client

diff --git a/Engine/Source/Runtime/ApplicationCore/Private/IOS/IOSAppDelegate.cpp b/Engine/Source/Runtime/ApplicationCore/Private/IOS/IOSAppDelegate.cpp
--- a/Engine/Source/Runtime/ApplicationCore/Private/IOS/IOSAppDelegate.cpp
+++ b/Engine/Source/Runtime/ApplicationCore/Private/IOS/IOSAppDelegate.cpp
@@ -1462,7 +1462,6 @@ FCriticalSection RenderSuspend;
 	 instead of applicationWillTerminate: when the user quits.
 	 */
 
-#if BUILD_EMBEDDED_APP
 	FEmbeddedCommunication::KeepAwake(TEXT("Background"), false);
 
 	[FIOSAsyncTask CreateTaskWithBlock : ^ bool(void)
@@ -1475,9 +1474,6 @@ FCriticalSection RenderSuspend;
 		FEmbeddedCommunication::AllowSleep(TEXT("Background"));
 		return true;
 	}];
-#else
-	FCoreDelegates::ApplicationWillEnterBackgroundDelegate.Broadcast();
-#endif //BUILD_EBMEDDED_APP
 }
 
 - (void)applicationWillEnterForeground:(UIApplication *)application
