Fix GetSurveyResults not being exported.

diff --git a/Engine/Source/Runtime/ApplicationCore/Public/Unix/UnixPlatformSurvey.h b/Engine/Source/Runtime/ApplicationCore/Public/Unix/UnixPlatformSurvey.h
--- a/Engine/Source/Runtime/ApplicationCore/Public/Unix/UnixPlatformSurvey.h
+++ b/Engine/Source/Runtime/ApplicationCore/Public/Unix/UnixPlatformSurvey.h
@@ -10,12 +10,12 @@
 #include "GenericPlatform/GenericPlatformSurvey.h"
 
 /**
-* Android implementation of FGenericPlatformSurvey
+* Unix implementation of FGenericPlatformSurvey
 **/
 struct FUnixPlatformSurvey : public FGenericPlatformSurvey
 {
 	/** Start, or check on, the hardware survey */
-	static bool GetSurveyResults(FHardwareSurveyResults& OutResults, bool bWait = false);
+	APPLICATIONCORE_API static bool GetSurveyResults(FHardwareSurveyResults& OutResults, bool bWait = false);
 
 private:
 
