Fix controller v2 firmware button mapping when connected to certain phones.
#android

diff --git a/Engine/Source/Runtime/ApplicationCore/Private/Android/AndroidInputInterface.cpp b/Engine/Source/Runtime/ApplicationCore/Private/Android/AndroidInputInterface.cpp
--- a/Engine/Source/Runtime/ApplicationCore/Private/Android/AndroidInputInterface.cpp
+++ b/Engine/Source/Runtime/ApplicationCore/Private/Android/AndroidInputInterface.cpp
@@ -924,8 +924,9 @@ void FAndroidInputInterface::SendControllerEvents()
 						else if (CurrentDevice.DeviceInfo.Name.StartsWith(TEXT("PS4 Wireless Controller")))
 						{
 							CurrentDevice.ControllerClass = ControllerClassType::PS4Wireless;
-							if (CurrentDevice.DeviceInfo.Name.EndsWith(TEXT(" (v2)")))
+							if (CurrentDevice.DeviceInfo.Name.EndsWith(TEXT(" (v2)")) && FAndroidMisc::GetCPUVendor() != TEXT("Sony"))
 							{
+								// Only needed for non-Sony devices with v2 firmware
 								CurrentDevice.ButtonRemapping = ButtonRemapType::PS4;
 							}
 							CurrentDevice.bSupportsHat = true;
