Addition of iPhone 11, 11 Pro, and 11 Pro Max device profiles

diff --git a/Engine/Config/BaseDeviceProfiles.ini b/Engine/Config/BaseDeviceProfiles.ini
--- a/Engine/Config/BaseDeviceProfiles.ini
+++ b/Engine/Config/BaseDeviceProfiles.ini
@@ -36,6 +36,9 @@
 +DeviceProfileNameAndTypes=iPhoneXS,IOS
 +DeviceProfileNameAndTypes=iPhoneXSMax,IOS
 +DeviceProfileNameAndTypes=iPhoneXR,IOS
++DeviceProfileNameAndTypes=iPhone11,IOS
++DeviceProfileNameAndTypes=iPhone11Pro,IOS
++DeviceProfileNameAndTypes=iPhone11ProMax,IOS
 +DeviceProfileNameAndTypes=iPadPro,IOS
 +DeviceProfileNameAndTypes=iPadPro105,IOS
 +DeviceProfileNameAndTypes=iPadPro129,IOS
@@ -43,6 +46,7 @@
 +DeviceProfileNameAndTypes=iPadPro2_129,IOS
 +DeviceProfileNameAndTypes=iPad5,IOS
 +DeviceProfileNameAndTypes=iPad6,IOS
++DeviceProfileNameAndTypes=iPad7,IOS
 +DeviceProfileNameAndTypes=iPadPro11,IOS
 +DeviceProfileNameAndTypes=iPadPro3_129,IOS
 +DeviceProfileNameAndTypes=iPadAir3,IOS
@@ -322,6 +326,21 @@ DeviceType=IOS
 BaseProfileName=iPhone8
 +CVars=r.CustomUnsafeZones="(L:free[0,-15][896,15]);(P:fixed[104,0][206,30])"
 
+[iPhone11 DeviceProfile]
+DeviceType=IOS
+BaseProfileName=iPhone8
++CVars=r.CustomUnsafeZones="(L:free[0,-15][812,15]);(P:fixed[83,0][206,30])"
+
+[iPhone11Pro DeviceProfile]
+DeviceType=IOS
+BaseProfileName=iPhone8
++CVars=r.CustomUnsafeZones="(L:free[0,-15][812,15]);(P:fixed[83,0][206,30])"
+
+[iPhone11ProMax DeviceProfile]
+DeviceType=IOS
+BaseProfileName=iPhone8
++CVars=r.CustomUnsafeZones="(L:free[0,-15][812,15]);(P:fixed[83,0][206,30])"
+
 [iPadPro DeviceProfile]
 DeviceType=IOS
 BaseProfileName=IOS
@@ -372,6 +391,10 @@ BaseProfileName=IPadPro97
 DeviceType=IOS
 BaseProfileName=IPadPro97
 
+[iPad7 DeviceProfile]
+DeviceType=IOS
+BaseProfileName=IPadPro97
+
 [iPadPro2_129 DeviceProfile]
 DeviceType=IOS
 BaseProfileName=IPadPro
diff --git a/Engine/Source/Runtime/ApplicationCore/Private/IOS/IOSAppDelegate.cpp b/Engine/Source/Runtime/ApplicationCore/Private/IOS/IOSAppDelegate.cpp
--- a/Engine/Source/Runtime/ApplicationCore/Private/IOS/IOSAppDelegate.cpp
+++ b/Engine/Source/Runtime/ApplicationCore/Private/IOS/IOSAppDelegate.cpp
@@ -984,7 +984,7 @@ static FAutoConsoleVariableRef CVarGEnableThermalsReport(
 			[ImageString appendString : @"-Portrait"];
 		}
 	}
-    else if (Device == FPlatformMisc::IOS_IPhoneX || Device == FPlatformMisc::IOS_IPhoneXS)
+	else if (Device == FPlatformMisc::IOS_IPhoneX || Device == FPlatformMisc::IOS_IPhoneXS || Device == FPlatformMisc::IOS_IPhone11Pro)
 	{
 		[ImageString appendString : @"-IPhoneXS"];
 		if (!self.bDeviceInPortraitMode)
@@ -996,7 +996,7 @@ static FAutoConsoleVariableRef CVarGEnableThermalsReport(
 			[ImageString appendString : @"-Portrait"];
 		}
 	}
-    else if (Device == FPlatformMisc::IOS_IPhoneXSMax)
+	else if (Device == FPlatformMisc::IOS_IPhoneXSMax || Device == FPlatformMisc::IOS_IPhone11ProMax)
     {
         [ImageString appendString : @"-IPhoneXSMax"];
         if (!self.bDeviceInPortraitMode)
@@ -1008,7 +1008,7 @@ static FAutoConsoleVariableRef CVarGEnableThermalsReport(
             [ImageString appendString : @"-Portrait"];
         }
     }
-    else if (Device == FPlatformMisc::IOS_IPhoneXR)
+	else if (Device == FPlatformMisc::IOS_IPhoneXR || Device == FPlatformMisc::IOS_IPhone11)
     {
         [ImageString appendString : @"-IPhoneXR"];
         if (!self.bDeviceInPortraitMode)
diff --git a/Engine/Source/Runtime/ApplicationCore/Private/IOS/IOSPlatformApplicationMisc.cpp b/Engine/Source/Runtime/ApplicationCore/Private/IOS/IOSPlatformApplicationMisc.cpp
--- a/Engine/Source/Runtime/ApplicationCore/Private/IOS/IOSPlatformApplicationMisc.cpp
+++ b/Engine/Source/Runtime/ApplicationCore/Private/IOS/IOSPlatformApplicationMisc.cpp
@@ -250,7 +250,7 @@ void FIOSPlatformApplicationMisc::ClipboardPaste(class FString& Result)
 EScreenPhysicalAccuracy FIOSPlatformApplicationMisc::ComputePhysicalScreenDensity(int32& ScreenDensity)
 {
 	FPlatformMisc::EIOSDevice Device = FPlatformMisc::GetIOSDeviceType();
-	static_assert( FPlatformMisc::EIOSDevice::IOS_Unknown == 41, "Every device needs to be handled here." );
+	static_assert( FPlatformMisc::EIOSDevice::IOS_Unknown == 45, "Every device needs to be handled here." );
 
 	ScreenDensity = 0;
 	EScreenPhysicalAccuracy Accuracy = EScreenPhysicalAccuracy::Unknown;
@@ -273,6 +273,7 @@ EScreenPhysicalAccuracy FIOSPlatformApplicationMisc::ComputePhysicalScreenDensit
 	case FPlatformMisc::IOS_IPhone7:
 	case FPlatformMisc::IOS_IPhone8:
     case FPlatformMisc::IOS_IPhoneXR:
+	case FPlatformMisc::IOS_IPhone11:
 		ScreenDensity = 326;
 		Accuracy = EScreenPhysicalAccuracy::Truth;
 		break;
@@ -283,6 +284,8 @@ EScreenPhysicalAccuracy FIOSPlatformApplicationMisc::ComputePhysicalScreenDensit
 	case FPlatformMisc::IOS_IPhoneX:
     case FPlatformMisc::IOS_IPhoneXS:
     case FPlatformMisc::IOS_IPhoneXSMax:
+	case FPlatformMisc::IOS_IPhone11Pro:
+	case FPlatformMisc::IOS_IPhone11ProMax:
 		ScreenDensity = 401;
 		Accuracy = EScreenPhysicalAccuracy::Truth;
 		break;
@@ -298,6 +301,7 @@ EScreenPhysicalAccuracy FIOSPlatformApplicationMisc::ComputePhysicalScreenDensit
 	case FPlatformMisc::IOS_IPad4:
 	case FPlatformMisc::IOS_IPad5:
 	case FPlatformMisc::IOS_IPad6:
+	case FPlatformMisc::IOS_IPad7:
 	case FPlatformMisc::IOS_IPadAir:
 	case FPlatformMisc::IOS_IPadAir2:
     case FPlatformMisc::IOS_IPadAir3:
diff --git a/Engine/Source/Runtime/Core/Private/IOS/IOSPlatformMisc.cpp b/Engine/Source/Runtime/Core/Private/IOS/IOSPlatformMisc.cpp
--- a/Engine/Source/Runtime/Core/Private/IOS/IOSPlatformMisc.cpp
+++ b/Engine/Source/Runtime/Core/Private/IOS/IOSPlatformMisc.cpp
@@ -488,6 +488,10 @@ FIOSPlatformMisc::EIOSDevice FIOSPlatformMisc::GetIOSDeviceType()
 			{
 				DeviceType = IOS_IPad6;
 			}
+			else if (Minor == 7 || Minor == 8)
+			{
+				DeviceType = IOS_IPad7;
+			}
 			else
 			{
 				DeviceType = IOS_IPadPro2_129;
@@ -612,17 +616,32 @@ FIOSPlatformMisc::EIOSDevice FIOSPlatformMisc::GetIOSDeviceType()
                 DeviceType = IOS_IPhoneXR;
             }
         }
-		else if (Major >= 12)
+		else if (Major == 12)
+		{
+			if (Minor < 3)
+			{
+				DeviceType = IOS_IPhone11;
+			}
+			else if (Minor < 5)
+			{
+				DeviceType = IOS_IPhone11Pro;
+			}
+			else if (Minor < 7)
+			{
+				DeviceType = IOS_IPhone11ProMax;
+			}
+		}
+		else if (Major >= 13)
 		{
 			// for going forward into unknown devices (like 8/8+?), we can't use Minor,
 			// so treat devices with a scale > 2.5 to be 6SPlus type devices, < 2.5 to be 6S type devices
 			if ([UIScreen mainScreen].scale > 2.5f)
 			{
-				DeviceType = IOS_IPhoneXSMax;
+				DeviceType = IOS_IPhone11ProMax;
 			}
 			else
 			{
-				DeviceType = IOS_IPhoneXS;
+				DeviceType = IOS_IPhone11Pro;
 			}
 		}
 	}
diff --git a/Engine/Source/Runtime/Core/Public/IOS/IOSPlatformMisc.h b/Engine/Source/Runtime/Core/Public/IOS/IOSPlatformMisc.h
--- a/Engine/Source/Runtime/Core/Public/IOS/IOSPlatformMisc.h
+++ b/Engine/Source/Runtime/Core/Public/IOS/IOSPlatformMisc.h
@@ -140,12 +140,16 @@ struct CORE_API FIOSPlatformMisc : public FApplePlatformMisc
         IOS_IPhoneXS,
         IOS_IPhoneXSMax,
         IOS_IPhoneXR,
+		IOS_IPhone11,
+		IOS_IPhone11Pro,
+		IOS_IPhone11ProMax,
 		IOS_IPad6,
 		IOS_IPadPro_11,
 		IOS_IPadPro3_129,
         IOS_IPadAir3,
         IOS_IPadMini5,
 		IOS_IPodTouch7,
+		IOS_IPad7,
 		IOS_Unknown,
 	};
 
@@ -190,12 +194,16 @@ struct CORE_API FIOSPlatformMisc : public FApplePlatformMisc
             TEXT("IPhoneXS"),
             TEXT("IPhoneXSMax"),
             TEXT("IPhoneXR"),
+			TEXT("IPhone11"),
+			TEXT("IPhone11Pro"),
+			TEXT("IPhone11ProMax"),
 			TEXT("IPad6"),
 			TEXT("IPadPro11"),
 			TEXT("IPadPro3_129"),
             TEXT("IPadAir3"),
             TEXT("IPadMini5"),
 			TEXT("IPodTouch7"),
+			TEXT("IPad7"),
 			TEXT("Unknown"),
 		};
 		static_assert((sizeof(IOSDeviceNames) / sizeof(IOSDeviceNames[0])) == ((int32)IOS_Unknown + 1), "Mismatched IOSDeviceNames and EIOSDevice.");
