fix mis-resolved files

i.e. fix electric commander warnings, fix building for non-HTML5 platforms, etc.

#rnx

diff --git a/Engine/Platforms/HTML5/Plugins/WebSocketNetworking/WebSocketNetworking.uplugin b/Engine/Platforms/HTML5/Plugins/WebSocketNetworking/WebSocketNetworking.uplugin
--- a/Engine/Platforms/HTML5/Plugins/WebSocketNetworking/WebSocketNetworking.uplugin
+++ b/Engine/Platforms/HTML5/Plugins/WebSocketNetworking/WebSocketNetworking.uplugin
@@ -10,7 +10,7 @@
 	"DocsURL" : "",
 	"MarketplaceURL" : "",
 	"SupportURL" : "",
-	"EnabledByDefault" : true,
+	"EnabledByDefault" : false,
 	"CanContainContent" : false,
 	"IsBetaVersion" : false,
 	"Installed" : false,
diff --git a/Engine/Source/Runtime/ApplicationCore/ApplicationCore.Build.cs b/Engine/Source/Runtime/ApplicationCore/ApplicationCore.Build.cs
--- a/Engine/Source/Runtime/ApplicationCore/ApplicationCore.Build.cs
+++ b/Engine/Source/Runtime/ApplicationCore/ApplicationCore.Build.cs
@@ -34,7 +34,7 @@ public class ApplicationCore : ModuleRules
 			);
 			if (Target.bCompileWithAccessibilitySupport && !Target.bIsBuildingConsoleApplication)
 			{
-				PublicAdditionalLibraries.Add("uiautomationcore.lib");
+				PublicSystemLibraries.Add("uiautomationcore.lib");
 			}
 		}
 		else if (Target.Platform == UnrealTargetPlatform.Mac)
