- Add IOS specific UIWindow info to FDisplayMetrics

diff --git a/Engine/Source/Runtime/ApplicationCore/Private/IOS/IOSApplication.cpp b/Engine/Source/Runtime/ApplicationCore/Private/IOS/IOSApplication.cpp
--- a/Engine/Source/Runtime/ApplicationCore/Private/IOS/IOSApplication.cpp
+++ b/Engine/Source/Runtime/ApplicationCore/Private/IOS/IOSApplication.cpp
@@ -131,6 +131,9 @@ void FDisplayMetrics::RebuildDisplayMetrics(FDisplayMetrics& OutDisplayMetrics)
 	OutDisplayMetrics.PrimaryDisplayWidth = OutDisplayMetrics.PrimaryDisplayWorkAreaRect.Right - OutDisplayMetrics.PrimaryDisplayWorkAreaRect.Left;
 	OutDisplayMetrics.PrimaryDisplayHeight = OutDisplayMetrics.PrimaryDisplayWorkAreaRect.Bottom - OutDisplayMetrics.PrimaryDisplayWorkAreaRect.Top;
 
+	// Get ui window rect
+	OutDisplayMetrics.IosUiWindowAreaRect = FIOSWindow::GetUIWindowRect();
+
 #if !PLATFORM_TVOS
 	if (@available(iOS 11, *))
 	{
diff --git a/Engine/Source/Runtime/ApplicationCore/Private/IOS/IOSWindow.cpp b/Engine/Source/Runtime/ApplicationCore/Private/IOS/IOSWindow.cpp
--- a/Engine/Source/Runtime/ApplicationCore/Private/IOS/IOSWindow.cpp
+++ b/Engine/Source/Runtime/ApplicationCore/Private/IOS/IOSWindow.cpp
@@ -84,6 +84,22 @@ FPlatformRect FIOSWindow::GetScreenRect()
 	return ScreenRect;
 }
 
+FPlatformRect FIOSWindow::GetUIWindowRect()
+{
+	// get the main window's bounds
+	IOSAppDelegate* AppDelegate = [IOSAppDelegate GetDelegate];
+	UIWindow* Window = AppDelegate.Window;
+	CGRect Bounds = [Window bounds];
+
+	FPlatformRect WindowRect;
+	WindowRect.Top = Bounds.origin.y;
+	WindowRect.Bottom = Bounds.origin.y + Bounds.size.height;
+	WindowRect.Left = Bounds.origin.x;
+	WindowRect.Right = Bounds.origin.x + Bounds.size.width;
+
+	return WindowRect;
+}
+
 
 bool FIOSWindow::GetFullScreenInfo( int32& X, int32& Y, int32& Width, int32& Height ) const
 {
diff --git a/Engine/Source/Runtime/ApplicationCore/Private/IOS/IOSWindow.h b/Engine/Source/Runtime/ApplicationCore/Private/IOS/IOSWindow.h
--- a/Engine/Source/Runtime/ApplicationCore/Private/IOS/IOSWindow.h
+++ b/Engine/Source/Runtime/ApplicationCore/Private/IOS/IOSWindow.h
@@ -25,6 +25,9 @@ public:
 	/** Returns a rect describing the main screen */
 	static FPlatformRect GetScreenRect();
 
+	/** Returns a rect describing the main screen */
+	static FPlatformRect GetUIWindowRect();
+
 	/** Returns a void pointer to the hWnd (for other apis)*/
 	virtual void* GetOSWindowHandle() const override { return Window; }
 
diff --git a/Engine/Source/Runtime/ApplicationCore/Public/GenericPlatform/GenericApplication.h b/Engine/Source/Runtime/ApplicationCore/Public/GenericPlatform/GenericApplication.h
--- a/Engine/Source/Runtime/ApplicationCore/Public/GenericPlatform/GenericApplication.h
+++ b/Engine/Source/Runtime/ApplicationCore/Public/GenericPlatform/GenericApplication.h
@@ -367,6 +367,11 @@ struct FDisplayMetrics
 	/** Virtual display coordinate range (includes all active displays) */
 	FPlatformRect VirtualDisplayRect;
 
+#if PLATFORM_IOS
+	/** Area of the UIWindow on IOS devices */
+	FPlatformRect IosUiWindowAreaRect;
+#endif
+
 	/**
 	 * The safe area for all content on TVs (see http://en.wikipedia.org/wiki/Safe_area_%28television%29) - content will be inset 
 	 * Left - X
