Fixed up header renames for the HoloLens platform

diff --git a/Engine/Source/Runtime/ApplicationCore/Private/HoloLens/HoloLensApplication.cpp b/Engine/Source/Runtime/ApplicationCore/Private/HoloLens/HoloLensApplication.cpp
--- a/Engine/Source/Runtime/ApplicationCore/Private/HoloLens/HoloLensApplication.cpp
+++ b/Engine/Source/Runtime/ApplicationCore/Private/HoloLens/HoloLensApplication.cpp
@@ -4,14 +4,13 @@
 #include "HoloLens/HoloLensWindow.h"
 #include "HoloLens/HoloLensCursor.h"
 #include "HoloLens/HoloLensInputInterface.h"
-#include "HoloLens/HoloLensMisc.h"
+#include "HoloLens/HoloLensPlatformMisc.h"
 #include "GenericPlatform/GenericApplication.h"
 #include "Misc/CoreDelegates.h"
 #include "Misc/App.h"
-// @MIXEDREALITY_CHANGE : BEGIN
 #include "IInputDeviceModule.h"
 #include "IInputDevice.h"
-// @MIXEDREALITY_CHANGE : END
+
 DEFINE_LOG_CATEGORY(LogHoloLens);
 
 // Flip this on to make license checks operate against the retail Windows Store environment.
diff --git a/Engine/Source/Runtime/Core/Private/HoloLens/HoloLensMemory.cpp b/Engine/Source/Runtime/Core/Private/HoloLens/HoloLensMemory.cpp
--- a/Engine/Source/Runtime/Core/Private/HoloLens/HoloLensMemory.cpp
+++ b/Engine/Source/Runtime/Core/Private/HoloLens/HoloLensMemory.cpp
@@ -1,7 +1,7 @@
 // Copyright 1998-2019 Epic Games, Inc. All Rights Reserved.
 
 
-#include "HoloLensMemory.h"
+#include "HoloLensPlatformMemory.h"
 #include "HAL/MallocTBB.h"
 #include "HAL/MallocAnsi.h"
 #include "GenericPlatform/GenericPlatformMemoryPoolStats.h"
diff --git a/Engine/Source/Runtime/Core/Private/HoloLens/HoloLensMisc.cpp b/Engine/Source/Runtime/Core/Private/HoloLens/HoloLensMisc.cpp
--- a/Engine/Source/Runtime/Core/Private/HoloLens/HoloLensMisc.cpp
+++ b/Engine/Source/Runtime/Core/Private/HoloLens/HoloLensMisc.cpp
@@ -4,7 +4,7 @@
 	HoloLensMisc.cpp: HoloLens implementations of misc functions
 =============================================================================*/
 
-#include "HoloLensMisc.h"
+#include "HoloLensPlatformMisc.h"
 #include "Internationalization/Internationalization.h"
 #include "Misc/MessageDialog.h"
 #include "Misc/App.h"
diff --git a/Engine/Source/Runtime/Core/Private/HoloLens/HoloLensOutputDevices.cpp b/Engine/Source/Runtime/Core/Private/HoloLens/HoloLensOutputDevices.cpp
--- a/Engine/Source/Runtime/Core/Private/HoloLens/HoloLensOutputDevices.cpp
+++ b/Engine/Source/Runtime/Core/Private/HoloLens/HoloLensOutputDevices.cpp
@@ -4,7 +4,7 @@
 	HoloLensOutputDevices.cpp: HoloLens implementations of OutputDevices functions
 =============================================================================*/
 
-#include "HoloLens/HoloLensOutputDevices.h"
+#include "HoloLens/HoloLensPlatformOutputDevices.h"
 #include "Misc/App.h"
 #include "Misc/CoreDelegates.h"
 
diff --git a/Engine/Source/Runtime/Core/Private/HoloLens/HoloLensProcess.cpp b/Engine/Source/Runtime/Core/Private/HoloLens/HoloLensProcess.cpp
--- a/Engine/Source/Runtime/Core/Private/HoloLens/HoloLensProcess.cpp
+++ b/Engine/Source/Runtime/Core/Private/HoloLens/HoloLensProcess.cpp
@@ -4,7 +4,7 @@
 HoloLensProcess.cpp: HoloLens implementations of Process functions
 =============================================================================*/
 
-#include "HoloLensProcess.h"
+#include "HoloLensPlatformProcess.h"
 #include "Misc/SingleThreadEvent.h"
 #include "HoloLensRunnableThread.h"
 #include "Misc/CommandLine.h"
diff --git a/Engine/Source/Runtime/Core/Private/HoloLens/HoloLensRunnableThread.h b/Engine/Source/Runtime/Core/Private/HoloLens/HoloLensRunnableThread.h
--- a/Engine/Source/Runtime/Core/Private/HoloLens/HoloLensRunnableThread.h
+++ b/Engine/Source/Runtime/Core/Private/HoloLens/HoloLensRunnableThread.h
@@ -8,7 +8,7 @@
 #include "HAL/Event.h"
 #include "Containers/StringConv.h"
 #include "HoloLens/WindowsHWrapper.h"
-#include "HoloLens/HoloLensProcess.h"
+#include "HoloLens/HoloLensPlatformProcess.h"
 #include "HoloLens/AllowWindowsPlatformTypes.h"
 
 /**
diff --git a/Engine/Source/Runtime/Core/Private/HoloLens/HoloLensStackWalk.cpp b/Engine/Source/Runtime/Core/Private/HoloLens/HoloLensStackWalk.cpp
--- a/Engine/Source/Runtime/Core/Private/HoloLens/HoloLensStackWalk.cpp
+++ b/Engine/Source/Runtime/Core/Private/HoloLens/HoloLensStackWalk.cpp
@@ -4,5 +4,5 @@
 	HoloLensStackWalk.cpp: HoloLens implementations of stack walk functions
 =============================================================================*/
 
-#include "HoloLensStackWalk.h"
+#include "HoloLensPlatformStackWalk.h"
 
diff --git a/Engine/Source/Runtime/Core/Private/HoloLens/HoloLensTime.cpp b/Engine/Source/Runtime/Core/Private/HoloLens/HoloLensTime.cpp
--- a/Engine/Source/Runtime/Core/Private/HoloLens/HoloLensTime.cpp
+++ b/Engine/Source/Runtime/Core/Private/HoloLens/HoloLensTime.cpp
@@ -5,7 +5,7 @@ HoloLensTime.cpp: HoloLens implementations of time functions
 =============================================================================*/
 
 //#include "CorePrivatePCH.h"
-#include "HoloLensTime.h"
+#include "HoloLensPlatformTime.h"
 #include "Misc/AssertionMacros.h"
 #include "Containers/Ticker.h"
 #include "HoloLens/WindowsHWrapper.h"
