Rename to HoloLensPlatformSplash

diff --git a/Engine/Source/Runtime/ApplicationCore/Private/HoloLens/HoloLensPlatformSplash.cpp b/Engine/Source/Runtime/ApplicationCore/Private/HoloLens/HoloLensPlatformSplash.cpp
new file mode 100644
--- /dev/null
+++ b/Engine/Source/Runtime/ApplicationCore/Private/HoloLens/HoloLensPlatformSplash.cpp
@@ -0,0 +1,21 @@
+// Copyright 1998-2019 Epic Games, Inc. All Rights Reserved.
+
+#include "HoloLens/HoloLensPlatformSplash.h"
+#include "HoloLens/HoloLensPlatformApplicationMisc.h"
+
+void FHoloLensSplash::Show()
+{
+	[REDACTED_EMAIL]: Implement me
+	FHoloLensPlatformApplicationMisc::PumpMessages(true);
+}
+
+void FHoloLensSplash::Hide()
+{
+	[REDACTED_EMAIL]: Implement me
+	FHoloLensPlatformApplicationMisc::PumpMessages(true);
+}
+
+void FHoloLensSplash::SetSplashText( const SplashTextType::Type InType, const TCHAR* InText )
+{
+	[REDACTED_EMAIL]: Implement me
+}
diff --git a/Engine/Source/Runtime/ApplicationCore/Private/HoloLens/HoloLensSplash.cpp b/Engine/Source/Runtime/ApplicationCore/Private/HoloLens/HoloLensSplash.cpp
deleted file mode 100644
diff --git a/Engine/Source/Runtime/ApplicationCore/Public/HoloLens/HoloLensPlatformSplash.h b/Engine/Source/Runtime/ApplicationCore/Public/HoloLens/HoloLensPlatformSplash.h
new file mode 100644
--- /dev/null
+++ b/Engine/Source/Runtime/ApplicationCore/Public/HoloLens/HoloLensPlatformSplash.h
@@ -0,0 +1,29 @@
+// Copyright 1998-2019 Epic Games, Inc. All Rights Reserved.
+
+#pragma once
+#include "GenericPlatform/GenericPlatformSplash.h"
+
+/**
+ * HoloLens splash implementation
+ */
+struct CORE_API FHoloLensSplash : public FGenericPlatformSplash
+{
+	/**
+	* Show the splash screen
+	*/
+	static void Show();
+	/**
+	* Hide the splash screen
+	*/
+	static void Hide();
+
+	/**
+	 * Sets the text displayed on the splash screen (for startup/loading progress)
+	 *
+	 * @param	InType		Type of text to change
+	 * @param	InText		Text to display
+	 */
+	static void SetSplashText( const SplashTextType::Type InType, const TCHAR* InText );
+};
+
+typedef FHoloLensSplash FPlatformSplash;
diff --git a/Engine/Source/Runtime/ApplicationCore/Public/HoloLens/HoloLensSplash.h b/Engine/Source/Runtime/ApplicationCore/Public/HoloLens/HoloLensSplash.h
deleted file mode 100644
