Added stubbed versions of CreateCursorFrom*() methods

diff --git a/Engine/Source/Runtime/ApplicationCore/Public/HoloLens/HoloLensCursor.h b/Engine/Source/Runtime/ApplicationCore/Public/HoloLens/HoloLensCursor.h
--- a/Engine/Source/Runtime/ApplicationCore/Public/HoloLens/HoloLensCursor.h
+++ b/Engine/Source/Runtime/ApplicationCore/Public/HoloLens/HoloLensCursor.h
@@ -58,6 +58,17 @@ public:
 	*/
 	void SetTypeShape(EMouseCursor::Type InCursorType, void* CursorHandle);
 
+	virtual void* CreateCursorFromFile(const FString& InPathToCursorWithoutExtension, FVector2D HotSpot) override
+	{
+		return nullptr;
+	}
+
+	/** Creates a hardware cursor from bitmap data. Can return nullptr when not available. */
+	virtual void* CreateCursorFromRGBABuffer(const FColor* Pixels, int32 Width, int32 Height, FVector2D InHotSpot) override
+	{
+		return nullptr;
+	}
+
 private:
 
 	void SetUseRawMouse(bool bUse);
