Fixed pre-processor error in FAndroidWindow.

#rnx

diff --git a/Engine/Source/Runtime/ApplicationCore/Private/Android/AndroidWindow.cpp b/Engine/Source/Runtime/ApplicationCore/Private/Android/AndroidWindow.cpp
--- a/Engine/Source/Runtime/ApplicationCore/Private/Android/AndroidWindow.cpp
+++ b/Engine/Source/Runtime/ApplicationCore/Private/Android/AndroidWindow.cpp
@@ -166,7 +166,7 @@ void FAndroidWindow::ReleaseWindowRef(ANativeWindow* InWindow)
 {
 #if USE_ANDROID_EVENTS
 	check(IsInAndroidEventThread());
-#endif
+#else
 	NativeWindow = InWindow; //using raw native window handle for now. Could be changed to use AndroidWindow later if needed
 #endif
 }
