Remove OpenGL ES support on IOS

diff --git a/Engine/Source/Programs/UnrealBuildTool/Platform/IOS/UEBuildIOS.cs b/Engine/Source/Programs/UnrealBuildTool/Platform/IOS/UEBuildIOS.cs
--- a/Engine/Source/Programs/UnrealBuildTool/Platform/IOS/UEBuildIOS.cs
+++ b/Engine/Source/Programs/UnrealBuildTool/Platform/IOS/UEBuildIOS.cs
@@ -1,4 +1,4 @@
-﻿// Copyright 1998-2019 Epic Games, Inc. All Rights Reserved.
+// Copyright 1998-2019 Epic Games, Inc. All Rights Reserved.
 
 using System;
 using System.Collections.Generic;
@@ -1114,8 +1114,6 @@ namespace UnrealBuildTool
 				CompileEnvironment.Definitions.Add("ENABLE_ADVERTISING_IDENTIFIER=1");
 			}
 
-			CompileEnvironment.Definitions.Add("HAS_OPENGL_ES=" + ((Target.IOSPlatform.RuntimeVersion < 12.0) ? "1" : "0"));
-
 			// if the project has an Oodle compression Dll, enable the decompressor on IOS
 			if (Target.ProjectFile != null)
 			{
diff --git a/Engine/Source/Runtime/ApplicationCore/Private/IOS/IOSView.cpp b/Engine/Source/Runtime/ApplicationCore/Private/IOS/IOSView.cpp
--- a/Engine/Source/Runtime/ApplicationCore/Private/IOS/IOSView.cpp
+++ b/Engine/Source/Runtime/ApplicationCore/Private/IOS/IOSView.cpp
@@ -116,8 +116,8 @@ id<MTLDevice> GMetalDevice = nil;
 
 
 
-#if !HAS_METAL && !HAS_OPENGL_ES
-#error One of HAS_METAL or HAS_OPENGL_ES must be defined
+#if !HAS_METAL
+#error HAS_METAL must be defined
 #endif
 
 /**
@@ -172,11 +172,7 @@ id<MTLDevice> GMetalDevice = nil;
 	else
 #endif
 	{
-#if HAS_OPENGL_ES
-		return [CAEAGLLayer class];
-#else
 		return nil;
-#endif
 	}
 }
 
@@ -207,31 +203,7 @@ id<MTLDevice> GMetalDevice = nil;
 		MetalLayer.framebufferOnly = NO;
 		
 	}
-	else
 #endif
-	{
-#if HAS_OPENGL_ES
-		bIsUsingMetal = false;
-		
-		// Get the layer
-		CAEAGLLayer *EaglLayer = (CAEAGLLayer *)self.layer;
-		EaglLayer.opaque = YES;
-		NSMutableDictionary* Dict = [NSMutableDictionary dictionary];
-		[Dict setValue : [NSNumber numberWithBool : NO] forKey : kEAGLDrawablePropertyRetainedBacking];
-		[Dict setValue : kEAGLColorFormatRGBA8 forKey : kEAGLDrawablePropertyColorFormat];
-		EaglLayer.drawableProperties = Dict;
-		
-		// Initialize a single, static OpenGL ES 2.0 context, shared by all EAGLView objects
-		Context = [[EAGLContext alloc] initWithAPI:kEAGLRenderingAPIOpenGLES2];
-		
-		// delete this on failure
-		if (!Context || ![EAGLContext setCurrentContext : Context])
-		{
-			[self release];
-			return nil;
-		}
-#endif
-	}
 	
 	NSLog(@"::: Created a UIView that will support %@ :::", bIsUsingMetal ? @"Metal" : @"@GLES");
 	
@@ -348,54 +320,7 @@ id<MTLDevice> GMetalDevice = nil;
 			DrawableSize.height *= self.contentScaleFactor;
 			MetalLayer.drawableSize = DrawableSize;
 		}
-		else
-#endif
-		{
-#if HAS_OPENGL_ES
-			// make sure this is current
-			[self MakeCurrent];
-
-			// This is causing a pretty large slow down on iPad 3 and 4 using iOS 6, for now going to comment it out
-	//		if (self.contentScaleFactor == 1.0f || self.contentScaleFactor == 2.0f)
-	//		{
-	//			UE_LOG(LogIOS,Log,TEXT("Setting layer filter to NEAREST"));
-	//			CAEAGLLayer *EaglLayer = (CAEAGLLayer *)self.layer;
-	//			EaglLayer.magnificationFilter = kCAFilterNearest;
-	//		}
-
-			// Create our standard displayable surface
-			glGenRenderbuffers(1, &OnScreenColorRenderBuffer);
-			check(glGetError() == 0);
-			glBindRenderbuffer(GL_RENDERBUFFER, OnScreenColorRenderBuffer);
-			check(glGetError() == 0);
-			[Context renderbufferStorage:GL_RENDERBUFFER fromDrawable:(CAEAGLLayer*)self.layer];
-
-			// Get the size of the surface
-			GLint OnScreenWidth, OnScreenHeight;
-			glGetRenderbufferParameteriv(GL_RENDERBUFFER, GL_RENDERBUFFER_WIDTH, &OnScreenWidth);
-			glGetRenderbufferParameteriv(GL_RENDERBUFFER, GL_RENDERBUFFER_HEIGHT, &OnScreenHeight);
-
-			// NOTE: This resolve FBO is necessary even if we don't plan on using MSAA because otherwise
-			// the shaders will not warm properly. Future investigation as to why; it seems unnecessary.
-
-			// Create an FBO used to target the resolve surface
-			glGenFramebuffers(1, &ResolveFrameBuffer);
-			check(glGetError() == 0);
-			glBindFramebuffer(GL_FRAMEBUFFER, ResolveFrameBuffer);
-			check(glGetError() == 0);
-			glFramebufferRenderbuffer(GL_FRAMEBUFFER, GL_COLOR_ATTACHMENT0, GL_RENDERBUFFER, OnScreenColorRenderBuffer);
-			check(glGetError() == 0);
-			check(glCheckFramebufferStatus(GL_FRAMEBUFFER) == GL_FRAMEBUFFER_COMPLETE);
-
-#if defined(USE_DETAILED_IPHONE_MEM_TRACKING) && USE_DETAILED_IPHONE_MEM_TRACKING
-			//This value is used to allow the engine to track gl allocated memory see GetIPhoneOpenGLBackBufferSize
-			uint32 SingleBufferSize = OnScreenWidth * OnScreenHeight * 4/*rgba8*/;
-			IPhoneBackBufferMemSize = singleBufferSize * 3/*iphone back buffer system is tripple buffered*/;
-
-			UE_LOG(LogIOS, Log, TEXT("IPhone Back Buffer Size: %i MB"), (IPhoneBackBufferMemSize / 1024) / 1024.f);
-#endif
 #endif
-		}
 
 		bIsInitialized = true;
 	}    
@@ -428,17 +353,6 @@ id<MTLDevice> GMetalDevice = nil;
 	}
 #endif
 
-#if HAS_OPENGL_ES
-	// Allocate color buffer based on the current layer size
-	glBindRenderbuffer(GL_RENDERBUFFER, OnScreenColorRenderBuffer);
-	check(glGetError() == 0);
-	[Context renderbufferStorage:GL_RENDERBUFFER fromDrawable:(CAEAGLLayer*)self.layer];
-	glBindFramebuffer(GL_FRAMEBUFFER, ResolveFrameBuffer);
-	check(glGetError() == 0);
-	glFramebufferRenderbuffer(GL_FRAMEBUFFER, GL_COLOR_ATTACHMENT0, GL_RENDERBUFFER, OnScreenColorRenderBuffer);
-	check(glGetError() == 0);
-	check(glCheckFramebufferStatus(GL_FRAMEBUFFER) == GL_FRAMEBUFFER_COMPLETE);
-#endif
 }
 
 #if HAS_METAL
@@ -452,32 +366,6 @@ id<MTLDevice> GMetalDevice = nil;
 {
 	if (bIsInitialized)
 	{
-		// nothing to do here for Metal
-		if (!bIsUsingMetal)
-		{
-#if HAS_OPENGL_ES
-			// toss framebuffers
-			if(ResolveFrameBuffer)
-			{
-				glDeleteFramebuffers(1, &ResolveFrameBuffer);
-				ResolveFrameBuffer = 0;
-			}
-			if(OnScreenColorRenderBuffer)
-			{
-				glDeleteRenderbuffers(1, &OnScreenColorRenderBuffer);
-				OnScreenColorRenderBuffer = 0;
-			}
-// 			if( GMSAAAllowed )
-// 			{
-// 				if(OnScreenColorRenderBufferMSAA)
-// 				{
-// 					glDeleteRenderbuffers(1, &OnScreenColorRenderBufferMSAA);
-// 					OnScreenColorRenderBufferMSAA = 0;
-// 				}
-//
-#endif
- 		}
-
 		// we are ready to be re-initialized
 		bIsInitialized = false;
 	}
@@ -485,64 +373,14 @@ id<MTLDevice> GMetalDevice = nil;
 
 - (void)MakeCurrent
 {
-	if (!bIsUsingMetal)
-	{
-#if HAS_OPENGL_ES
-		[EAGLContext setCurrentContext:Context];
-#endif
-	}
 }
 
 - (void)UnmakeCurrent
 {
-	if (!bIsUsingMetal)
-	{
-#if HAS_OPENGL_ES
-		[EAGLContext setCurrentContext:nil];
-#endif
-	}
 }
 
 - (void)SwapBuffers
 {
-	if (!bIsUsingMetal)
-	{
-#if HAS_OPENGL_ES
-// 		// We may need this in the MSAA case
-// 		GLint CurrentFramebuffer = 0;
-// 		// @todo-mobile: Fix this when we have MSAA support
-// 		if( GMSAAAllowed && GMSAAEnabled )
-// 		{
-// 			// Get the currently bound FBO
-// 			glGetIntegerv(GL_FRAMEBUFFER_BINDING, &CurrentFramebuffer);
-// 
-// 			// Set up and perform the resolve (the READ is already set)
-// 			glBindFramebuffer(GL_DRAW_FRAMEBUFFER_APPLE, ResolveFrameBuffer);
-// 			glResolveMultisampleFramebufferAPPLE();
-// 
-// 			// After the resolve, we can discard the old attachments
-// 			GLenum attachments[] = { GL_COLOR_ATTACHMENT0, GL_DEPTH_ATTACHMENT, GL_STENCIL_ATTACHMENT };
-// 			glDiscardFramebufferEXT( GL_READ_FRAMEBUFFER_APPLE, 3, attachments );
-// 		}
-// 		else
-// 		{
-// 			// Discard the now-unncessary depth buffer
-// 			GLenum attachments[] = { GL_DEPTH_ATTACHMENT, GL_STENCIL_ATTACHMENT };
-// 			glDiscardFramebufferEXT( GL_READ_FRAMEBUFFER_APPLE, 2, attachments );
-// 		}
-// 
-		// Perform the actual present with the on-screen renderbuffer
-		//glBindRenderbuffer(GL_RENDERBUFFER, OnScreenColorRenderBuffer);
-		[Context presentRenderbuffer:GL_RENDERBUFFER];
-
-// 		if( GMSAAAllowed && GMSAAEnabled )
-// 		{
-// 			// Restore the DRAW framebuffer object
-// 			glBindFramebuffer(GL_DRAW_FRAMEBUFFER_APPLE, CurrentFramebuffer);
-// 		}
-#endif
-	}
-
 	// increment our swap counter
 	SwapCount++;
 }
diff --git a/Engine/Source/Runtime/ApplicationCore/Public/IOS/IOSView.h b/Engine/Source/Runtime/ApplicationCore/Public/IOS/IOSView.h
--- a/Engine/Source/Runtime/ApplicationCore/Public/IOS/IOSView.h
+++ b/Engine/Source/Runtime/ApplicationCore/Public/IOS/IOSView.h
@@ -51,15 +51,6 @@ APPLICATIONCORE_API
 	bool HasMoved[10];
 
 
-#if HAS_OPENGL_ES
-	//// GL MEMBERS
-	// the GL context
-	EAGLContext* Context;
-
-	// the internal MSAA FBO used to resolve the color buffer at present-time
-	GLuint ResolveFrameBuffer;
-#endif
-
 	//// METAL MEMBERS
 #if HAS_METAL
 	// global metal device
