fix for recreating the abnormal termination file when receiving a low memory warning while in the background

diff --git a/Engine/Source/Runtime/ApplicationCore/Private/IOS/IOSAppDelegate.cpp b/Engine/Source/Runtime/ApplicationCore/Private/IOS/IOSAppDelegate.cpp
--- a/Engine/Source/Runtime/ApplicationCore/Private/IOS/IOSAppDelegate.cpp
+++ b/Engine/Source/Runtime/ApplicationCore/Private/IOS/IOSAppDelegate.cpp
@@ -60,6 +60,7 @@ extern bool GShowSplashScreen;
 
 FIOSCoreDelegates::FOnOpenURL FIOSCoreDelegates::OnOpenURL;
 FIOSCoreDelegates::FOnWillResignActive FIOSCoreDelegates::OnWillResignActive;
+FIOSCoreDelegates::FOnDidBecomeActive FIOSCoreDelegates::OnDidBecomeActive;
 TArray<FIOSCoreDelegates::FFilterDelegateAndHandle> FIOSCoreDelegates::PushNotificationFilters;
 
 static uint GEnabledAudioFeatures[(uint8)EAudioFeature::NumFeatures];
@@ -1398,8 +1399,6 @@ extern EDeviceScreenOrientation ConvertFromUIInterfaceOrientation(UIInterfaceOri
 FCriticalSection RenderSuspend;
 - (void)applicationWillResignActive:(UIApplication *)application
 {
-	FIOSCoreDelegates::OnWillResignActive.Broadcast();
-	
     FIOSPlatformMisc::ResetBrightness();
     
     /*
@@ -1458,6 +1457,8 @@ FCriticalSection RenderSuspend;
             }, TStatId(), NULL, ENamedThreads::ActualRenderingThread);
         }
     }
+	
+	FIOSCoreDelegates::OnWillResignActive.Broadcast();
 }
 
 - (void)applicationDidEnterBackground:(UIApplication *)application
@@ -1506,7 +1507,9 @@ FCriticalSection RenderSuspend;
 extern double GCStartTime;
 - (void)applicationDidBecomeActive:(UIApplication *)application
 {
-    // make sure a GC will not timeout because it was started before entering background
+	FIOSCoreDelegates::OnDidBecomeActive.Broadcast();
+
+	// make sure a GC will not timeout because it was started before entering background
     GCStartTime = FPlatformTime::Seconds();
 	/*
 	 Restart any tasks that were paused (or not yet started) while the application was inactive. If the application was previously in the background, optionally refresh the user interface.
diff --git a/Engine/Source/Runtime/ApplicationCore/Public/IOS/IOSAppDelegate.h b/Engine/Source/Runtime/ApplicationCore/Public/IOS/IOSAppDelegate.h
--- a/Engine/Source/Runtime/ApplicationCore/Public/IOS/IOSAppDelegate.h
+++ b/Engine/Source/Runtime/ApplicationCore/Public/IOS/IOSAppDelegate.h
@@ -59,6 +59,10 @@ public:
 	DECLARE_MULTICAST_DELEGATE(FOnWillResignActive);
 	static FOnWillResignActive OnWillResignActive;
 
+	/** INTERNAL - called when becoming active - this is not thread-safe with the game thread or render thread as it is called from the app's Main thread */
+	DECLARE_MULTICAST_DELEGATE(FOnDidBecomeActive);
+	static FOnWillResignActive OnDidBecomeActive;
+	
 private:
 	struct FFilterDelegateAndHandle
 	{
