Making threshold trigger button pressed events optional by device type. Some controllers do not need this functionality and will send double input events.

diff --git a/Engine/Source/Runtime/ApplicationCore/Private/Android/AndroidInputInterface.cpp b/Engine/Source/Runtime/ApplicationCore/Private/Android/AndroidInputInterface.cpp
--- a/Engine/Source/Runtime/ApplicationCore/Private/Android/AndroidInputInterface.cpp
+++ b/Engine/Source/Runtime/ApplicationCore/Private/Android/AndroidInputInterface.cpp
@@ -854,6 +854,7 @@ void FAndroidInputInterface::SendControllerEvents()
 						CurrentDevice.ButtonRemapping = ButtonRemapType::Normal;
 						CurrentDevice.LTAnalogRangeMinimum = 0.0f;
 						CurrentDevice.RTAnalogRangeMinimum = 0.0f;
+						CurrentDevice.bTriggersUseThresholdForClick = false;
 						CurrentDevice.bSupportsHat = true;
 						CurrentDevice.bMapL1R1ToTriggers = false;
 						CurrentDevice.bMapZRZToTriggers = false;
@@ -895,11 +896,13 @@ void FAndroidInputInterface::SendControllerEvents()
 						{
 							CurrentDevice.ControllerClass = ControllerClassType::XBoxWired;
 							CurrentDevice.bSupportsHat = true;
+							CurrentDevice.bTriggersUseThresholdForClick = true;
 						}
 						else if (CurrentDevice.DeviceInfo.Name.StartsWith(TEXT("Xbox Wireless Controller")))
 						{
 							CurrentDevice.ControllerClass = ControllerClassType::XBoxWireless;
 							CurrentDevice.bSupportsHat = true;
+							CurrentDevice.bTriggersUseThresholdForClick = true;
 
 							if (GAndroidOldXBoxWirelessFirmware == 1)
 							{
@@ -1045,23 +1048,31 @@ void FAndroidInputInterface::SendControllerEvents()
 				NewControllerState.ButtonStates[MAX_NUM_PHYSICAL_CONTROLLER_BUTTONS + 6] = NewControllerState.RYAnalog >= RepeatDeadzone;
 				NewControllerState.ButtonStates[MAX_NUM_PHYSICAL_CONTROLLER_BUTTONS + 7] = NewControllerState.RYAnalog <= -RepeatDeadzone;
 			}
+			
+			const bool bUseTriggerThresholdForClick = DeviceMapping[ControllerIndex].bTriggersUseThresholdForClick;
 			if (NewControllerState.LTAnalog != OldControllerState.LTAnalog)
 			{
 				//LOGD("    Sending updated LeftTriggerAnalog value of %f", NewControllerState.LTAnalog);
 				MessageHandler->OnControllerAnalog(FGamepadKeyNames::LeftTriggerAnalog, NewControllerState.DeviceId, NewControllerState.LTAnalog);
 
-				// Handle the trigger theshold "virtual" button state
-				//check(ButtonMapping[10] == FGamepadKeyNames::LeftTriggerThreshold);
-				NewControllerState.ButtonStates[10] = NewControllerState.LTAnalog >= ANDROID_GAMEPAD_TRIGGER_THRESHOLD;
+				if (bUseTriggerThresholdForClick)
+				{
+					// Handle the trigger theshold "virtual" button state
+					//check(ButtonMapping[10] == FGamepadKeyNames::LeftTriggerThreshold);
+					NewControllerState.ButtonStates[10] = NewControllerState.LTAnalog >= ANDROID_GAMEPAD_TRIGGER_THRESHOLD;
+				}
 			}
 			if (NewControllerState.RTAnalog != OldControllerState.RTAnalog)
 			{
 				//LOGD("    Sending updated RightTriggerAnalog value of %f", NewControllerState.RTAnalog);
 				MessageHandler->OnControllerAnalog(FGamepadKeyNames::RightTriggerAnalog, NewControllerState.DeviceId, NewControllerState.RTAnalog);
 
-				// Handle the trigger theshold "virtual" button state
-				//check(ButtonMapping[11] == FGamepadKeyNames::RightTriggerThreshold);
-				NewControllerState.ButtonStates[11] = NewControllerState.RTAnalog >= ANDROID_GAMEPAD_TRIGGER_THRESHOLD;
+				if (bUseTriggerThresholdForClick)
+				{
+					// Handle the trigger theshold "virtual" button state
+					//check(ButtonMapping[11] == FGamepadKeyNames::RightTriggerThreshold);
+					NewControllerState.ButtonStates[11] = NewControllerState.RTAnalog >= ANDROID_GAMEPAD_TRIGGER_THRESHOLD;
+				}
 			}
 
 			const double CurrentTime = FPlatformTime::Seconds();
diff --git a/Engine/Source/Runtime/ApplicationCore/Public/Android/AndroidInputInterface.h b/Engine/Source/Runtime/ApplicationCore/Public/Android/AndroidInputInterface.h
--- a/Engine/Source/Runtime/ApplicationCore/Public/Android/AndroidInputInterface.h
+++ b/Engine/Source/Runtime/ApplicationCore/Public/Android/AndroidInputInterface.h
@@ -141,6 +141,9 @@ struct FAndroidGamepadDeviceMapping
 	// Device supports hat as dpad
 	bool bSupportsHat;
 
+	// Device uses threshold to send button pressed events.
+	bool bTriggersUseThresholdForClick;
+
 	// Map L1 and R1 to LTRIGGER and RTRIGGER
 	bool bMapL1R1ToTriggers;
 
