Adding support for iOS thumbsticks

diff --git a/Engine/Source/Runtime/ApplicationCore/Private/IOS/IOSInputInterface.cpp b/Engine/Source/Runtime/ApplicationCore/Private/IOS/IOSInputInterface.cpp
--- a/Engine/Source/Runtime/ApplicationCore/Private/IOS/IOSInputInterface.cpp
+++ b/Engine/Source/Runtime/ApplicationCore/Private/IOS/IOSInputInterface.cpp
@@ -12,6 +12,13 @@
 
 DECLARE_LOG_CATEGORY_EXTERN(LogIOSInput, Log, All);
 
+@interface GCExtendedGamepad()
+#if (__IPHONE_OS_VERSION_MAX_ALLOWED < 121000 || __APPLETV_OS_VERSION_MAX_ALLOWED < 121000 || __MAC_OS_VERSION_MAX_ALLOWED < 1401000)
+@property (nonatomic, readwrite, nullable) GCControllerButtonInput *leftThumbstickButton;
+@property (nonatomic, readwrite, nullable) GCControllerButtonInput *rightThumbstickButton;
+#endif
+@end
+
 
 static TAutoConsoleVariable<float> CVarHapticsKickHeavy(TEXT("ios.VibrationHapticsKickHeavyValue"), 0.65f, TEXT("Vibation values higher than this will kick a haptics heavy Impact"));
 static TAutoConsoleVariable<float> CVarHapticsKickMedium(TEXT("ios.VibrationHapticsKickMediumValue"), 0.5f, TEXT("Vibation values higher than this will kick a haptics medium Impact"));
@@ -451,6 +458,12 @@ if ((Controller.Previous##Gamepad != nil && Gamepad.GCAxis.value != Controller.P
 			HANDLE_ANALOG_VIRTUAL_BUTTONS(ExtendedGamepad, rightThumbstick.xAxis, FGamepadKeyNames::RightStickLeft, FGamepadKeyNames::RightStickRight);
 			HANDLE_ANALOG_VIRTUAL_BUTTONS(ExtendedGamepad, rightThumbstick.yAxis, FGamepadKeyNames::RightStickDown, FGamepadKeyNames::RightStickUp);
 
+			if(ExtendedGamepad.leftThumbstickButton != nil)
+			{
+				HANDLE_BUTTON(ExtendedGamepad, leftThumbstickButton,        FGamepadKeyNames::LeftThumb);
+				HANDLE_BUTTON(ExtendedGamepad, rightThumbstickButton,        FGamepadKeyNames::RightThumb);
+			}
+
 			[Controller.PreviousExtendedGamepad release];
 			Controller.PreviousExtendedGamepad = [ExtendedGamepad saveSnapshot];
 			[Controller.PreviousExtendedGamepad retain];
