Added support for custom splash screens in engine

diff --git a/Engine/Source/Runtime/ApplicationCore/Private/GenericPlatform/GenericPlatformSplash.cpp b/Engine/Source/Runtime/ApplicationCore/Private/GenericPlatform/GenericPlatformSplash.cpp
--- a/Engine/Source/Runtime/ApplicationCore/Private/GenericPlatform/GenericPlatformSplash.cpp
+++ b/Engine/Source/Runtime/ApplicationCore/Private/GenericPlatform/GenericPlatformSplash.cpp
@@ -13,6 +13,9 @@ const TCHAR* SupportedSplashImageExt[] =
 	nullptr
 };
 
+
+static FString GCustomSplashScreenFileName;
+
 /**
 * Return the filename found (look for PNG, JPG and BMP in that order, try to avoid BMP, use more space...)
 */
@@ -34,6 +37,11 @@ FString GetSplashFilename(const FString& ContentDir, const FString& Filename)
 	return ImageName + TEXT(".bmp");
 }
 
+void FGenericPlatformSplash::SetCustomSplashImage(const TCHAR* SplashFilename)
+{
+	GCustomSplashScreenFileName = SplashFilename;
+}
+
 /**
 * Finds a usable splash pathname for the given filename
 *
@@ -48,21 +56,40 @@ bool FGenericPlatformSplash::GetSplashPath(const TCHAR* SplashFilename, FString&
 {
 	FString Filename = FString(TEXT("Splash/")) + SplashFilename;
 
-	// first look in game's splash directory
-	OutPath = FPaths::ConvertRelativePathToFull(GetSplashFilename(FPaths::ProjectContentDir(), Filename));
-	OutIsCustom = true;
-
-	// if this was found, then we're done
-	if (IFileManager::Get().FileSize(*OutPath) != -1)
-		return true;
+	if(GCustomSplashScreenFileName.IsEmpty())
+	{
+		// first look in game's splash directory
+		OutPath = FPaths::ConvertRelativePathToFull(GetSplashFilename(FPaths::ProjectContentDir(), Filename));
+		OutIsCustom = true;
+
+		// if this was found, then we're done
+		if (IFileManager::Get().FileSize(*OutPath) != -1)
+		{
+			return true;
+		}
+
+		// next look in Engine/Splash
+		OutPath = FPaths::ConvertRelativePathToFull(GetSplashFilename(FPaths::EngineContentDir(), Filename));
+		OutIsCustom = false;
+
+		// if this was found, then we're done
+		if (IFileManager::Get().FileSize(*OutPath) != -1)
+		{
+			return true;
+		}
 
-	// next look in Engine/Splash
-	OutPath = FPaths::ConvertRelativePathToFull(GetSplashFilename(FPaths::EngineContentDir(), Filename));
-	OutIsCustom = false;
+	}
+	else
+	{ 
+		OutPath = GCustomSplashScreenFileName;
+		OutIsCustom = true;
+		// if this was found, then we're done
+		if (IFileManager::Get().FileSize(*OutPath) != -1)
+		{
+			return true;
+		}
+	}
 
-	// if this was found, then we're done
-	if (IFileManager::Get().FileSize(*OutPath) != -1)
-		return true;
 
 	// if not found yet, then return failure
 	return false;
@@ -70,26 +97,29 @@ bool FGenericPlatformSplash::GetSplashPath(const TCHAR* SplashFilename, FString&
 
 bool FGenericPlatformSplash::GetSplashPath(const TCHAR* SplashFilename, const TCHAR* IconFilename, FString& OutPath, FString& OutIconPath, bool& OutIsCustom)
 {
-	FString Filename = FString(TEXT("Splash/")) + SplashFilename;
 	FString IconName = FString(TEXT("Splash/")) + IconFilename;
 
+	// Get the path for the splash image
+	bool bGetSplashSucceeded = GetSplashPath(SplashFilename, OutPath, OutIsCustom);
+
 	// first look in game's splash directory
-	OutPath = FPaths::ConvertRelativePathToFull(GetSplashFilename(FPaths::ProjectContentDir(), Filename));
 	OutIconPath = GetSplashFilename(FPaths::ProjectContentDir(), IconName);
-	OutIsCustom = true;
 
-	// if this was found, then we're done
-	if (IFileManager::Get().FileSize(*OutPath) != -1)
-		return true;
+	// if the icon was found, then we're done
+	if (IFileManager::Get().FileSize(*OutIconPath) != -1)
+	{
+		return bGetSplashSucceeded;
+	}
 
 	// next look in Engine/Splash
-	OutPath = FPaths::ConvertRelativePathToFull(GetSplashFilename(FPaths::EngineContentDir(), Filename));
 	OutIconPath = GetSplashFilename(FPaths::EngineContentDir(), IconName);
 	OutIsCustom = false;
 
 	// if this was found, then we're done
-	if (IFileManager::Get().FileSize(*OutPath) != -1)
-		return true;
+	if (IFileManager::Get().FileSize(*OutIconPath) != -1)
+	{
+		return bGetSplashSucceeded;
+	}
 
 	// if not found yet, then return failure
 	return false;
diff --git a/Engine/Source/Runtime/ApplicationCore/Public/GenericPlatform/GenericPlatformSplash.h b/Engine/Source/Runtime/ApplicationCore/Public/GenericPlatform/GenericPlatformSplash.h
--- a/Engine/Source/Runtime/ApplicationCore/Public/GenericPlatform/GenericPlatformSplash.h
+++ b/Engine/Source/Runtime/ApplicationCore/Public/GenericPlatform/GenericPlatformSplash.h
@@ -44,6 +44,14 @@ struct APPLICATIONCORE_API FGenericPlatformSplash
 	/** Hide the splash screen. */
 	FORCEINLINE static void Hide() { }
 
+	/**
+	 * Sets a custom splash image to display
+	 * 
+	 * @param SplashFilename Full path to the splash image to display
+	 */
+	static void SetCustomSplashImage(const TCHAR* SplashFilename);
+
+
 	/**
 	 * Sets the text displayed on the splash screen (for startup/loading progress)
 	 *
