Platform - When calculating screen density, we no longer fallback to 0.  We instead report unknown, but we use 96 dpi as our default when returning screen density if we don't know it.  This is both the default windows assumes, and what we use internally for fonts in Slate.  This allows us to at least do something reasonable on platforms that can't know the physical size of their screen.

diff --git a/Engine/Source/Runtime/ApplicationCore/Private/GenericPlatform/GenericPlatformApplicationMisc.cpp b/Engine/Source/Runtime/ApplicationCore/Private/GenericPlatform/GenericPlatformApplicationMisc.cpp
--- a/Engine/Source/Runtime/ApplicationCore/Private/GenericPlatform/GenericPlatformApplicationMisc.cpp
+++ b/Engine/Source/Runtime/ApplicationCore/Private/GenericPlatform/GenericPlatformApplicationMisc.cpp
@@ -97,10 +97,16 @@ void FGenericPlatformApplicationMisc:: ClipboardPaste(class FString& Dest)
 
 EScreenPhysicalAccuracy FGenericPlatformApplicationMisc::GetPhysicalScreenDensity(int32& ScreenDensity)
 {
-	if ( !CachedPhysicalScreenData )
+	if (!CachedPhysicalScreenData)
 	{
 		CachedPhysicalScreenData = true;
 		CachedPhysicalScreenAccuracy = FPlatformApplicationMisc::ComputePhysicalScreenDensity(CachedPhysicalScreenDensity);
+		if (CachedPhysicalScreenAccuracy == EScreenPhysicalAccuracy::Unknown)
+		{
+			// If the screen density is unknown we use 96, which is the default on Windows,
+			// but it's also the assumed DPI scale of Slate internally for fonts.
+			CachedPhysicalScreenDensity = 96;
+		}
 	}
 
 	ScreenDensity = CachedPhysicalScreenDensity;
@@ -109,10 +115,10 @@ EScreenPhysicalAccuracy FGenericPlatformApplicationMisc::GetPhysicalScreenDensit
 
 EScreenPhysicalAccuracy FGenericPlatformApplicationMisc::ConvertInchesToPixels(float Inches, float& OutPixels)
 {
-	int32 ScreenDensity;
-	EScreenPhysicalAccuracy Accuracy = GetPhysicalScreenDensity(ScreenDensity);
-	
-	if ( Accuracy != EScreenPhysicalAccuracy::Unknown )
+	int32 ScreenDensity = 0;
+	const EScreenPhysicalAccuracy Accuracy = GetPhysicalScreenDensity(ScreenDensity);
+
+	if (ScreenDensity != 0)
 	{
 		OutPixels = Inches * ScreenDensity;
 	}
@@ -126,10 +132,10 @@ EScreenPhysicalAccuracy FGenericPlatformApplicationMisc::ConvertInchesToPixels(f
 
 EScreenPhysicalAccuracy FGenericPlatformApplicationMisc::ConvertPixelsToInches(float Pixels, float& OutInches)
 {
-	int32 ScreenDensity;
-	EScreenPhysicalAccuracy Accuracy = GetPhysicalScreenDensity(ScreenDensity);
+	int32 ScreenDensity = 0;
+	const EScreenPhysicalAccuracy Accuracy = GetPhysicalScreenDensity(ScreenDensity);
 
-	if ( Accuracy != EScreenPhysicalAccuracy::Unknown )
+	if (ScreenDensity != 0)
 	{
 		OutInches = Pixels / (float)ScreenDensity;
 	}
