fix for potential thread safety issue with the battery messaging

diff --git a/Engine/Source/Runtime/ApplicationCore/Private/IOS/IOSAppDelegate.cpp b/Engine/Source/Runtime/ApplicationCore/Private/IOS/IOSAppDelegate.cpp
--- a/Engine/Source/Runtime/ApplicationCore/Private/IOS/IOSAppDelegate.cpp
+++ b/Engine/Source/Runtime/ApplicationCore/Private/IOS/IOSAppDelegate.cpp
@@ -1876,7 +1876,11 @@ CORE_API bool IOSShowAchievementsUI()
     
     // Battery level is from 0.0 to 1.0, get it in terms of 0-100
     self.BatteryLevel = ((int)([Device batteryLevel] * 100));
-    UE_LOG(LogIOS, Display, TEXT("Battery Level Changed: %d"), self.BatteryLevel);
+	[FIOSAsyncTask CreateTaskWithBlock : ^ bool(void)
+	{
+		UE_LOG(LogIOS, Display, TEXT("Battery Level Changed: %d"), self.BatteryLevel);
+		return true;
+	}];
 #endif
 }
 
@@ -1886,7 +1890,11 @@ CORE_API bool IOSShowAchievementsUI()
     UIDevice* Device = [UIDevice currentDevice];
     UIDeviceBatteryState State = Device.batteryState;
     self.bBatteryState = State == UIDeviceBatteryStateUnplugged || State == UIDeviceBatteryStateUnknown;
-    UE_LOG(LogIOS, Display, TEXT("Battery State Changed: %d"), self.bBatteryState);
+	[FIOSAsyncTask CreateTaskWithBlock : ^ bool(void)
+	{
+		UE_LOG(LogIOS, Display, TEXT("Battery State Changed: %d"), self.bBatteryState);
+		return true;
+	}];
 #endif
 }
 
@@ -1907,8 +1915,12 @@ CORE_API bool IOSShowAchievementsUI()
 			case NSProcessInfoThermalStateCritical:	Severity = FCoreDelegates::ETemperatureSeverity::Critical; Level = TEXT("Critical"); break;
 		}
 
-        UE_LOG(LogIOS, Display, TEXT("Temperaure Changed: %s"), *Level);
-		FCoreDelegates::OnTemperatureChange.Broadcast(Severity);
+		[FIOSAsyncTask CreateTaskWithBlock : ^ bool(void)
+		{
+			UE_LOG(LogIOS, Display, TEXT("Temperature Changed: %s"), *Level);
+			FCoreDelegates::OnTemperatureChange.Broadcast(Severity);
+			return true;
+		}];
 	}
 #endif
 }
