UE4R - Deprecate ES2 & ES3

diff --git a/Engine/Source/Runtime/ApplicationCore/Private/IOS/IOSView.cpp b/Engine/Source/Runtime/ApplicationCore/Private/IOS/IOSView.cpp
--- a/Engine/Source/Runtime/ApplicationCore/Private/IOS/IOSView.cpp
+++ b/Engine/Source/Runtime/ApplicationCore/Private/IOS/IOSView.cpp
@@ -138,13 +138,10 @@ id<MTLDevice> GMetalDevice = nil;
 	GConfig->GetBool(TEXT("/Script/IOSRuntimeSettings.IOSRuntimeSettings"), TEXT("bSupportsMetal"), bSupportsMetal, GEngineIni);
 	GConfig->GetBool(TEXT("/Script/IOSRuntimeSettings.IOSRuntimeSettings"), TEXT("bSupportsMetalMRT"), bSupportsMetalMRT, GEngineIni);
 
-	// does commandline override?
-	bool bForceES2 = FParse::Param(FCommandLine::Get(), TEXT("ES2"));
-	
 	bool bTriedToInit = false;
 
 	// the check for the function pointer itself is to determine if the Metal framework exists, before calling it
-	if ((bSupportsMetal || bSupportsMetalMRT) && !bForceES2 && MTLCreateSystemDefaultDevice != NULL)
+	if ((bSupportsMetal || bSupportsMetalMRT) && MTLCreateSystemDefaultDevice != NULL)
 	{
 		SCOPED_BOOT_TIMING("CreateMetalDevice");
 		// if the device is unable to run with Metal (pre-A7), this will return nil
@@ -159,7 +156,7 @@ id<MTLDevice> GMetalDevice = nil;
 	{
 		FPlatformMisc::LowLevelOutputDebugStringf(TEXT("Not using Metal because: [Project Settings Disabled Metal? %s :: Commandline Forced ES2? %s :: Older OS? %s :: Pre-A7 Device? %s]"),
 			bSupportsMetal ? TEXT("No") : TEXT("Yes"),
-			bForceES2? TEXT("Yes") : TEXT("No"),
+			TEXT("No"),
 			(MTLCreateSystemDefaultDevice == NULL) ? TEXT("Yes") : TEXT("No"),
 			bTriedToInit ? TEXT("Yes") : TEXT("Unknown (didn't test)"));
 	}
diff --git a/Engine/Source/Runtime/Engine/Private/ShaderCompiler/ShaderCompiler.cpp b/Engine/Source/Runtime/Engine/Private/ShaderCompiler/ShaderCompiler.cpp
--- a/Engine/Source/Runtime/Engine/Private/ShaderCompiler/ShaderCompiler.cpp
+++ b/Engine/Source/Runtime/Engine/Private/ShaderCompiler/ShaderCompiler.cpp
@@ -4501,6 +4501,8 @@ void RecompileShadersForRemote(
 			{
 				// These platforms are deprecated and we should warn about their use
 				if (ShaderPlatform == SP_OPENGL_SM5 || ShaderPlatform == SP_PCD3D_SM4_DEPRECATED || ShaderPlatform == SP_OPENGL_ES2_IOS_DEPRECATED ||
+					ShaderPlatform == SP_PCD3D_ES2 || ShaderPlatform == SP_METAL_MACES2 || ShaderPlatform == SP_OPENGL_PCES2 ||
+					ShaderPlatform == SP_OPENGL_ES2_ANDROID || ShaderPlatform == SP_OPENGL_ES2_WEBGL ||
 					ShaderPlatform == SP_VULKAN_SM4_DEPRECATED)
 				{
 					UE_LOG(LogShaderCompilers, Warning, TEXT("You are compiling shaders for a deprecated platform '%s'"), *LegacyShaderPlatformToShaderFormat(ShaderPlatform).ToString());
diff --git a/Engine/Source/Runtime/RHI/Private/RHI.cpp b/Engine/Source/Runtime/RHI/Private/RHI.cpp
--- a/Engine/Source/Runtime/RHI/Private/RHI.cpp
+++ b/Engine/Source/Runtime/RHI/Private/RHI.cpp
@@ -7,6 +7,7 @@
 #include "RHI.h"
 #include "Modules/ModuleManager.h"
 #include "Misc/ConfigCacheIni.h"
+#include "Misc/MessageDialog.h"
 #include "RHIShaderFormatDefinitions.inl"
 #include "ProfilingDebugging/CsvProfiler.h"
 
@@ -897,6 +898,11 @@ bool RHIGetPreviewFeatureLevel(ERHIFeatureLevel::Type& PreviewFeatureLevelOUT)
 
 	if (bForceFeatureLevelES2)
 	{
+		if (!UE_BUILD_SHIPPING)
+		{
+			FMessageDialog::Open(EAppMsgType::Ok, NSLOCTEXT("RHI", "ES2Deprecated", "Warning: FeatureLevel ES2 is deprecated, please use FeatureLevel ES3.1."));
+		}
+
 		PreviewFeatureLevelOUT = ERHIFeatureLevel::ES2;
 	}
 	else if (bForceFeatureLevelES3_1)
