Moved a couple of Cocoa calls that hide the window on Mac before it's destroyed to the main thread. This should solve a rare game hang at exit.

diff --git a/Engine/Source/Runtime/ApplicationCore/Private/Mac/MacWindow.cpp b/Engine/Source/Runtime/ApplicationCore/Private/Mac/MacWindow.cpp
--- a/Engine/Source/Runtime/ApplicationCore/Private/Mac/MacWindow.cpp
+++ b/Engine/Source/Runtime/ApplicationCore/Private/Mac/MacWindow.cpp
@@ -270,8 +270,14 @@ void FMacWindow::Destroy()
 	{
 		SCOPED_AUTORELEASE_POOL;
 		bIsClosed = true;
-		[WindowHandle setAlphaValue:0.0f];
-		[WindowHandle setBackgroundColor:[NSColor clearColor]];
+
+		FCocoaWindow* WindowHandleCopy = WindowHandle;
+		MainThreadCall(^{
+			SCOPED_AUTORELEASE_POOL;
+			[WindowHandleCopy setAlphaValue:0.0f];
+			[WindowHandleCopy setBackgroundColor:[NSColor clearColor]];
+		}, UE4ShowEventMode, false);
+
 		MacApplication->OnWindowDestroyed(SharedThis(this));
 		WindowHandle = nullptr;
 	}
