When destroing the Mac console log window, only call PumpMessages() if DestroyConsole() was called on the game thread

diff --git a/Engine/Source/Runtime/ApplicationCore/Private/Mac/MacConsoleOutputDevice.cpp b/Engine/Source/Runtime/ApplicationCore/Private/Mac/MacConsoleOutputDevice.cpp
--- a/Engine/Source/Runtime/ApplicationCore/Private/Mac/MacConsoleOutputDevice.cpp
+++ b/Engine/Source/Runtime/ApplicationCore/Private/Mac/MacConsoleOutputDevice.cpp
@@ -120,10 +120,13 @@ void FMacConsoleOutputDevice::DestroyConsole()
 		FMacConsoleWindow* ConsoleWindow = ConsoleHandle;
 		ConsoleHandle = nullptr; // Stop further serialization as soon as possible
 
-		do
+		if ([NSThread isGameThread])
 		{
-			FMacPlatformApplicationMisc::PumpMessages( true );
-		} while(OutstandingTasks);
+			do
+			{
+				FMacPlatformApplicationMisc::PumpMessages( true );
+			} while(OutstandingTasks);
+		}
 
 		MainThreadCall(^{
 			SCOPED_AUTORELEASE_POOL;
