Allow other audio to continue in the background while our app runs (iOS). Don't call setActive on the audio session until after the category has been set.

diff --git a/Engine/Source/Runtime/ApplicationCore/Private/IOS/IOSAppDelegate.cpp b/Engine/Source/Runtime/ApplicationCore/Private/IOS/IOSAppDelegate.cpp
--- a/Engine/Source/Runtime/ApplicationCore/Private/IOS/IOSAppDelegate.cpp
+++ b/Engine/Source/Runtime/ApplicationCore/Private/IOS/IOSAppDelegate.cpp
@@ -604,17 +604,6 @@ static IOSAppDelegate* CachedDelegate = nil;
 		return;
 	}
 	
-	// set the AVAudioSession active if necessary
-	NSError* ActiveError = nil;
-	if (bActive)
-	{
-		[[AVAudioSession sharedInstance] setActive:bActive error:&ActiveError];
-		if (ActiveError)
-		{
-			UE_LOG(LogIOSAudioSession, Error, TEXT("Failed to set audio session to active = %d [Error = %s]"), bActive, *FString([ActiveError description]));
-		}
-	}
-
 	self.bAudioActive = bActive;
 	
 	// get the category and settings to use
@@ -690,6 +679,7 @@ static IOSAppDelegate* CachedDelegate = nil;
 		TestOptions |= AVAudioSessionCategoryOptionMixWithOthers;
 	}
 	// set the category if anything has changed
+	NSError* ActiveError = nil;
 	if ([[[AVAudioSession sharedInstance] category] compare:Category] != NSOrderedSame ||
 		[[[AVAudioSession sharedInstance] mode] compare:Mode] != NSOrderedSame ||
 		[[AVAudioSession sharedInstance] categoryOptions] != TestOptions)
@@ -702,6 +692,16 @@ static IOSAppDelegate* CachedDelegate = nil;
 			UE_LOG(LogIOSAudioSession, Error, TEXT("Failed to set AVAudioSession category to Category:%s Mode:%s Options:%x! [Error = %s]"), *FString(Category), *FString(Mode), Options, *FString([ActiveError description]));
 		}
 	}
+	
+	// set the AVAudioSession active if necessary
+	if (bActive)
+	{
+		[[AVAudioSession sharedInstance] setActive:bActive error:&ActiveError];
+		if (ActiveError)
+		{
+			UE_LOG(LogIOSAudioSession, Error, TEXT("Failed to set audio session to active = %d [Error = %s]"), bActive, *FString([ActiveError description]));
+		}
+	}
 }
 
 - (bool)IsBackgroundAudioPlaying
