On Mac don't wait for the Cocoa call to highlight menu bar item when processing key down event

diff --git a/Engine/Source/Runtime/ApplicationCore/Private/Mac/MacApplication.cpp b/Engine/Source/Runtime/ApplicationCore/Private/Mac/MacApplication.cpp
--- a/Engine/Source/Runtime/ApplicationCore/Private/Mac/MacApplication.cpp
+++ b/Engine/Source/Runtime/ApplicationCore/Private/Mac/MacApplication.cpp
@@ -1056,7 +1056,7 @@ void FMacApplication::ProcessKeyDownEvent(const FDeferredMacEvent& Event, TShare
 		if (MainMenu)
 		{
 			NSEvent* NativeEvent = Event.Event;
-			MainThreadCall(^{ [MainMenu highlightKeyEquivalent:NativeEvent]; }, NSDefaultRunLoopMode, true);
+			MainThreadCall(^{ [MainMenu highlightKeyEquivalent:NativeEvent]; }, NSDefaultRunLoopMode, false);
 		}
 	}
 	else
