iOS - Fix embedded app fighting over screen saver control with the engine on startup.

diff --git a/Engine/Source/Runtime/ApplicationCore/Private/IOS/IOSAppDelegate.cpp b/Engine/Source/Runtime/ApplicationCore/Private/IOS/IOSAppDelegate.cpp
--- a/Engine/Source/Runtime/ApplicationCore/Private/IOS/IOSAppDelegate.cpp
+++ b/Engine/Source/Runtime/ApplicationCore/Private/IOS/IOSAppDelegate.cpp
@@ -328,7 +328,9 @@ static IOSAppDelegate* CachedDelegate = nil;
 	}
 	[self ToggleAudioSession:true];
 
+#if !BUILD_EMBEDDED_APP
 	[self InitIdleTimerSettings];
+#endif
 
 	bEngineInit = true;
     
