fix for display metrics not getting properly updated on IOS when the device is rotated

diff --git a/Engine/Source/Runtime/ApplicationCore/Private/IOS/IOSAppDelegate.cpp b/Engine/Source/Runtime/ApplicationCore/Private/IOS/IOSAppDelegate.cpp
--- a/Engine/Source/Runtime/ApplicationCore/Private/IOS/IOSAppDelegate.cpp
+++ b/Engine/Source/Runtime/ApplicationCore/Private/IOS/IOSAppDelegate.cpp
@@ -1340,6 +1340,8 @@ extern EDeviceScreenOrientation ConvertFromUIInterfaceOrientation(UIInterfaceOri
     {
 		FFunctionGraphTask::CreateAndDispatchWhenReady([Orientation]()
 		{
+			FIOSApplication* Application = [IOSAppDelegate GetDelegate].IOSApplication;
+			Application->OrientationChanged(Orientation);
 			FCoreDelegates::ApplicationReceivedScreenOrientationChangedNotificationDelegate.Broadcast((int32)ConvertFromUIInterfaceOrientation(Orientation));
 
 			//we also want to fire off the safe frame event
