IOS thumbstick support

https://p4-swarm.epicgames.net/reviews/9960353

diff --git a/Engine/Source/Runtime/ApplicationCore/Private/IOS/IOSInputInterface.cpp b/Engine/Source/Runtime/ApplicationCore/Private/IOS/IOSInputInterface.cpp
--- a/Engine/Source/Runtime/ApplicationCore/Private/IOS/IOSInputInterface.cpp
+++ b/Engine/Source/Runtime/ApplicationCore/Private/IOS/IOSInputInterface.cpp
@@ -12,6 +12,34 @@
 
 DECLARE_LOG_CATEGORY_EXTERN(LogIOSInput, Log, All);
 
+#if PLATFORM_TVOS
+#define TVOS_THUMBSTICKS_UNDECLARED (__APPLETV_OS_VERSION_MAX_ALLOWED < 121000)
+#else
+#define TVOS_THUMBSTICKS_UNDECLARED false
+#endif
+
+#if (__IPHONE_OS_VERSION_MAX_ALLOWED < 121000 || TVOS_THUMBSTICKS_UNDECLARED || __MAC_OS_VERSION_MAX_ALLOWED < 1401000)
+@interface GCExtendedGamepadSnapshot()
+@property (nonatomic, readwrite, nullable) GCControllerButtonInput *leftThumbstickButton;
+@property (nonatomic, readwrite, nullable) GCControllerButtonInput *rightThumbstickButton;
+@end
+
+@interface GCExtendedGamepad()
+@property (nonatomic, readwrite, nullable) GCControllerButtonInput *leftThumbstickButton;
+@property (nonatomic, readwrite, nullable) GCControllerButtonInput *rightThumbstickButton;
+@end
+
+@interface GCControllerButtonInput (capture)
+- (bool)_setValue:(float)arg1;
+@end
+#endif
+
+#if (__IPHONE_OS_VERSION_MAX_ALLOWED < 130000)
+// only redefine these values pre-iOS13 SDK
+@interface GCController (capture)
+- (GCController *)capture;
+@end
+#endif
 
 static TAutoConsoleVariable<float> CVarHapticsKickHeavy(TEXT("ios.VibrationHapticsKickHeavyValue"), 0.65f, TEXT("Vibation values higher than this will kick a haptics heavy Impact"));
 static TAutoConsoleVariable<float> CVarHapticsKickMedium(TEXT("ios.VibrationHapticsKickMediumValue"), 0.5f, TEXT("Vibation values higher than this will kick a haptics medium Impact"));
@@ -399,7 +427,7 @@ else \
 		
 #define HANDLE_BUTTON(Gamepad, GCButton, UEButton) \
 { \
-	const bool bWasPressed = Controller.Previous##Gamepad != nil && Controller.Previous##Gamepad.GCButton.pressed; \
+	const bool bWasPressed = Previous##Gamepad != nil && Previous##Gamepad.GCButton.pressed; \
 	const bool bPressed = Gamepad.GCButton.pressed; \
 	HANDLE_BUTTON_INTERNAL(Gamepad, bWasPressed, bPressed, UEButton); \
 }
@@ -408,7 +436,7 @@ else \
         const float RepeatDeadzone = 0.24;
         
 #define HANDLE_ANALOG(Gamepad, GCAxis, UEAxis) \
-if ((Controller.Previous##Gamepad != nil && Gamepad.GCAxis.value != Controller.Previous##Gamepad.GCAxis.value) || (Gamepad.GCAxis.value < -RepeatDeadzone || Gamepad.GCAxis.value > RepeatDeadzone)) \
+if ((Previous##Gamepad != nil && Gamepad.GCAxis.value != Previous##Gamepad.GCAxis.value) || (Gamepad.GCAxis.value < -RepeatDeadzone || Gamepad.GCAxis.value > RepeatDeadzone)) \
 { \
 	NSLog(@"Axis %s is %f", TCHAR_TO_ANSI(*UEAxis.ToString()), Gamepad.GCAxis.value); \
 	MessageHandler->OnControllerAnalog(UEAxis, Cont.playerIndex, Gamepad.GCAxis.value); \
@@ -416,16 +444,18 @@ if ((Controller.Previous##Gamepad != nil && Gamepad.GCAxis.value != Controller.P
 
 #define HANDLE_ANALOG_VIRTUAL_BUTTONS(Gamepad, GCAxis, UEButtonNegative, UEButtonPositive) \
 { \
-	const bool bWasNegativePressed = Controller.Previous##Gamepad != nil && Controller.Previous##Gamepad.GCAxis.value <= -RepeatDeadzone; \
+	const bool bWasNegativePressed = Previous##Gamepad != nil && Previous##Gamepad.GCAxis.value <= -RepeatDeadzone; \
 	const bool bNegativePressed = Gamepad.GCAxis.value <= -RepeatDeadzone; \
 	HANDLE_BUTTON_INTERNAL(Gamepad, bWasNegativePressed, bNegativePressed, UEButtonNegative) \
-	const bool bWasPositivePressed = Controller.Previous##Gamepad != nil && Controller.Previous##Gamepad.GCAxis.value >= RepeatDeadzone; \
+	const bool bWasPositivePressed = Previous##Gamepad != nil && Previous##Gamepad.GCAxis.value >= RepeatDeadzone; \
 	const bool bPositivePressed = Gamepad.GCAxis.value >= RepeatDeadzone; \
 	HANDLE_BUTTON_INTERNAL(Gamepad, bWasPositivePressed, bPositivePressed, UEButtonPositive) \
 }
 		
 		if (ExtendedGamepad != nil)
 		{
+            const GCExtendedGamepad* PreviousExtendedGamepad = Controller.PreviousExtendedGamepad;
+            
 			HANDLE_BUTTON(ExtendedGamepad, buttonA,			FGamepadKeyNames::FaceButtonBottom);
 			HANDLE_BUTTON(ExtendedGamepad, buttonB,			FGamepadKeyNames::FaceButtonRight);
 			HANDLE_BUTTON(ExtendedGamepad, buttonX,			FGamepadKeyNames::FaceButtonLeft);
@@ -450,15 +480,36 @@ if ((Controller.Previous##Gamepad != nil && Gamepad.GCAxis.value != Controller.P
 			HANDLE_ANALOG_VIRTUAL_BUTTONS(ExtendedGamepad, leftThumbstick.yAxis, FGamepadKeyNames::LeftStickDown, FGamepadKeyNames::LeftStickUp);
 			HANDLE_ANALOG_VIRTUAL_BUTTONS(ExtendedGamepad, rightThumbstick.xAxis, FGamepadKeyNames::RightStickLeft, FGamepadKeyNames::RightStickRight);
 			HANDLE_ANALOG_VIRTUAL_BUTTONS(ExtendedGamepad, rightThumbstick.yAxis, FGamepadKeyNames::RightStickDown, FGamepadKeyNames::RightStickUp);
-
-			[Controller.PreviousExtendedGamepad release];
-			Controller.PreviousExtendedGamepad = [ExtendedGamepad saveSnapshot];
-			[Controller.PreviousExtendedGamepad retain];
+            
+            if(ExtendedGamepad.leftThumbstickButton != nil)
+            {
+                
+                HANDLE_BUTTON(ExtendedGamepad, leftThumbstickButton,        FGamepadKeyNames::LeftThumb);
+                HANDLE_BUTTON(ExtendedGamepad, rightThumbstickButton,        FGamepadKeyNames::RightThumb);
+            }
+
+            [Controller.PreviousExtendedGamepad release];
+            if ([Cont respondsToSelector:@selector(capture)])
+            {
+                Controller.PreviousExtendedGamepad = (GCExtendedGamepadSnapshot*)[ExtendedGamepad.controller capture].extendedGamepad;
+            }
+            else
+            {
+                Controller.PreviousExtendedGamepad = [ExtendedGamepad saveSnapshot];
+                if(ExtendedGamepad.leftThumbstickButton != nil)
+                {
+                    [Controller.PreviousExtendedGamepad.leftThumbstickButton _setValue:ExtendedGamepad.leftThumbstickButton.value];
+                    [Controller.PreviousExtendedGamepad.rightThumbstickButton _setValue:ExtendedGamepad.rightThumbstickButton.value];
+                }
+            }
+            [Controller.PreviousExtendedGamepad retain];
 		}
 #if PLATFORM_TVOS
         // get micro input (shouldn't have the other two)
         else if (MicroGamepad != nil)
         {
+            const GCExtendedGamepad* PreviousMicroGamepad = Controller.PreviousMicroGamepad;
+            
 			// if we want virtual joysticks, then use the dpad values (and drain the touch queue to not leak memory)
 			if (bUseRemoteAsVirtualJoystick)
 			{
@@ -481,9 +532,16 @@ if ((Controller.Previous##Gamepad != nil && Gamepad.GCAxis.value != Controller.P
 			HANDLE_BUTTON(MicroGamepad, buttonX,	FGamepadKeyNames::FaceButtonRight);
 		
 			
-			[Controller.PreviousMicroGamepad release];
-			Controller.PreviousMicroGamepad = [MicroGamepad saveSnapshot];
-			[Controller.PreviousMicroGamepad retain];
+            [Controller.PreviousMicroGamepad release];
+            if ([Cont respondsToSelector:@selector(capture)])
+            {
+                Controller.PreviousMicroGamepad = [ExtendedGamepad.controller capture].microGamepad;
+            }
+            else
+            {
+                Controller.PreviousMicroGamepad = [MicroGamepad saveSnapshot];
+            }
+            [Controller.PreviousMicroGamepad retain];
         }
         
 		// motion is orthogonal to buttons
