Fixing local compiler error when building for iOS
#tests iOS Client

diff --git a/Engine/Source/Runtime/ApplicationCore/Private/IOS/IOSInputInterface.cpp b/Engine/Source/Runtime/ApplicationCore/Private/IOS/IOSInputInterface.cpp
--- a/Engine/Source/Runtime/ApplicationCore/Private/IOS/IOSInputInterface.cpp
+++ b/Engine/Source/Runtime/ApplicationCore/Private/IOS/IOSInputInterface.cpp
@@ -18,6 +18,11 @@ DECLARE_LOG_CATEGORY_EXTERN(LogIOSInput, Log, All);
 #define TVOS_THUMBSTICKS_UNDECLARED false
 #endif
 
+//Avoid compile error if this isn't defined
+#ifndef __MAC_OS_VERSION_MAX_ALLOWED
+#define __MAC_OS_VERSION_MAX_ALLOWED 0
+#endif
+
 #if (__IPHONE_OS_VERSION_MAX_ALLOWED < 121000 || TVOS_THUMBSTICKS_UNDECLARED || __MAC_OS_VERSION_MAX_ALLOWED < 1401000)
 @interface GCExtendedGamepadSnapshot()
 @property (nonatomic, readwrite, nullable) GCControllerButtonInput *leftThumbstickButton;
