Fixing some compiler warning when building for tvos and rolling back micro gamepad changes that also had some warnings, as they were unneeded.

#tests iOS Client

diff --git a/Engine/Source/Runtime/ApplicationCore/Private/IOS/IOSInputInterface.cpp b/Engine/Source/Runtime/ApplicationCore/Private/IOS/IOSInputInterface.cpp
--- a/Engine/Source/Runtime/ApplicationCore/Private/IOS/IOSInputInterface.cpp
+++ b/Engine/Source/Runtime/ApplicationCore/Private/IOS/IOSInputInterface.cpp
@@ -12,17 +12,20 @@
 
 DECLARE_LOG_CATEGORY_EXTERN(LogIOSInput, Log, All);
 
+#ifndef __MAC_OS_VERSION_MAX_ALLOWED
+#define __MAC_OS_VERSION_MAX_ALLOWED 0
+#endif
+
+#ifndef __APPLETV_OS_VERSION_MAX_ALLOWED
+#define __APPLETV_OS_VERSION_MAX_ALLOWED 0
+#endif
+
 #if PLATFORM_TVOS
 #define TVOS_THUMBSTICKS_UNDECLARED (__APPLETV_OS_VERSION_MAX_ALLOWED < 121000)
 #else
 #define TVOS_THUMBSTICKS_UNDECLARED false
 #endif
 
-//Avoid compile error if this isn't defined
-#ifndef __MAC_OS_VERSION_MAX_ALLOWED
-#define __MAC_OS_VERSION_MAX_ALLOWED 0
-#endif
-
 #if (__IPHONE_OS_VERSION_MAX_ALLOWED < 121000 || TVOS_THUMBSTICKS_UNDECLARED || __MAC_OS_VERSION_MAX_ALLOWED < 1401000)
 @interface GCExtendedGamepadSnapshot()
 @property (nonatomic, readwrite, nullable) GCControllerButtonInput *leftThumbstickButton;
@@ -513,7 +516,7 @@ if ((Previous##Gamepad != nil && Gamepad.GCAxis.value != Previous##Gamepad.GCAxi
         // get micro input (shouldn't have the other two)
         else if (MicroGamepad != nil)
         {
-            const GCExtendedGamepad* PreviousMicroGamepad = Controller.PreviousMicroGamepad;
+            const GCMicroGamepad* PreviousMicroGamepad = Controller.PreviousMicroGamepad;
             
 			// if we want virtual joysticks, then use the dpad values (and drain the touch queue to not leak memory)
 			if (bUseRemoteAsVirtualJoystick)
@@ -537,16 +540,9 @@ if ((Previous##Gamepad != nil && Gamepad.GCAxis.value != Previous##Gamepad.GCAxi
 			HANDLE_BUTTON(MicroGamepad, buttonX,	FGamepadKeyNames::FaceButtonRight);
 		
 			
-            [Controller.PreviousMicroGamepad release];
-            if ([Cont respondsToSelector:@selector(capture)])
-            {
-                Controller.PreviousMicroGamepad = [ExtendedGamepad.controller capture].microGamepad;
-            }
-            else
-            {
-                Controller.PreviousMicroGamepad = [MicroGamepad saveSnapshot];
-            }
-            [Controller.PreviousMicroGamepad retain];
+			[Controller.PreviousMicroGamepad release];
+			Controller.PreviousMicroGamepad = [MicroGamepad saveSnapshot];
+			[Controller.PreviousMicroGamepad retain];
         }
         
 		// motion is orthogonal to buttons
