InitializeIMM now returns false if ImmCreateContext returns null

#rnx

diff --git a/Engine/Source/Runtime/ApplicationCore/Private/Windows/WindowsTextInputMethodSystem.cpp b/Engine/Source/Runtime/ApplicationCore/Private/Windows/WindowsTextInputMethodSystem.cpp
--- a/Engine/Source/Runtime/ApplicationCore/Private/Windows/WindowsTextInputMethodSystem.cpp
+++ b/Engine/Source/Runtime/ApplicationCore/Private/Windows/WindowsTextInputMethodSystem.cpp
@@ -254,6 +254,12 @@ bool FWindowsTextInputMethodSystem::InitializeIMM()
 	IMMContextId = ::ImmCreateContext();
 	UpdateIMMProperty(::GetKeyboardLayout(0));
 
+	if (!IMMContextId)
+	{
+		UE_CLOG(!GIsBuildMachine, LogWindowsTextInputMethodSystem, Warning, TEXT("Initialization failed while creating the IMM context."));
+		return false;
+	}
+
 	UE_LOG(LogWindowsTextInputMethodSystem, Verbose, TEXT("Initialized IMM!"));
 
 	return true;
@@ -404,10 +410,10 @@ bool FWindowsTextInputMethodSystem::InitializeTSF()
 		Result = ::CoCreateInstance(CLSID_TF_ThreadMgr, NULL, CLSCTX_INPROC_SERVER, IID_ITfThreadMgr, reinterpret_cast<void**>(&(RawPointerTSFThreadManager)));
 		if(FAILED(Result))
 		{
-			TCHAR ErrorMsg[1024];
-			FPlatformMisc::GetSystemErrorMessage(ErrorMsg, 1024, Result);
 			if (!GIsBuildMachine)
 			{
+				TCHAR ErrorMsg[1024];
+				FPlatformMisc::GetSystemErrorMessage(ErrorMsg, 1024, Result);
 				UE_LOG(LogWindowsTextInputMethodSystem, Warning, TEXT("Initialization failed while creating the TSF thread manager. %s (0x%08x)"), ErrorMsg, Result);
 			}
 			TSFInputProcessorProfiles.Reset();
@@ -506,10 +512,7 @@ bool FWindowsTextInputMethodSystem::InitializeTSF()
 	{
 		TCHAR ErrorMsg[1024];
 		FPlatformMisc::GetSystemErrorMessage(ErrorMsg, 1024, Result);
-		if (!GIsBuildMachine)
-		{
-			UE_LOG(LogWindowsTextInputMethodSystem, Warning, TEXT("Initialization failed while creating the TSF thread manager. %s (0x%08x)"), ErrorMsg, Result);
-		}
+		UE_LOG(LogWindowsTextInputMethodSystem, Warning, TEXT("Initialization failed while creating the TSF document manager. %s (0x%08x)"), ErrorMsg, Result);
 		TSFInputProcessorProfiles.Reset();
 		TSFInputProcessorProfileManager.Reset();
 		TSFThreadManager.Reset();
@@ -523,12 +526,11 @@ bool FWindowsTextInputMethodSystem::InitializeTSF()
 	{
 		TCHAR ErrorMsg[1024];
 		FPlatformMisc::GetSystemErrorMessage(ErrorMsg, 1024, Result);
-		UE_LOG(LogWindowsTextInputMethodSystem, Error, TEXT("Initialization failed while activating the TSF thread manager. %s (0x%08x)"), ErrorMsg, Result);
+		UE_LOG(LogWindowsTextInputMethodSystem, Error, TEXT("Initialization failed while activating the TSF document manager. %s (0x%08x)"), ErrorMsg, Result);
 		TSFInputProcessorProfiles.Reset();
 		TSFInputProcessorProfileManager.Reset();
 		TSFThreadManager.Reset();
 		TSFActivationProxy.Reset();
-		TSFThreadManager.Reset();
 		return false;
 	}
 
diff --git a/Engine/Source/Runtime/ApplicationCore/Public/Windows/WindowsTextInputMethodSystem.h b/Engine/Source/Runtime/ApplicationCore/Public/Windows/WindowsTextInputMethodSystem.h
--- a/Engine/Source/Runtime/ApplicationCore/Public/Windows/WindowsTextInputMethodSystem.h
+++ b/Engine/Source/Runtime/ApplicationCore/Public/Windows/WindowsTextInputMethodSystem.h
@@ -107,8 +107,6 @@ private:
 
 private:
 
-	TSharedPtr<ITextInputMethodContext> ActiveContext;
-
 	enum class EAPI
 	{
 		Unknown,
@@ -120,39 +118,32 @@ private:
 	TComPtr<ITfInputProcessorProfiles> TSFInputProcessorProfiles;
 	TComPtr<ITfInputProcessorProfileMgr> TSFInputProcessorProfileManager;
 	TComPtr<ITfThreadMgr> TSFThreadManager;
-	TfClientId TSFClientId;
+	TfClientId TSFClientId = 0;
 	TComPtr<ITfDocumentMgr> TSFDisabledDocumentManager;
 	TComPtr<FTSFActivationProxy> TSFActivationProxy;
 
+	// IMM Implementation
+	HIMC IMMContextId = nullptr;
+	DWORD IMMProperties = 0;
+
 	struct FInternalContext
 	{
-		FInternalContext()
-			: WindowHandle(nullptr)
-		{
-			IMMContext.IsComposing = false;
-			IMMContext.IsDeactivating = false;
-			IMMContext.CompositionBeginIndex = 0;
-			IMMContext.CompositionLength = 0;
-		}
-
-		HWND WindowHandle;
+		HWND WindowHandle = nullptr;
 
 		TComPtr<FTextStoreACP> TSFContext;
 
 		struct
 		{
-			bool IsComposing;
-			bool IsDeactivating;
-			int32 CompositionBeginIndex;
-			uint32 CompositionLength;
+			bool IsComposing = false;
+			bool IsDeactivating = false;
+			int32 CompositionBeginIndex = 0;
+			uint32 CompositionLength = 0;
 		} IMMContext;
 	};
 
-	TMap< TWeakPtr<ITextInputMethodContext>, FInternalContext > ContextToInternalContextMap;
+	TSharedPtr<ITextInputMethodContext> ActiveContext;
 
-	// IMM Implementation
-	HIMC IMMContextId;
-	DWORD IMMProperties;
+	TMap< TWeakPtr<ITextInputMethodContext>, FInternalContext > ContextToInternalContextMap;
 
 	TSet<TWeakPtr<FGenericWindow>> KnownWindows;
 };
