Fixing bindings for dualshock controllers on Android devices running version 10 or greater.

diff --git a/Engine/Source/Runtime/ApplicationCore/Private/Android/AndroidInputInterface.cpp b/Engine/Source/Runtime/ApplicationCore/Private/Android/AndroidInputInterface.cpp
--- a/Engine/Source/Runtime/ApplicationCore/Private/Android/AndroidInputInterface.cpp
+++ b/Engine/Source/Runtime/ApplicationCore/Private/Android/AndroidInputInterface.cpp
@@ -927,7 +927,8 @@ void FAndroidInputInterface::SendControllerEvents()
 						else if (CurrentDevice.DeviceInfo.Name.StartsWith(TEXT("PS4 Wireless Controller")))
 						{
 							CurrentDevice.ControllerClass = ControllerClassType::PS4Wireless;
-							if (CurrentDevice.DeviceInfo.Name.EndsWith(TEXT(" (v2)")) && FAndroidMisc::GetCPUVendor() != TEXT("Sony"))
+							if (CurrentDevice.DeviceInfo.Name.EndsWith(TEXT(" (v2)")) && FAndroidMisc::GetCPUVendor() != TEXT("Sony")
+								&& FAndroidMisc::GetAndroidBuildVersion() < 10)
 							{
 								// Only needed for non-Sony devices with v2 firmware
 								CurrentDevice.ButtonRemapping = ButtonRemapType::PS4;
