Store the current opacity in FMacWindow so we can revert back to it when necessary
#rnx

diff --git a/Engine/Source/Runtime/ApplicationCore/Private/Mac/MacWindow.cpp b/Engine/Source/Runtime/ApplicationCore/Private/Mac/MacWindow.cpp
--- a/Engine/Source/Runtime/ApplicationCore/Private/Mac/MacWindow.cpp
+++ b/Engine/Source/Runtime/ApplicationCore/Private/Mac/MacWindow.cpp
@@ -11,6 +11,7 @@
 FMacWindow::FMacWindow()
 :	WindowHandle(nullptr)
 ,	DisplayID(kCGNullDirectDisplay)
+,   CachedOpacity(1.0f)
 ,	bIsVisible(false)
 ,	bIsClosed(false)
 ,	bIsFirstTimeVisible(true)
@@ -440,6 +441,7 @@ void FMacWindow::SetOpacity( const float InOpacity )
 {
 	MainThreadCall(^{
 		SCOPED_AUTORELEASE_POOL;
+        CachedOpacity = InOpacity;
 		[WindowHandle setAlphaValue:InOpacity];
 	}, UE4NilEventMode, true);
 }
@@ -649,8 +651,8 @@ void FMacWindow::ApplySizeAndModeChanges(int32 X, int32 Y, int32 Width, int32 He
 					{
 						[WindowHandle setFrame:Rect display:YES];
 					}
-
-					const float WindowOpacity = (Definition->TransparencySupport == EWindowTransparency::PerWindow) ? Definition->Opacity : 1.0f;
+                    
+                    const float WindowOpacity = (Definition->TransparencySupport == EWindowTransparency::PerWindow) ? CachedOpacity : 1.0f;
 					[WindowHandle setAlphaValue:(Width > 0 && Height > 0) ? WindowOpacity : 0.0f];
 
 					if (Definition->ShouldPreserveAspectRatio)
diff --git a/Engine/Source/Runtime/ApplicationCore/Public/Mac/MacWindow.h b/Engine/Source/Runtime/ApplicationCore/Public/Mac/MacWindow.h
--- a/Engine/Source/Runtime/ApplicationCore/Public/Mac/MacWindow.h
+++ b/Engine/Source/Runtime/ApplicationCore/Public/Mac/MacWindow.h
@@ -104,6 +104,8 @@ private:
 	FCocoaWindow* WindowHandle;
 
 	CGDirectDisplayID DisplayID;
+    
+    float CachedOpacity;
 
 	struct FWindowedModeSavedState
 	{
