# Fix for bug when detecting window restored from minimized
#rnx

diff --git a/Engine/Source/Runtime/ApplicationCore/Private/Windows/WindowsApplication.cpp b/Engine/Source/Runtime/ApplicationCore/Private/Windows/WindowsApplication.cpp
--- a/Engine/Source/Runtime/ApplicationCore/Private/Windows/WindowsApplication.cpp
+++ b/Engine/Source/Runtime/ApplicationCore/Private/Windows/WindowsApplication.cpp
@@ -2185,7 +2185,7 @@ int32 FWindowsApplication::ProcessDeferredMessage( const FDeferredWindowsMessage
 						break;
 					case SW_PARENTOPENING:
 						CurrentNativeEventWindowPtr->OnParentWindowRestored();
-						bMinimized = true;
+						bMinimized = false;
 						break;
 					default:
 						break;
