- added a 120fps option for ios
- added bEnableDynamicMaxFPS option

#rnx

diff --git a/Engine/Config/BaseEngine.ini b/Engine/Config/BaseEngine.ini
--- a/Engine/Config/BaseEngine.ini
+++ b/Engine/Config/BaseEngine.ini
@@ -1979,6 +1979,7 @@ BundleName=[PROJECT_NAME]
 BundleIdentifier=com.YourCompany.[PROJECT_NAME]
 VersionInfo=1.0
 FrameRateLock=PUFRL_30
+bEnableDynamicMaxFPS=False
 MinimumiOSVersion=IOS_11
 bSupportsIPad=True
 bSupportsIPhone=True
diff --git a/Engine/Source/Runtime/ApplicationCore/Private/IOS/IOSPlatformFramePacer.cpp b/Engine/Source/Runtime/ApplicationCore/Private/IOS/IOSPlatformFramePacer.cpp
--- a/Engine/Source/Runtime/ApplicationCore/Private/IOS/IOSPlatformFramePacer.cpp
+++ b/Engine/Source/Runtime/ApplicationCore/Private/IOS/IOSPlatformFramePacer.cpp
@@ -15,7 +15,7 @@ static NSMutableSet<FIOSFramePacerHandler>* Handlers = [NSMutableSet new];
 
 namespace IOSDisplayConstants
 {
-    const uint32 MaxRefreshRate = 60;
+	const uint32 MaxRefreshRate = 60;
 }
 
 /*******************************************************************
@@ -64,7 +64,7 @@ namespace IOSDisplayConstants
 #if UE4_HAS_IOS10
 	if ([displayLink respondsToSelector : @selector(preferredFramesPerSecond)] == YES)
 	{
-		displayLink.preferredFramesPerSecond = IOSDisplayConstants::MaxRefreshRate / FIOSPlatformRHIFramePacer::FrameInterval;
+		displayLink.preferredFramesPerSecond = FIOSPlatformRHIFramePacer::GetMaxRefreshRate() / FIOSPlatformRHIFramePacer::FrameInterval;
 	}
 	else
 #endif
@@ -111,21 +111,23 @@ namespace IOSDisplayConstants
 		if (VSyncCVar && VSyncCVar->GetValueOnRenderThread() > 0)
 		{
 			uint32 NewFrameInterval = VSyncIntervalCVar ? VSyncIntervalCVar->GetValueOnRenderThread() : FIOSPlatformRHIFramePacer::FrameInterval;
-//			NewFrameInterval = (NewFrameInterval < FIOSPlatformRHIFramePacer::MinFrameInterval) ? FIOSPlatformRHIFramePacer::MinFrameInterval : NewFrameInterval;
 
 			// If changed, update the display link
 			if (NewFrameInterval != FIOSPlatformRHIFramePacer::FrameInterval)
 			{
+				
 				FIOSPlatformRHIFramePacer::FrameInterval = NewFrameInterval;
-	
+				uint32 MaxRefreshRate = FIOSPlatformRHIFramePacer::GetMaxRefreshRate();
+
 #if (UE4_HAS_IOS10 || UE4_TARGET_PRE_IOS10)	
 				CADisplayLink* displayLinkParam = (CADisplayLink*)param;
 #endif
 		
 #if UE4_HAS_IOS10
+				
 				if (displayLinkParam.preferredFramesPerSecond > 0)
 				{
-					displayLinkParam.preferredFramesPerSecond = IOSDisplayConstants::MaxRefreshRate / FIOSPlatformRHIFramePacer::FrameInterval;
+					displayLinkParam.preferredFramesPerSecond = MaxRefreshRate / FIOSPlatformRHIFramePacer::FrameInterval;
 				}
 				else
 #endif
@@ -136,7 +138,7 @@ namespace IOSDisplayConstants
 				}
 
 				// Update pacing for present
-				FIOSPlatformRHIFramePacer::Pace = IOSDisplayConstants::MaxRefreshRate / FIOSPlatformRHIFramePacer::FrameInterval;
+				FIOSPlatformRHIFramePacer::Pace = MaxRefreshRate / FIOSPlatformRHIFramePacer::FrameInterval;
 			}
 		}
 	}	
@@ -170,7 +172,8 @@ bool FIOSPlatformRHIFramePacer::IsEnabled()
 		FString FrameRateLockAsEnum;
 		GConfig->GetString(TEXT("/Script/IOSRuntimeSettings.IOSRuntimeSettings"), TEXT("FrameRateLock"), FrameRateLockAsEnum, GEngineIni);
 
-		uint32 FrameRateLock = IOSDisplayConstants::MaxRefreshRate;
+		uint32 MaxRefreshRate = GetMaxRefreshRate();
+		uint32 FrameRateLock = MaxRefreshRate;
 		FParse::Value(*FrameRateLockAsEnum, TEXT("PUFRL_"), FrameRateLock);
 
         const bool bOverridesFrameRate = FParse::Value( FCommandLine::Get(), TEXT( "FrameRateLock=" ), FrameRateLockAsEnum );
@@ -181,13 +184,13 @@ bool FIOSPlatformRHIFramePacer::IsEnabled()
         
         if (FrameRateLock == 0)
         {
-            FrameRateLock = IOSDisplayConstants::MaxRefreshRate;
+            FrameRateLock = MaxRefreshRate;
         }
 
         if (!bIsRHIFramePacerEnabled)
 		{
-			check((IOSDisplayConstants::MaxRefreshRate % FrameRateLock) == 0);
-			FrameInterval = IOSDisplayConstants::MaxRefreshRate / FrameRateLock;
+			check((MaxRefreshRate % FrameRateLock) == 0);
+			FrameInterval = MaxRefreshRate / FrameRateLock;
 			MinFrameInterval = FrameInterval;
 
 			bIsRHIFramePacerEnabled = (FrameInterval > 0);
@@ -201,6 +204,20 @@ bool FIOSPlatformRHIFramePacer::IsEnabled()
 	return bIsRHIFramePacerEnabled;
 }
 
+uint32 FIOSPlatformRHIFramePacer::GetMaxRefreshRate()
+{
+	static bool bEnableDynamicMaxFPS = false;
+	static bool bInitialized = false;
+	
+	if (!bInitialized)
+	{
+		GConfig->GetString(TEXT("/Script/IOSRuntimeSettings.IOSRuntimeSettings"), TEXT("bEnableDynamicMaxFPS"), bEnableDynamicMaxFPS, GEngineIni);
+		bInitialized = true;
+	}
+	
+	return bEnableDynamicMaxFPS ? [UIScreen mainScreen].maximumFramesPerSecond : IOSDisplayConstants::MaxRefreshRate;
+}
+
 void FIOSPlatformRHIFramePacer::InitWithEvent(FEvent* TriggeredEvent)
 {
     // Create display link thread
diff --git a/Engine/Source/Runtime/ApplicationCore/Public/IOS/IOSPlatformFramePacer.h b/Engine/Source/Runtime/ApplicationCore/Public/IOS/IOSPlatformFramePacer.h
--- a/Engine/Source/Runtime/ApplicationCore/Public/IOS/IOSPlatformFramePacer.h
+++ b/Engine/Source/Runtime/ApplicationCore/Public/IOS/IOSPlatformFramePacer.h
@@ -21,6 +21,7 @@ struct FIOSPlatformRHIFramePacer : public FGenericPlatformRHIFramePacer
 {
     // FGenericPlatformRHIFramePacer interface
     static bool IsEnabled();
+	static uint32 GetMaxRefreshRate();
 	static void InitWithEvent(class FEvent* TriggeredEvent);
 	static void AddHandler(FIOSFramePacerHandler Handler);
 	static void RemoveHandler(FIOSFramePacerHandler Handler);
diff --git a/Engine/Source/Runtime/IOS/IOSRuntimeSettings/Classes/IOSRuntimeSettings.h b/Engine/Source/Runtime/IOS/IOSRuntimeSettings/Classes/IOSRuntimeSettings.h
--- a/Engine/Source/Runtime/IOS/IOSRuntimeSettings/Classes/IOSRuntimeSettings.h
+++ b/Engine/Source/Runtime/IOS/IOSRuntimeSettings/Classes/IOSRuntimeSettings.h
@@ -395,6 +395,10 @@ public:
 	UPROPERTY(GlobalConfig, EditAnywhere, Category = PowerUsage, meta = (ConfigHierarchyEditable))
 	EPowerUsageFrameRateLock FrameRateLock;
 
+	//Whether or not to allow taking the MaxRefreshRate from the device instead of a constant (60fps) in IOSPlatformFramePacer
+	UPROPERTY(GlobalConfig, EditAnywhere, Category = PowerUsage, meta = (ConfigHierarchyEditable))
+	bool bEnableDynamicMaxFPS;
+
 	// Minimum iOS version this game supports
 	UPROPERTY(GlobalConfig, EditAnywhere, Category = OSInfo, meta = (DisplayName = "Minimum iOS Version"))
 	EIOSVersion MinimumiOSVersion;
diff --git a/Engine/Source/Runtime/IOS/IOSRuntimeSettings/Private/IOSRuntimeSettings.cpp b/Engine/Source/Runtime/IOS/IOSRuntimeSettings/Private/IOSRuntimeSettings.cpp
--- a/Engine/Source/Runtime/IOS/IOSRuntimeSettings/Private/IOSRuntimeSettings.cpp
+++ b/Engine/Source/Runtime/IOS/IOSRuntimeSettings/Private/IOSRuntimeSettings.cpp
@@ -19,6 +19,7 @@ UIOSRuntimeSettings::UIOSRuntimeSettings(const FObjectInitializer& ObjectInitial
 	BundleIdentifier = TEXT("com.YourCompany.GameNameNoSpaces");
 	VersionInfo = TEXT("1.0.0");
     FrameRateLock = EPowerUsageFrameRateLock::PUFRL_30;
+	bEnableDynamicMaxFPS = false;
 	bSupportsIPad = true;
 	bSupportsIPhone = true;
 	MinimumiOSVersion = EIOSVersion::IOS_11;
