Fixes compile error

#rnx

diff --git a/Engine/Source/Runtime/ApplicationCore/Private/IOS/IOSPlatformFramePacer.cpp b/Engine/Source/Runtime/ApplicationCore/Private/IOS/IOSPlatformFramePacer.cpp
--- a/Engine/Source/Runtime/ApplicationCore/Private/IOS/IOSPlatformFramePacer.cpp
+++ b/Engine/Source/Runtime/ApplicationCore/Private/IOS/IOSPlatformFramePacer.cpp
@@ -211,7 +211,7 @@ uint32 FIOSPlatformRHIFramePacer::GetMaxRefreshRate()
 	
 	if (!bInitialized)
 	{
-		GConfig->GetString(TEXT("/Script/IOSRuntimeSettings.IOSRuntimeSettings"), TEXT("bEnableDynamicMaxFPS"), bEnableDynamicMaxFPS, GEngineIni);
+		GConfig->GetBool(TEXT("/Script/IOSRuntimeSettings.IOSRuntimeSettings"), TEXT("bEnableDynamicMaxFPS"), bEnableDynamicMaxFPS, GEngineIni);
 		bInitialized = true;
 	}
 	
