Background HTTP Work
- Fix for IOS Background HTTP tasks incorrectly stalling.
- Added some additional logging around session handling.
#tests ios client

diff --git a/Engine/Source/Runtime/ApplicationCore/Private/IOS/IOSBackgroundURLSessionHandler.cpp b/Engine/Source/Runtime/ApplicationCore/Private/IOS/IOSBackgroundURLSessionHandler.cpp
--- a/Engine/Source/Runtime/ApplicationCore/Private/IOS/IOSBackgroundURLSessionHandler.cpp
+++ b/Engine/Source/Runtime/ApplicationCore/Private/IOS/IOSBackgroundURLSessionHandler.cpp
@@ -138,6 +138,8 @@ const FString FBackgroundURLSessionHandler::GetTemporaryFilePathFromURL(const FS
 
 void FBackgroundURLSessionHandler::CallBackgroundURLSessionCompleteHandler()
 {
+	NSLog(@"OnIOSBackgroundDownload_SessionDidFinishAllEvents Delegates ALL Finshed. Attempting to callback into iOS to inform session work is finished.");
+	
 	IOSAppDelegate* AppDelegate = [IOSAppDelegate GetDelegate];
 	
 	if ((nullptr != AppDelegate) && (nil != AppDelegate.BackgroundSessionEventCompleteDelegate))
@@ -151,13 +153,20 @@ void FBackgroundURLSessionHandler::CallBackgroundURLSessionCompleteHandler()
 		
 	   //Need to issue our stored completion callback on main thread as its a UI callback
 		[[NSOperationQueue mainQueue] addOperationWithBlock:^{
+	
 			//call the copied completion handler
-			completionHandler();
+			 NSLog(@"OnIOSBackgroundDownload_SessionDidFinishAllEvents Calling iOS Completion handler.");
+			 completionHandler();
 			
 			//Release local copy so we don't leak the completion handler
 			[completionHandler release];
 		}];
 	}
+	else
+	{
+		NSLog(@"Error: OnIOSBackgroundDownload_SessionDidFinishAllEvents Can not call iOS Completion Handler as it has not been set correctly! AppDelegate:");
+		ensureAlwaysMsgf(false, TEXT("Error! CallBackgroundURLSessionCompleteHandler can not call into our iOS completionHandler from handleEventsForBackgroundURLSession. Likely stored incorrectly!"));
+	}
 }
 
 void FBackgroundURLSessionHandler::AddDelayedBackgroundURLSessionComplete()
@@ -168,6 +177,7 @@ void FBackgroundURLSessionHandler::AddDelayedBackgroundURLSessionComplete()
 void FBackgroundURLSessionHandler::OnDelayedBackgroundURLSessionCompleteHandlerCalled()
 {
 	int NewDecrementedCount = FPlatformAtomics::InterlockedDecrement(&DelayedBackgroundURLSessionCompleteCount);
+	NSLog(@"OnIOSBackgroundDownload_SessionDidFinishAllEvents Delegate Finshed. Still pending: %d", NewDecrementedCount);
 	
 	if (NewDecrementedCount == 0)
 	{
@@ -224,6 +234,8 @@ void FBackgroundURLSessionHandler::OnDelayedBackgroundURLSessionCompleteHandlerC
 	
 	if (bHasURLSessionDelegateBeenSetup)
 	{
+		NSLog(@"Calling into OnIOSBackgroundDownload_SessionDidFinishAllEvents Waiting for Delegates to complete.");
+		
 		//We call this here and then the completion handler ourselves below this. That way we delay the complete until all systems responding to this delegate have received the broadcast.
 		FBackgroundURLSessionHandler::AddDelayedBackgroundURLSessionComplete();
 		
diff --git a/Engine/Source/Runtime/Online/BackgroundHTTP/Private/IOS/ApplePlatformBackgroundHttpRequest.cpp b/Engine/Source/Runtime/Online/BackgroundHTTP/Private/IOS/ApplePlatformBackgroundHttpRequest.cpp
--- a/Engine/Source/Runtime/Online/BackgroundHTTP/Private/IOS/ApplePlatformBackgroundHttpRequest.cpp
+++ b/Engine/Source/Runtime/Online/BackgroundHTTP/Private/IOS/ApplePlatformBackgroundHttpRequest.cpp
@@ -49,7 +49,8 @@ void FApplePlatformBackgroundHttpRequest::CompleteRequest_Internal(bool bWasRequ
 {
     UE_LOG(LogBackgroundHttpRequest, Display, TEXT("Marking Request Complete -- RequestDebugID:%s | bWasRequestSuccess:%d | CompletedTempDownloadLocation:%s"), *GetRequestDebugID(), (int)bWasRequestSuccess, *CompletedTempDownloadLocationIn);
     
-    FPlatformAtomics::InterlockedExchange(&bIsTaskActive, false);
+	//Purposefully avoid setting bIsTaskActive to false here as we still expect that to be set until the request is deleted.
+	
 	FPlatformAtomics::InterlockedExchange(&bIsCompleted, true);
 	FPlatformAtomics::InterlockedExchange(&bIsFailed, !bWasRequestSuccess);
     FPlatformAtomics::InterlockedExchange(&bIsPendingCancel, false);
