V1 accessibility implementation for Slate, Windows, and UMG
Custom min/max support for sliders

diff --git a/Engine/Plugins/Media/MediaIOFramework/Source/MediaIOEditor/Public/Widgets/SMediaPermutationsSelector.inl b/Engine/Plugins/Media/MediaIOFramework/Source/MediaIOEditor/Public/Widgets/SMediaPermutationsSelector.inl
--- a/Engine/Plugins/Media/MediaIOFramework/Source/MediaIOEditor/Public/Widgets/SMediaPermutationsSelector.inl
+++ b/Engine/Plugins/Media/MediaIOFramework/Source/MediaIOEditor/Public/Widgets/SMediaPermutationsSelector.inl
@@ -22,7 +22,8 @@ ItemType SMediaPermutationsSelector<ItemType, ItemBuilder>::GetSelectedItem() co
 template <typename ItemType, typename ItemBuilder>
 void SMediaPermutationsSelector<ItemType, ItemBuilder>::Construct(const FArguments& InArgs)
 {
-	SWidget::Construct(InArgs._ToolTipText, InArgs._ToolTip, InArgs._Cursor, InArgs._IsEnabled, InArgs._Visibility, InArgs._RenderOpacity, InArgs._RenderTransform, InArgs._RenderTransformPivot, InArgs._Tag, InArgs._ForceVolatile, InArgs._Clipping, InArgs._FlowDirectionPreference, InArgs.MetaData);
+	SWidget::Construct(InArgs._ToolTipText, InArgs._ToolTip, InArgs._Cursor, InArgs._IsEnabled, InArgs._Visibility, InArgs._RenderOpacity, InArgs._RenderTransform, InArgs._RenderTransformPivot, InArgs._Tag, InArgs._ForceVolatile, InArgs._Clipping,
+					   InArgs._FlowDirectionPreference, InArgs._AccessibleParams, InArgs.MetaData);
 
 	PermutationsSource = InArgs._PermutationsSource;
 	SelectedPermutationIndex = INDEX_NONE;
diff --git a/Engine/Source/Editor/PropertyEditor/Private/IDetailsViewPrivate.h b/Engine/Source/Editor/PropertyEditor/Private/IDetailsViewPrivate.h
--- a/Engine/Source/Editor/PropertyEditor/Private/IDetailsViewPrivate.h
+++ b/Engine/Source/Editor/PropertyEditor/Private/IDetailsViewPrivate.h
@@ -79,8 +79,6 @@ public:
 	 */
 	virtual bool IsPropertyReadOnly( const struct FPropertyAndParent& PropertyAndParent ) const = 0;
 
-	virtual TSharedPtr<IDetailPropertyExtensionHandler> GetExtensionHandler() = 0;
-
 	/**
 	 * @return The thumbnail pool that should be used for thumbnails being rendered in this view
 	 */
diff --git a/Engine/Source/Editor/PropertyEditor/Public/IDetailsView.h b/Engine/Source/Editor/PropertyEditor/Public/IDetailsView.h
--- a/Engine/Source/Editor/PropertyEditor/Public/IDetailsView.h
+++ b/Engine/Source/Editor/PropertyEditor/Public/IDetailsView.h
@@ -266,6 +266,8 @@ public:
 
 	virtual void SetExtensionHandler(TSharedPtr<class IDetailPropertyExtensionHandler> InExtensionandler) = 0;
 
+	virtual TSharedPtr<class IDetailPropertyExtensionHandler> GetExtensionHandler() = 0;
+
 	/**
 	 * @return true if property editing is enabled (based on the FIsPropertyEditingEnabled delegate)
 	 */ 
diff --git a/Engine/Source/Editor/UMGEditor/Private/Customizations/UMGDetailCustomizations.cpp b/Engine/Source/Editor/UMGEditor/Private/Customizations/UMGDetailCustomizations.cpp
--- a/Engine/Source/Editor/UMGEditor/Private/Customizations/UMGDetailCustomizations.cpp
+++ b/Engine/Source/Editor/UMGEditor/Private/Customizations/UMGDetailCustomizations.cpp
@@ -25,6 +25,8 @@
 #include "Components/PanelSlot.h"
 #include "Details/SPropertyBinding.h"
 #include "Widgets/Layout/SWidgetSwitcher.h"
+#include "IDetailsView.h"
+#include "IDetailPropertyExtensionHandler.h"
 
 #define LOCTEXT_NAMESPACE "UMG"
 
@@ -283,6 +285,11 @@ void FBlueprintWidgetCustomization::CustomizeDetails( IDetailLayoutBuilder& Deta
 	static const FName LayoutCategoryKey(TEXT("Layout"));
 	static const FName LocalizationCategoryKey(TEXT("Localization"));
 
+	static const FName AccessibleBehaviorKey(TEXT("AccessibleBehavior"));
+	static const FName AccessibleTextKey(TEXT("AccessibleText"));
+	static const FName AccessibleSummaryBehaviorKey(TEXT("AccessibleSummaryBehavior"));
+	static const FName AccessibleSummaryTextKey(TEXT("AccessibleSummaryText"));
+
 	DetailLayout.EditCategory(LocalizationCategoryKey, FText::GetEmpty(), ECategoryPriority::Uncommon);
 
 	TArray< TWeakObjectPtr<UObject> > OutObjects;
@@ -307,6 +314,9 @@ void FBlueprintWidgetCustomization::CustomizeDetails( IDetailLayoutBuilder& Deta
 		}
 	}
 
+	CustomizeAccessibilityProperty(DetailLayout, AccessibleBehaviorKey, AccessibleTextKey);
+	CustomizeAccessibilityProperty(DetailLayout, AccessibleSummaryBehaviorKey, AccessibleSummaryTextKey);
+
 	PerformBindingCustomization(DetailLayout);
 }
 
@@ -342,4 +352,66 @@ void FBlueprintWidgetCustomization::PerformBindingCustomization(IDetailLayoutBui
 	}
 }
 
+void FBlueprintWidgetCustomization::CustomizeAccessibilityProperty(IDetailLayoutBuilder& DetailLayout, const FName& BehaviorPropertyName, const FName& TextPropertyName)
+{
+	static const FName AccessibilityCategoryKey(TEXT("Accessibility"));
+
+	// Treat the *Behavior property as the "base" property for the row, and then add the *Text properties to the end of it.
+	IDetailCategoryBuilder& AccessibilityCategory = DetailLayout.EditCategory(AccessibilityCategoryKey);
+	TSharedRef<IPropertyHandle> AccessibleBehaviorPropertyHandle = DetailLayout.GetProperty(BehaviorPropertyName);
+	IDetailPropertyRow& AccessibilityRow = AccessibilityCategory.AddProperty(AccessibleBehaviorPropertyHandle);
+
+	TSharedRef<IPropertyHandle> AccessibleTextPropertyHandle = DetailLayout.GetProperty(TextPropertyName);
+	// Make sure the old *Text properties are hidden so we don't get duplicate widgets
+	DetailLayout.HideProperty(AccessibleTextPropertyHandle);
+	// The *Text properties are put into an HBox which will consist of the text box and the delegate's combo box.
+	// The combo box may not always be valid for the current selection, as its creation is determined by the ExtensionHandler.
+	TSharedRef<SHorizontalBox> CustomTextLayout = SNew(SHorizontalBox)
+	.Visibility(TAttribute<EVisibility>::Create([AccessibleBehaviorPropertyHandle]() -> EVisibility
+	{
+		// Toggle the visibility of the *Text properties on only if the *Behavior property is set to Custom
+		uint8 Behavior = 0;
+		AccessibleBehaviorPropertyHandle->GetValue(Behavior);
+		return (ESlateAccessibleBehavior)Behavior == ESlateAccessibleBehavior::Custom ? EVisibility::Visible : EVisibility::Hidden;
+	}))
+	+ SHorizontalBox::Slot()
+	.Padding(FMargin(4.0f, 0.0f))
+	[
+		AccessibleTextPropertyHandle->CreatePropertyValueWidget()
+	];
+
+	TSharedRef<SWidget> ExtensionWidget = const_cast<IDetailsView*>(DetailLayout.GetDetailsView())->GetExtensionHandler()->GenerateExtensionWidget(UWidget::StaticClass(), AccessibleTextPropertyHandle);
+	if (ExtensionWidget != SNullWidget::NullWidget)
+	{
+		CustomTextLayout->AddSlot()
+		.AutoWidth()
+		[
+			ExtensionWidget
+		];
+	}
+
+	TSharedPtr<SWidget> AccessibleBehaviorNameWidget, AccessibleBehaviorValueWidget;
+	AccessibilityRow.GetDefaultWidgets(AccessibleBehaviorNameWidget, AccessibleBehaviorValueWidget);
+
+	AccessibilityRow.CustomWidget()
+	.NameContent()
+	[
+		AccessibleBehaviorNameWidget.ToSharedRef()
+	]
+	.ValueContent()
+	.HAlign(HAlign_Fill)
+	[
+		SNew(SHorizontalBox)
+		+ SHorizontalBox::Slot()
+		.AutoWidth()
+		[
+			AccessibleBehaviorValueWidget.ToSharedRef()
+		]
+		+ SHorizontalBox::Slot()
+		[
+			CustomTextLayout
+		]
+	];
+}
+
 #undef LOCTEXT_NAMESPACE
diff --git a/Engine/Source/Editor/UMGEditor/Private/Customizations/UMGDetailCustomizations.h b/Engine/Source/Editor/UMGEditor/Private/Customizations/UMGDetailCustomizations.h
--- a/Engine/Source/Editor/UMGEditor/Private/Customizations/UMGDetailCustomizations.h
+++ b/Engine/Source/Editor/UMGEditor/Private/Customizations/UMGDetailCustomizations.h
@@ -43,6 +43,9 @@ private:
 	FReply HandleAddOrViewEventForVariable(const FName EventName, FName PropertyName, TWeakObjectPtr<UClass> PropertyClass);
 
 	int32 HandleAddOrViewIndexForButton(const FName EventName, FName PropertyName) const;
+
+	/** Handle specific customizations for ESlateAccessibleBehavior properties */
+	void CustomizeAccessibilityProperty(IDetailLayoutBuilder& DetailLayout, const FName& BehaviorPropertyName, const FName& TextPropertyName);
 private:
 
 	TWeakPtr<FWidgetBlueprintEditor> Editor;
diff --git a/Engine/Source/Programs/UnrealBuildTool/Configuration/TargetRules.cs b/Engine/Source/Programs/UnrealBuildTool/Configuration/TargetRules.cs
--- a/Engine/Source/Programs/UnrealBuildTool/Configuration/TargetRules.cs
+++ b/Engine/Source/Programs/UnrealBuildTool/Configuration/TargetRules.cs
@@ -520,9 +520,15 @@ namespace UnrealBuildTool
 		[RequiresUniqueBuildEnvironment]
 		public bool bIncludePluginsForTargetPlatforms = false;
 
-        /// <summary>
-        /// Whether to include PerfCounters support.
-        /// </summary>
+		/// <summary>
+		/// Whether to allow accessibility code in both Slate and the OS layer.
+		/// </summary>
+		[RequiresUniqueBuildEnvironment]
+		public bool bCompileWithAccessibilitySupport = true;
+
+		/// <summary>
+		/// Whether to include PerfCounters support.
+		/// </summary>
 		[RequiresUniqueBuildEnvironment]
 		[ConfigFile(ConfigHierarchyType.Engine, "/Script/BuildSettings.BuildSettings", "bWithPerfCounters")]
         public bool bWithPerfCounters = false;
@@ -1815,7 +1821,12 @@ namespace UnrealBuildTool
 			get { return Inner.bIncludePluginsForTargetPlatforms; }
 		}
 
-        public bool bWithPerfCounters
+		public bool bCompileWithAccessibilitySupport
+		{
+			get { return Inner.bCompileWithAccessibilitySupport; }
+		}
+
+		public bool bWithPerfCounters
 		{
 			get { return Inner.bWithPerfCounters; }
 		}
diff --git a/Engine/Source/Programs/UnrealBuildTool/Configuration/UEBuildTarget.cs b/Engine/Source/Programs/UnrealBuildTool/Configuration/UEBuildTarget.cs
--- a/Engine/Source/Programs/UnrealBuildTool/Configuration/UEBuildTarget.cs
+++ b/Engine/Source/Programs/UnrealBuildTool/Configuration/UEBuildTarget.cs
@@ -3076,7 +3076,16 @@ namespace UnrealBuildTool
 				GlobalCompileEnvironment.Definitions.Add("WITH_PLUGIN_SUPPORT=0");
 			}
 
-            if (Rules.bWithPerfCounters)
+			if (Rules.bCompileWithAccessibilitySupport)
+			{
+				GlobalCompileEnvironment.Definitions.Add("WITH_ACCESSIBILITY=1");
+			}
+			else
+			{
+				GlobalCompileEnvironment.Definitions.Add("WITH_ACCESSIBILITY=0");
+			}
+
+			if (Rules.bWithPerfCounters)
             {
                 GlobalCompileEnvironment.Definitions.Add("WITH_PERFCOUNTERS=1");
             }
diff --git a/Engine/Source/Runtime/ApplicationCore/ApplicationCore.Build.cs b/Engine/Source/Runtime/ApplicationCore/ApplicationCore.Build.cs
--- a/Engine/Source/Runtime/ApplicationCore/ApplicationCore.Build.cs
+++ b/Engine/Source/Runtime/ApplicationCore/ApplicationCore.Build.cs
@@ -32,7 +32,11 @@ public class ApplicationCore : ModuleRules
 			AddEngineThirdPartyPrivateStaticDependencies(Target, 
 				"XInput"
 				);
-		}
+            if (Target.bCompileWithAccessibilitySupport)
+            {
+                PublicAdditionalLibraries.Add("uiautomationcore.lib");
+            }
+        }
 		else if (Target.Platform == UnrealTargetPlatform.Mac)
 		{
 			AddEngineThirdPartyPrivateStaticDependencies(Target, 
diff --git a/Engine/Source/Runtime/ApplicationCore/Private/GenericPlatform/GenericAccessibleInterfaces.cpp b/Engine/Source/Runtime/ApplicationCore/Private/GenericPlatform/GenericAccessibleInterfaces.cpp
new file mode 100644
--- /dev/null
+++ b/Engine/Source/Runtime/ApplicationCore/Private/GenericPlatform/GenericAccessibleInterfaces.cpp
@@ -0,0 +1,5 @@
+// Copyright 1998-2019 Epic Games, Inc. All Rights Reserved.
+
+#include "GenericPlatform/GenericAccessibleInterfaces.h"
+
+DEFINE_LOG_CATEGORY(LogAccessibility);
diff --git a/Engine/Source/Runtime/ApplicationCore/Private/Windows/Accessibility/WindowsUIABaseProvider.cpp b/Engine/Source/Runtime/ApplicationCore/Private/Windows/Accessibility/WindowsUIABaseProvider.cpp
new file mode 100644
--- /dev/null
+++ b/Engine/Source/Runtime/ApplicationCore/Private/Windows/Accessibility/WindowsUIABaseProvider.cpp
@@ -0,0 +1,55 @@
+// Copyright 1998-2019 Epic Games, Inc. All Rights Reserved.
+
+#if WITH_ACCESSIBILITY
+
+#include "Windows/Accessibility/WindowsUIABaseProvider.h"
+
+#include "Misc/AssertionMacros.h"
+#include "Windows/Accessibility/WindowsUIAManager.h"
+#include "GenericPlatform/GenericAccessibleInterfaces.h"
+
+FWindowsUIABaseProvider::FWindowsUIABaseProvider(FWindowsUIAManager& InManager, TSharedRef<IAccessibleWidget> InWidget)
+	: UIAManager(&InManager)
+	, Widget(InWidget)
+	, RefCount(1)
+{
+	// Register with the UIA Manager in order to receive OnUIAManagerDestroyed signal.
+	UIAManager->ProviderList.Add(this);
+}
+
+FWindowsUIABaseProvider::~FWindowsUIABaseProvider()
+{
+	if (UIAManager)
+	{
+		UIAManager->ProviderList.Remove(this);
+	}
+}
+
+void FWindowsUIABaseProvider::OnUIAManagerDestroyed()
+{
+	UIAManager = nullptr;
+}
+
+uint32 FWindowsUIABaseProvider::IncrementRef()
+{
+	// todo: check if this needs to be threadsafe using InterlockedIncrement
+	return ++RefCount;
+}
+
+uint32 FWindowsUIABaseProvider::DecrementRef()
+{
+	ensure(RefCount > 0);
+	if (--RefCount == 0)
+	{
+		delete this;
+		return 0;
+	}
+	return RefCount;
+}
+
+bool FWindowsUIABaseProvider::IsValid() const
+{
+	return UIAManager != nullptr && Widget->IsValid();
+}
+
+#endif
diff --git a/Engine/Source/Runtime/ApplicationCore/Private/Windows/Accessibility/WindowsUIAControlProvider.cpp b/Engine/Source/Runtime/ApplicationCore/Private/Windows/Accessibility/WindowsUIAControlProvider.cpp
new file mode 100644
--- /dev/null
+++ b/Engine/Source/Runtime/ApplicationCore/Private/Windows/Accessibility/WindowsUIAControlProvider.cpp
@@ -0,0 +1,700 @@
+// Copyright 1998-2019 Epic Games, Inc. All Rights Reserved.
+
+#if WITH_ACCESSIBILITY
+
+#include "Windows/Accessibility/WindowsUIAControlProvider.h"
+#include "Windows/Accessibility/WindowsUIAWidgetProvider.h"
+#include "Windows/Accessibility/WindowsUIAPropertyGetters.h"
+#include "Windows/Accessibility/WindowsUIAManager.h"
+#include "GenericPlatform/GenericAccessibleInterfaces.h"
+
+// FWindowsUIATextRangeProvider
+
+FWindowsUIATextRangeProvider::FWindowsUIATextRangeProvider(FWindowsUIAManager& InManager, TSharedRef<IAccessibleWidget> InWidget, FTextRange InRange)
+	: FWindowsUIABaseProvider(InManager, InWidget)
+	, TextRange(InRange)
+{
+}
+
+FWindowsUIATextRangeProvider::~FWindowsUIATextRangeProvider()
+{
+}
+
+HRESULT STDCALL FWindowsUIATextRangeProvider::QueryInterface(REFIID riid, void** ppInterface)
+{
+	*ppInterface = nullptr;
+
+	if (riid == __uuidof(IUnknown))
+	{
+		*ppInterface = static_cast<ITextRangeProvider*>(this);
+	}
+	else if (riid == __uuidof(ITextRangeProvider))
+	{
+		*ppInterface = static_cast<ITextRangeProvider*>(this);
+	}
+
+	if (*ppInterface)
+	{
+		AddRef();
+		return S_OK;
+	}
+	else
+	{
+		return E_NOINTERFACE;
+	}
+}
+
+FString FWindowsUIATextRangeProvider::TextFromTextRange()
+{
+	return TextFromTextRange(Widget->AsText()->GetText(), TextRange);
+}
+
+FString FWindowsUIATextRangeProvider::TextFromTextRange(const FString& InString, const FTextRange& InRange)
+{
+	return InString.Mid(InRange.BeginIndex, InRange.EndIndex);
+}
+
+ULONG STDCALL FWindowsUIATextRangeProvider::AddRef()
+{
+	return FWindowsUIABaseProvider::IncrementRef();
+}
+
+ULONG STDCALL FWindowsUIATextRangeProvider::Release()
+{
+	return FWindowsUIABaseProvider::DecrementRef();
+}
+
+HRESULT STDCALL FWindowsUIATextRangeProvider::Clone(ITextRangeProvider** pRetVal)
+{
+	if (IsValid())
+	{
+		*pRetVal = new FWindowsUIATextRangeProvider(*UIAManager, Widget, TextRange);
+		return S_OK;
+	}
+	else
+	{
+		return UIA_E_ELEMENTNOTAVAILABLE;
+	}
+}
+
+HRESULT STDCALL FWindowsUIATextRangeProvider::Compare(ITextRangeProvider* range, BOOL* pRetVal)
+{
+	// The documentation states that different endpoints that produce the same text are not equal,
+	// but doesn't say anything about the same endpoints that come from different control providers.
+	// Perhaps we can assume that comparing text ranges from different Widgets is not valid.
+	*pRetVal = TextRange == static_cast<FWindowsUIATextRangeProvider*>(range)->TextRange;
+	return S_OK;
+}
+
+HRESULT STDCALL FWindowsUIATextRangeProvider::CompareEndpoints(TextPatternRangeEndpoint endpoint, ITextRangeProvider* targetRange, TextPatternRangeEndpoint targetEndpoint, int* pRetVal)
+{
+	FWindowsUIATextRangeProvider* CastedRange = static_cast<FWindowsUIATextRangeProvider*>(targetRange);
+	int32 ThisEndpointIndex = (endpoint == TextPatternRangeEndpoint_Start) ? TextRange.BeginIndex : TextRange.EndIndex;
+	int32 OtherEndpointIndex = (targetEndpoint == TextPatternRangeEndpoint_Start) ? CastedRange->TextRange.BeginIndex : CastedRange->TextRange.EndIndex;
+	*pRetVal = ThisEndpointIndex - OtherEndpointIndex;
+	return S_OK;
+}
+
+HRESULT STDCALL FWindowsUIATextRangeProvider::ExpandToEnclosingUnit(TextUnit unit)
+{
+	if (IsValid())
+	{
+		switch (unit)
+		{
+		case TextUnit_Character:
+			TextRange.EndIndex = FMath::Min(TextRange.BeginIndex + 1, Widget->AsText()->GetText().Len());
+			break;
+		case TextUnit_Format:
+			return E_NOTIMPL;
+		case TextUnit_Word:
+			return E_NOTIMPL;
+		case TextUnit_Line:
+			return E_NOTIMPL;
+		case TextUnit_Paragraph:
+			return E_NOTIMPL;
+		case TextUnit_Page:
+			return E_NOTIMPL;
+		case TextUnit_Document:
+			TextRange = FTextRange(0, Widget->AsText()->GetText().Len());
+			break;
+		}
+
+		return S_OK;
+	}
+	else
+	{
+		return UIA_E_ELEMENTNOTAVAILABLE;
+	}
+}
+
+HRESULT STDCALL FWindowsUIATextRangeProvider::FindAttribute(TEXTATTRIBUTEID attributeId, VARIANT val, BOOL backward, ITextRangeProvider** pRetVal)
+{
+	return E_NOTIMPL;
+}
+
+HRESULT STDCALL FWindowsUIATextRangeProvider::FindText(BSTR text, BOOL backward, BOOL ignoreCase, ITextRangeProvider** pRetVal)
+{
+	if (IsValid())
+	{
+		FString TextToSearch(text);
+		int32 FoundIndex = TextFromTextRange().Find(TextToSearch, ignoreCase ? ESearchCase::IgnoreCase : ESearchCase::CaseSensitive, backward ? ESearchDir::FromEnd : ESearchDir::FromStart);
+		if (FoundIndex == INDEX_NONE)
+		{
+			*pRetVal = nullptr;
+		}
+		else
+		{
+			const int32 StartIndex = TextRange.BeginIndex + FoundIndex;
+			*pRetVal = new FWindowsUIATextRangeProvider(*UIAManager, Widget, FTextRange(StartIndex, StartIndex + TextToSearch.Len()));
+		}
+		return S_OK;
+	}
+	else
+	{
+		return UIA_E_ELEMENTNOTAVAILABLE;
+	}
+}
+
+HRESULT STDCALL FWindowsUIATextRangeProvider::GetAttributeValue(TEXTATTRIBUTEID attributeId, VARIANT* pRetVal)
+{
+	return E_NOTIMPL;
+}
+
+HRESULT STDCALL FWindowsUIATextRangeProvider::GetBoundingRectangles(SAFEARRAY** pRetVal)
+{
+	return E_NOTIMPL;
+}
+
+HRESULT STDCALL FWindowsUIATextRangeProvider::GetEnclosingElement(IRawElementProviderSimple** pRetVal)
+{
+	if (IsValid())
+	{
+		*pRetVal = static_cast<IRawElementProviderSimple*>(&UIAManager->GetWidgetProvider(Widget));
+		return S_OK;
+	}
+	else
+	{
+		return UIA_E_ELEMENTNOTAVAILABLE;
+	}
+}
+
+HRESULT STDCALL FWindowsUIATextRangeProvider::GetText(int maxLength, BSTR* pRetVal)
+{
+	if (IsValid())
+	{
+		*pRetVal = SysAllocString(*TextFromTextRange().Left(maxLength));
+		return S_OK;
+	}
+	else
+	{
+		return UIA_E_ELEMENTNOTAVAILABLE;
+	}
+}
+
+HRESULT STDCALL FWindowsUIATextRangeProvider::Move(TextUnit unit, int count, int* pRetVal)
+{
+	return E_NOTIMPL;
+}
+
+HRESULT STDCALL FWindowsUIATextRangeProvider::MoveEndpointByUnit(TextPatternRangeEndpoint endpoint, TextUnit unit, int count, int* pRetVal)
+{
+	return E_NOTIMPL;
+}
+
+HRESULT STDCALL FWindowsUIATextRangeProvider::MoveEndpointByRange(TextPatternRangeEndpoint endpoint, ITextRangeProvider* targetRange, TextPatternRangeEndpoint targetEndpoint)
+{
+	FWindowsUIATextRangeProvider* CastedRange = static_cast<FWindowsUIATextRangeProvider*>(targetRange);
+	int32 NewIndex = (targetEndpoint == TextPatternRangeEndpoint_Start) ? CastedRange->TextRange.BeginIndex : CastedRange->TextRange.EndIndex;
+	if (endpoint == TextPatternRangeEndpoint_Start)
+	{
+		TextRange.BeginIndex = NewIndex;
+		if (TextRange.BeginIndex > TextRange.EndIndex)
+		{
+			TextRange.EndIndex = TextRange.BeginIndex;
+		}
+	}
+	else
+	{
+		TextRange.EndIndex = NewIndex;
+		if (TextRange.BeginIndex > TextRange.EndIndex)
+		{
+			TextRange.BeginIndex = TextRange.EndIndex;
+		}
+	}
+	return S_OK;
+}
+
+HRESULT STDCALL FWindowsUIATextRangeProvider::Select()
+{
+	return E_NOTIMPL;
+}
+
+HRESULT STDCALL FWindowsUIATextRangeProvider::AddToSelection()
+{
+	return E_NOTIMPL;
+}
+
+HRESULT STDCALL FWindowsUIATextRangeProvider::RemoveFromSelection()
+{
+	return E_NOTIMPL;
+}
+
+HRESULT STDCALL FWindowsUIATextRangeProvider::ScrollIntoView(BOOL alignToTop)
+{
+	return E_NOTIMPL;
+}
+
+HRESULT STDCALL FWindowsUIATextRangeProvider::GetChildren(SAFEARRAY** pRetVal)
+{
+	*pRetVal = nullptr;
+	return S_OK;
+}
+
+// ~
+
+// FWindowsUIAControlProvider
+
+FWindowsUIAControlProvider::FWindowsUIAControlProvider(FWindowsUIAManager& InManager, TSharedRef<IAccessibleWidget> InWidget)
+	: FWindowsUIABaseProvider(InManager, InWidget)
+{
+}
+
+FWindowsUIAControlProvider::~FWindowsUIAControlProvider()
+{
+}
+
+HRESULT STDCALL FWindowsUIAControlProvider::QueryInterface(REFIID riid, void** ppInterface)
+{
+	*ppInterface = nullptr;
+
+	if (riid == __uuidof(IInvokeProvider))
+	{
+		*ppInterface = static_cast<IInvokeProvider*>(this);
+	}
+	else if (riid == __uuidof(IRangeValueProvider))
+	{
+		*ppInterface = static_cast<IRangeValueProvider*>(this);
+	}
+	else if (riid == __uuidof(ITextProvider))
+	{
+		*ppInterface = static_cast<ITextProvider*>(this);
+	}
+	else if (riid == __uuidof(IToggleProvider))
+	{
+		*ppInterface = static_cast<IToggleProvider*>(this);
+	}
+	else if (riid == __uuidof(IValueProvider))
+	{
+		*ppInterface = static_cast<IValueProvider*>(this);
+	}
+	else if (riid == __uuidof(IWindowProvider))
+	{
+		*ppInterface = static_cast<IWindowProvider*>(this);
+	}
+
+	if (*ppInterface)
+	{
+		AddRef();
+		return S_OK;
+	}
+	else
+	{
+		return E_NOINTERFACE;
+	}
+}
+
+ULONG STDCALL FWindowsUIAControlProvider::AddRef()
+{
+	return FWindowsUIABaseProvider::IncrementRef();
+}
+
+ULONG STDCALL FWindowsUIAControlProvider::Release()
+{
+	return FWindowsUIABaseProvider::DecrementRef();
+}
+
+HRESULT STDCALL FWindowsUIAControlProvider::Invoke()
+{
+	if (IsValid())
+	{
+		Widget->AsActivatable()->Activate();
+		return S_OK;
+	}
+	else
+	{
+		return UIA_E_ELEMENTNOTAVAILABLE;
+	}
+}
+
+HRESULT STDCALL FWindowsUIAControlProvider::SetValue(double val)
+{
+	if (IsValid())
+	{
+		Widget->AsProperty()->SetValue(FString::SanitizeFloat(val));
+		return S_OK;
+	}
+	else
+	{
+		return UIA_E_ELEMENTNOTAVAILABLE;
+	}
+}
+
+HRESULT STDCALL FWindowsUIAControlProvider::get_Value(double* pRetVal)
+{
+	if (IsValid())
+	{
+		*pRetVal = WindowsUIAPropertyGetters::GetPropertyValue(Widget, UIA_RangeValueValuePropertyId).GetValue<double>();
+		return S_OK;
+	}
+	else
+	{
+		return UIA_E_ELEMENTNOTAVAILABLE;
+	}
+}
+
+HRESULT STDCALL FWindowsUIAControlProvider::get_IsReadOnly(BOOL* pRetVal)
+{
+	if (IsValid())
+	{
+		*pRetVal = WindowsUIAPropertyGetters::GetPropertyValue(Widget, UIA_ValueIsReadOnlyPropertyId).GetValue<bool>();
+		return S_OK;
+	}
+	else
+	{
+		return UIA_E_ELEMENTNOTAVAILABLE;
+	}
+}
+
+HRESULT STDCALL FWindowsUIAControlProvider::get_Maximum(double* pRetVal)
+{
+	if (IsValid())
+	{
+		*pRetVal = WindowsUIAPropertyGetters::GetPropertyValue(Widget, UIA_RangeValueMaximumPropertyId).GetValue<double>();
+		return S_OK;
+	}
+	else
+	{
+		return UIA_E_ELEMENTNOTAVAILABLE;
+	}
+}
+
+HRESULT STDCALL FWindowsUIAControlProvider::get_Minimum(double* pRetVal)
+{
+	if (IsValid())
+	{
+		*pRetVal = WindowsUIAPropertyGetters::GetPropertyValue(Widget, UIA_RangeValueMinimumPropertyId).GetValue<double>();
+		return S_OK;
+	}
+	else
+	{
+		return UIA_E_ELEMENTNOTAVAILABLE;
+	}
+}
+
+HRESULT STDCALL FWindowsUIAControlProvider::get_LargeChange(double* pRetVal)
+{
+	if (IsValid())
+	{
+		*pRetVal = WindowsUIAPropertyGetters::GetPropertyValue(Widget, UIA_RangeValueLargeChangePropertyId).GetValue<double>();
+		return S_OK;
+	}
+	else
+	{
+		return UIA_E_ELEMENTNOTAVAILABLE;
+	}
+}
+
+HRESULT STDCALL FWindowsUIAControlProvider::get_SmallChange(double* pRetVal)
+{
+	if (IsValid())
+	{
+		*pRetVal = WindowsUIAPropertyGetters::GetPropertyValue(Widget, UIA_RangeValueSmallChangePropertyId).GetValue<double>();
+		return S_OK;
+	}
+	else
+	{
+		return UIA_E_ELEMENTNOTAVAILABLE;
+	}
+}
+
+HRESULT STDCALL FWindowsUIAControlProvider::get_DocumentRange(ITextRangeProvider** pRetVal)
+{
+	if (IsValid())
+	{
+		*pRetVal = static_cast<ITextRangeProvider*>(new FWindowsUIATextRangeProvider(*UIAManager, Widget, FTextRange(0, Widget->AsText()->GetText().Len())));
+		return S_OK;
+	}
+	else
+	{
+		return UIA_E_ELEMENTNOTAVAILABLE;
+	}
+}
+
+HRESULT STDCALL FWindowsUIAControlProvider::get_SupportedTextSelection(SupportedTextSelection* pRetVal)
+{
+	// todo: implement selection
+	*pRetVal = SupportedTextSelection_None;
+	return S_OK;
+}
+
+HRESULT STDCALL FWindowsUIAControlProvider::GetSelection(SAFEARRAY** pRetVal)
+{
+	return E_NOTIMPL;
+}
+
+HRESULT STDCALL FWindowsUIAControlProvider::GetVisibleRanges(SAFEARRAY** pRetVal)
+{
+	return E_NOTIMPL;
+}
+
+HRESULT STDCALL FWindowsUIAControlProvider::RangeFromChild(IRawElementProviderSimple* childElement, ITextRangeProvider** pRetVal)
+{
+	return E_NOTIMPL;
+}
+
+HRESULT STDCALL FWindowsUIAControlProvider::RangeFromPoint(UiaPoint point, ITextRangeProvider** pRetVal)
+{
+	return E_NOTIMPL;
+}
+
+HRESULT STDCALL FWindowsUIAControlProvider::get_ToggleState(ToggleState* pRetVal)
+{
+	//ECheckBoxState CheckState = Widget->AsActivatable()->GetCheckedState();
+	//switch (CheckState)
+	//{
+	//case ECheckBoxState::Checked:
+	//	*pRetVal = ToggleState_On;
+	//	break;
+	//case ECheckBoxState::Unchecked:
+	//	*pRetVal = ToggleState_Off;
+	//	break;
+	//case ECheckBoxState::Undetermined:
+	//	*pRetVal = ToggleState_Indeterminate;
+	//	break;
+	//}
+	if (IsValid())
+	{
+		*pRetVal = static_cast<ToggleState>(WindowsUIAPropertyGetters::GetPropertyValue(Widget, UIA_ToggleToggleStatePropertyId).GetValue<int32>());
+		return S_OK;
+	}
+	else
+	{
+		return UIA_E_ELEMENTNOTAVAILABLE;
+	}
+}
+
+HRESULT STDCALL FWindowsUIAControlProvider::Toggle()
+{
+	if (IsValid())
+	{
+		Widget->AsActivatable()->Activate();
+		return S_OK;
+	}
+	else
+	{
+		return UIA_E_ELEMENTNOTAVAILABLE;
+	}
+}
+
+HRESULT STDCALL FWindowsUIAControlProvider::get_CanMove(BOOL *pRetVal)
+{
+	if (IsValid())
+	{
+		*pRetVal = WindowsUIAPropertyGetters::GetPropertyValue(Widget, UIA_TransformCanMovePropertyId).GetValue<bool>();
+		return S_OK;
+	}
+	else
+	{
+		return UIA_E_ELEMENTNOTAVAILABLE;
+	}
+}
+
+HRESULT STDCALL FWindowsUIAControlProvider::get_CanResize(BOOL *pRetVal)
+{
+	if (IsValid())
+	{
+		*pRetVal = WindowsUIAPropertyGetters::GetPropertyValue(Widget, UIA_TransformCanResizePropertyId).GetValue<bool>();
+		return S_OK;
+	}
+	else
+	{
+		return UIA_E_ELEMENTNOTAVAILABLE;
+	}
+}
+
+HRESULT STDCALL FWindowsUIAControlProvider::get_CanRotate(BOOL *pRetVal)
+{
+	if (IsValid())
+	{
+		*pRetVal = WindowsUIAPropertyGetters::GetPropertyValue(Widget, UIA_TransformCanRotatePropertyId).GetValue<bool>();
+		return S_OK;
+	}
+	else
+	{
+		return UIA_E_ELEMENTNOTAVAILABLE;
+	}
+}
+
+HRESULT STDCALL FWindowsUIAControlProvider::Move(double x, double y)
+{
+	return E_NOTIMPL;
+}
+
+HRESULT STDCALL FWindowsUIAControlProvider::Resize(double width, double height)
+{
+	return E_NOTIMPL;
+}
+
+HRESULT STDCALL FWindowsUIAControlProvider::Rotate(double degrees)
+{
+	return E_NOTIMPL;
+}
+
+HRESULT STDCALL FWindowsUIAControlProvider::SetValue(LPCWSTR val)
+{
+	if (IsValid())
+	{
+		Widget->AsProperty()->SetValue(FString(val));
+		return S_OK;
+	}
+	else
+	{
+		return UIA_E_ELEMENTNOTAVAILABLE;
+	}
+}
+
+HRESULT STDCALL FWindowsUIAControlProvider::get_Value(BSTR* pRetVal)
+{
+	if (IsValid())
+	{
+		*pRetVal = SysAllocString(*WindowsUIAPropertyGetters::GetPropertyValue(Widget, UIA_ValueValuePropertyId).GetValue<FString>());
+		return S_OK;
+	}
+	else
+	{
+		return UIA_E_ELEMENTNOTAVAILABLE;
+	}
+}
+
+//HRESULT STDCALL FWindowsUIAControlProvider::get_IsReadOnly(BOOL* pRetVal)
+//{
+//	*pRetVal = Widget->AsProperty()->IsReadOnly();
+//	return S_OK;
+//}
+
+HRESULT STDCALL FWindowsUIAControlProvider::Close()
+{
+	if (IsValid())
+	{
+		Widget->AsWindow()->Close();
+		return S_OK;
+	}
+	else
+	{
+		return UIA_E_ELEMENTNOTAVAILABLE;
+	}
+}
+
+HRESULT STDCALL FWindowsUIAControlProvider::get_CanMaximize(BOOL* pRetVal)
+{
+	if (IsValid())
+	{
+		*pRetVal = WindowsUIAPropertyGetters::GetPropertyValue(Widget, UIA_WindowCanMaximizePropertyId).GetValue<bool>();
+		return S_OK;
+	}
+	else
+	{
+		return UIA_E_ELEMENTNOTAVAILABLE;
+	}
+}
+
+HRESULT STDCALL FWindowsUIAControlProvider::get_CanMinimize(BOOL* pRetVal)
+{
+	if (IsValid())
+	{
+		*pRetVal = WindowsUIAPropertyGetters::GetPropertyValue(Widget, UIA_WindowCanMinimizePropertyId).GetValue<bool>();
+		return S_OK;
+	}
+	else
+	{
+		return UIA_E_ELEMENTNOTAVAILABLE;
+	}
+}
+
+HRESULT STDCALL FWindowsUIAControlProvider::get_IsModal(BOOL* pRetVal)
+{
+	if (IsValid())
+	{
+		*pRetVal = WindowsUIAPropertyGetters::GetPropertyValue(Widget, UIA_WindowIsModalPropertyId).GetValue<bool>();
+		return S_OK;
+	}
+	else
+	{
+		return UIA_E_ELEMENTNOTAVAILABLE;
+	}
+}
+
+HRESULT STDCALL FWindowsUIAControlProvider::get_IsTopmost(BOOL* pRetVal)
+{
+	// todo: not 100% sure what this is looking for. top window in hierarchy of child windows? on top of all other windows in Windows OS?
+	*pRetVal = WindowsUIAPropertyGetters::GetPropertyValue(Widget, UIA_WindowIsTopmostPropertyId).GetValue<bool>();
+	return E_NOTIMPL;
+}
+
+HRESULT STDCALL FWindowsUIAControlProvider::get_WindowInteractionState(WindowInteractionState* pRetVal)
+{
+	if (IsValid())
+	{
+		// todo: do we have a way to identify if the app is processing data vs idling?
+		*pRetVal = static_cast<WindowInteractionState>(WindowsUIAPropertyGetters::GetPropertyValue(Widget, UIA_WindowWindowInteractionStatePropertyId).GetValue<int32>());
+	}
+	else
+	{
+		*pRetVal = WindowInteractionState_Closing;
+	}
+	return S_OK;
+}
+
+HRESULT STDCALL FWindowsUIAControlProvider::get_WindowVisualState(WindowVisualState* pRetVal)
+{
+	if (IsValid())
+	{
+		*pRetVal = static_cast<WindowVisualState>(WindowsUIAPropertyGetters::GetPropertyValue(Widget, UIA_WindowWindowVisualStatePropertyId).GetValue<int32>());
+		return S_OK;
+	}
+	else
+	{
+		return UIA_E_ELEMENTNOTAVAILABLE;
+	}
+}
+
+HRESULT STDCALL FWindowsUIAControlProvider::SetVisualState(WindowVisualState state)
+{
+	if (IsValid())
+	{
+		switch (state)
+		{
+		case WindowVisualState_Normal:
+			Widget->AsWindow()->SetDisplayState(IAccessibleWindow::EWindowDisplayState::Normal);
+			break;
+		case WindowVisualState_Minimized:
+			Widget->AsWindow()->SetDisplayState(IAccessibleWindow::EWindowDisplayState::Minimize);
+			break;
+		case WindowVisualState_Maximized:
+			Widget->AsWindow()->SetDisplayState(IAccessibleWindow::EWindowDisplayState::Maximize);
+			break;
+		}
+		return S_OK;
+	}
+	else
+	{
+		return UIA_E_ELEMENTNOTAVAILABLE;
+	}
+}
+
+HRESULT STDCALL FWindowsUIAControlProvider::WaitForInputIdle(int milliseconds, BOOL* pRetVal)
+{
+	return E_NOTIMPL;
+}
+
+#endif
diff --git a/Engine/Source/Runtime/ApplicationCore/Private/Windows/Accessibility/WindowsUIAManager.cpp b/Engine/Source/Runtime/ApplicationCore/Private/Windows/Accessibility/WindowsUIAManager.cpp
new file mode 100644
--- /dev/null
+++ b/Engine/Source/Runtime/ApplicationCore/Private/Windows/Accessibility/WindowsUIAManager.cpp
@@ -0,0 +1,223 @@
+// Copyright 1998-2019 Epic Games, Inc. All Rights Reserved.
+
+#if WITH_ACCESSIBILITY
+
+#include "Windows/Accessibility/WindowsUIAManager.h"
+
+#include "Windows/AllowWindowsPlatformTypes.h"
+#include <Ole2.h>
+#include "Windows/HideWindowsPlatformTypes.h"
+#include <UIAutomation.h>
+
+#include "Misc/Variant.h"
+#include "Windows/Accessibility/WindowsUIAControlProvider.h"
+#include "Windows/Accessibility/WindowsUIAWidgetProvider.h"
+#include "Windows/Accessibility/WindowsUIAPropertyGetters.h"
+#include "Windows/WindowsApplication.h"
+#include "Windows/WindowsWindow.h"
+
+TMap<EAccessibleWidgetType, ULONG> FWindowsUIAManager::WidgetTypeToWindowsTypeMap;
+TMap<EAccessibleWidgetType, FText> FWindowsUIAManager::WidgetTypeToTextMap;
+
+#define LOCTEXT_NAMESPACE "SlateAccessibility"
+
+/**
+ * Helper function to raise a UIA property changed event.
+ */
+void EmitPropertyChangedEvent(FWindowsUIAWidgetProvider* Provider, PROPERTYID Property, const FVariant& OldValue, const FVariant& NewValue)
+{
+	UE_LOG(LogAccessibility, VeryVerbose, TEXT("UIA Property Changed: %i"), Property);
+	UiaRaiseAutomationPropertyChangedEvent(static_cast<IRawElementProviderSimple*>(Provider), Property,
+										   WindowsUIAPropertyGetters::FVariantToWindowsVariant(OldValue),
+										   WindowsUIAPropertyGetters::FVariantToWindowsVariant(NewValue));
+}
+
+FWindowsUIAManager::FWindowsUIAManager(const FWindowsApplication& InApplication)
+	: WindowsApplication(InApplication)
+{
+	OnAccessibleMessageHandlerChanged();
+
+	if (WidgetTypeToWindowsTypeMap.Num() == 0)
+	{
+			WidgetTypeToWindowsTypeMap.Add(EAccessibleWidgetType::Button, UIA_ButtonControlTypeId);
+			WidgetTypeToWindowsTypeMap.Add(EAccessibleWidgetType::CheckBox, UIA_CheckBoxControlTypeId);
+			WidgetTypeToWindowsTypeMap.Add(EAccessibleWidgetType::ComboBox, UIA_ComboBoxControlTypeId);
+			WidgetTypeToWindowsTypeMap.Add(EAccessibleWidgetType::Hyperlink, UIA_HyperlinkControlTypeId);
+			WidgetTypeToWindowsTypeMap.Add(EAccessibleWidgetType::Image, UIA_ImageControlTypeId);
+			WidgetTypeToWindowsTypeMap.Add(EAccessibleWidgetType::Layout, UIA_PaneControlTypeId);
+			WidgetTypeToWindowsTypeMap.Add(EAccessibleWidgetType::ScrollBar, UIA_ScrollBarControlTypeId);
+			WidgetTypeToWindowsTypeMap.Add(EAccessibleWidgetType::Slider, UIA_SliderControlTypeId);
+			WidgetTypeToWindowsTypeMap.Add(EAccessibleWidgetType::Text, UIA_TextControlTypeId);
+			WidgetTypeToWindowsTypeMap.Add(EAccessibleWidgetType::TextEdit, UIA_EditControlTypeId);
+			WidgetTypeToWindowsTypeMap.Add(EAccessibleWidgetType::Window, UIA_WindowControlTypeId);
+
+			WidgetTypeToTextMap.Add(EAccessibleWidgetType::Button, LOCTEXT("ControlTypeButton", "button"));
+			WidgetTypeToTextMap.Add(EAccessibleWidgetType::CheckBox, LOCTEXT("ControlTypeCheckBox", "check box"));
+			WidgetTypeToTextMap.Add(EAccessibleWidgetType::ComboBox, LOCTEXT("ControlTypeComboBox", "combobox"));
+			WidgetTypeToTextMap.Add(EAccessibleWidgetType::Hyperlink, LOCTEXT("ControlTypeHyperlink", "hyperlink"));
+			WidgetTypeToTextMap.Add(EAccessibleWidgetType::Image, LOCTEXT("ControlTypeImage", "image"));
+			WidgetTypeToTextMap.Add(EAccessibleWidgetType::Layout, LOCTEXT("ControlTypeLayout", "pane"));
+			WidgetTypeToTextMap.Add(EAccessibleWidgetType::ScrollBar, LOCTEXT("ControlTypeScrollBar", "scroll bar"));
+			WidgetTypeToTextMap.Add(EAccessibleWidgetType::Slider, LOCTEXT("ControlTypeSlider", "slider"));
+			WidgetTypeToTextMap.Add(EAccessibleWidgetType::Text, LOCTEXT("ControlTypeText", "text"));
+			WidgetTypeToTextMap.Add(EAccessibleWidgetType::TextEdit, LOCTEXT("ControlTypeTextEdit", "edit"));
+			WidgetTypeToTextMap.Add(EAccessibleWidgetType::Window, LOCTEXT("ControlTypeWindow", "window"));
+	}
+}
+
+void FWindowsUIAManager::OnAccessibleMessageHandlerChanged()
+{
+	WindowsApplication.GetAccessibleMessageHandler()->SetAccessibleEventDelegate(FGenericAccessibleMessageHandler::FAccessibleEvent::CreateRaw(this, &FWindowsUIAManager::OnEventRaised));
+}
+
+FWindowsUIAManager::~FWindowsUIAManager()
+{
+	WindowsApplication.GetAccessibleMessageHandler()->SetAccessibleEventDelegate(FGenericAccessibleMessageHandler::FAccessibleEvent());
+	for (FWindowsUIABaseProvider* Provider : ProviderList)
+	{
+		// The UIA Manager may be deleted before the Providers are, and external applications may attempt to call functions on those Providers.
+		Provider->OnUIAManagerDestroyed();
+	}
+}
+
+FWindowsUIAWidgetProvider& FWindowsUIAManager::GetWidgetProvider(TSharedRef<IAccessibleWidget> InWidget)
+{
+	FWindowsUIAWidgetProvider** CachedProvider = CachedWidgetProviders.Find(InWidget);
+	if (CachedProvider)
+	{
+		(*CachedProvider)->AddRef();
+		return **CachedProvider;
+	}
+	else if (InWidget->AsWindow())
+	{
+		return *CachedWidgetProviders.Add(InWidget, new FWindowsUIAWindowProvider(*this, InWidget));
+	}
+	else
+	{
+		return *CachedWidgetProviders.Add(InWidget, new FWindowsUIAWidgetProvider(*this, InWidget));
+	}
+}
+
+FWindowsUIAWindowProvider& FWindowsUIAManager::GetWindowProvider(TSharedRef<FWindowsWindow> InWindow)
+{
+	if (CachedWidgetProviders.Num() == 0)
+	{
+		// Enable application accessibility if this is the first Provider. The first Get*Provider() request
+		// MUST go through this function since IAccessibleWidgets will not exist until SetActive is called.
+		WindowsApplication.GetAccessibleMessageHandler()->SetActive(true);
+	}
+
+	return static_cast<FWindowsUIAWindowProvider&>(GetWidgetProvider(WindowsApplication.GetAccessibleMessageHandler()->GetAccessibleWindow(InWindow).ToSharedRef()));
+}
+
+int FWindowsUIAManager::GetRuntimeId(FWindowsUIAWidgetProvider* InProvider)
+{
+	const int* Value = RuntimeIds.Find(InProvider);
+	if (Value)
+	{
+		return *Value;
+	}
+	else
+	{
+		static int RuntimeIdCounter = 0;
+		if (RuntimeIdCounter == INT_MAX)
+		{
+			RuntimeIdCounter = INT_MIN;
+		}
+		return RuntimeIds.Add(InProvider, RuntimeIdCounter++);
+	}
+}
+
+void FWindowsUIAManager::OnWidgetProviderRemoved(TSharedRef<IAccessibleWidget> InWidget)
+{
+	FWindowsUIAWidgetProvider* Provider = nullptr;
+	if (CachedWidgetProviders.RemoveAndCopyValue(InWidget, Provider))
+	{
+		RuntimeIds.Remove(Provider);
+	}
+
+	if (CachedWidgetProviders.Num() == 0)
+	{
+		// If the last widget Provider is being removed, we can disable application accessibility.
+		// Technically an external application could still be running listening for mouse/keyboard events,
+		// but in practice I don't think its realistic to do this while having no references to any Provider.
+		//
+		// Note that there could still be control Providers with valid references. In practice I don't think
+		// this is possible, but if we run into problems then we can simply call AddRef and Release on the
+		// underlying widget Provider whenever a control Provider gets added/removed.
+		WindowsApplication.GetAccessibleMessageHandler()->SetActive(false);
+	}
+}
+
+#if WINVER > 0x0A00 // Win10
+typedef HRESULT(WINAPI* UiaRaiseNotificationEventFunc)(IRawElementProviderSimple*, NotificationKind, NotificationProcessing, BSTR, BSTR);
+#endif
+
+void FWindowsUIAManager::OnEventRaised(TSharedRef<IAccessibleWidget> Widget, EAccessibleEvent Event, FVariant OldValue, FVariant NewValue)
+{
+	if (UiaClientsAreListening())
+	{
+		FScopedWidgetProvider ScopedProvider(GetWidgetProvider(Widget));
+
+		switch (Event)
+		{
+		case EAccessibleEvent::FocusChange:
+		{
+			// On focus change, emit a generic FocusChanged event as well as a per-Provider PropertyChanged event
+			// todo: handle difference between any focus vs keyboard focus
+			if (Widget->HasFocus())
+			{
+				UiaRaiseAutomationEvent(&ScopedProvider.Provider, UIA_AutomationFocusChangedEventId);
+			}
+			EmitPropertyChangedEvent(&ScopedProvider.Provider, UIA_HasKeyboardFocusPropertyId, OldValue, NewValue);
+			break;
+		}
+		case EAccessibleEvent::Activate:
+			if (ScopedProvider.Provider.SupportsInterface(UIA_TogglePatternId))
+			{
+				EmitPropertyChangedEvent(&ScopedProvider.Provider, UIA_ToggleToggleStatePropertyId, OldValue, NewValue);
+
+			}
+			else if (ScopedProvider.Provider.SupportsInterface(UIA_InvokePatternId))
+			{
+				UiaRaiseAutomationEvent(&ScopedProvider.Provider, UIA_Invoke_InvokedEventId);
+			}
+			break;
+		case EAccessibleEvent::Notification:
+		{
+#if WINVER > 0x0A00 // Win10
+			// GetProcAddress returns a function type with no args and the compiler doesn't like the cast to one with args, but we know better.
+			// todo: are these checks necessary if the compiler is targeting Win10 already?
+			#pragma warning(suppress: 4191)
+			UiaRaiseNotificationEventFunc NotificationFunc = (UiaRaiseNotificationEventFunc)GetProcAddress(GetModuleHandle(TEXT("Uiautomationcore.dll")), "UiaRaiseNotificationEvent");
+			if (NotificationFunc)
+			{
+				NotificationFunc(&ScopedProvider.Provider, NotificationKind_ActionCompleted, NotificationProcessing_All, SysAllocString(*NewValue.GetValue<FString>()), SysAllocString(TEXT("")));
+			}
+#endif
+			break;
+		}
+		case EAccessibleEvent::BeforeRemoveFromParent:
+		{
+			// ChildRemoved events require an ID passed along with them to identify which child (which no longer exists) was removed
+			int id[2] = { UiaAppendRuntimeId, GetRuntimeId(&ScopedProvider.Provider) };
+			FScopedWidgetProvider ParentProvider(GetWidgetProvider(Widget->GetParent().ToSharedRef()));
+			UiaRaiseStructureChangedEvent(&ParentProvider.Provider, StructureChangeType_ChildRemoved, id, ARRAYSIZE(id));
+			break;
+		}
+		case EAccessibleEvent::AfterAddToParent:
+		{
+			FScopedWidgetProvider ParentProvider(GetWidgetProvider(Widget->GetParent().ToSharedRef()));
+			UiaRaiseStructureChangedEvent(&ParentProvider.Provider, StructureChangeType_ChildAdded, nullptr, 0);
+			break;
+		}
+		case EAccessibleEvent::WidgetRemoved:
+			UiaDisconnectProvider(&ScopedProvider.Provider);
+			break;
+		}
+	}
+}
+
+#undef LOCTEXT_NAMESPACE
+
+#endif
diff --git a/Engine/Source/Runtime/ApplicationCore/Private/Windows/Accessibility/WindowsUIAPropertyGetters.cpp b/Engine/Source/Runtime/ApplicationCore/Private/Windows/Accessibility/WindowsUIAPropertyGetters.cpp
new file mode 100644
--- /dev/null
+++ b/Engine/Source/Runtime/ApplicationCore/Private/Windows/Accessibility/WindowsUIAPropertyGetters.cpp
@@ -0,0 +1,106 @@
+// Copyright 1998-2019 Epic Games, Inc. All Rights Reserved.
+
+#if WITH_ACCESSIBILITY
+
+#include "Windows/Accessibility/WindowsUIAPropertyGetters.h"
+
+#include "GenericPlatform/GenericAccessibleInterfaces.h"
+#include "Misc/Variant.h"
+
+namespace WindowsUIAPropertyGetters
+{
+
+VARIANT GetPropertyValueWindows(TSharedRef<IAccessibleWidget> AccessibleWidget, PROPERTYID WindowsPropertyId)
+{
+	return FVariantToWindowsVariant(GetPropertyValue(AccessibleWidget, WindowsPropertyId));
+}
+
+FVariant GetPropertyValue(TSharedRef<IAccessibleWidget> AccessibleWidget, PROPERTYID WindowsPropertyId)
+{
+	switch (WindowsPropertyId)
+	{
+	case UIA_ValueIsReadOnlyPropertyId:
+	case UIA_RangeValueIsReadOnlyPropertyId:
+		return !AccessibleWidget->AsProperty()->IsReadOnly();
+	case UIA_RangeValueLargeChangePropertyId:
+	case UIA_RangeValueSmallChangePropertyId:
+		return static_cast<double>(AccessibleWidget->AsProperty()->GetStepSize());
+	case UIA_RangeValueMaximumPropertyId:
+		return static_cast<double>(AccessibleWidget->AsProperty()->GetMaximum());
+	case UIA_RangeValueMinimumPropertyId:
+		return static_cast<double>(AccessibleWidget->AsProperty()->GetMinimum());
+	case UIA_RangeValueValuePropertyId:
+		return static_cast<double>(FCString::Atof(*AccessibleWidget->AsProperty()->GetValue()));
+	case UIA_ToggleToggleStatePropertyId:
+		return static_cast<int32>(AccessibleWidget->AsActivatable()->GetCheckedState() ? ToggleState_On : ToggleState_Off);
+	case UIA_TransformCanMovePropertyId:
+		return false;
+	case UIA_TransformCanResizePropertyId:
+		return false;
+	case UIA_TransformCanRotatePropertyId:
+		return false;
+	case UIA_ValueValuePropertyId:
+		return AccessibleWidget->AsProperty()->GetValue();
+	case UIA_WindowCanMaximizePropertyId:
+		return AccessibleWidget->AsWindow()->SupportsDisplayState(IAccessibleWindow::EWindowDisplayState::Maximize);
+	case UIA_WindowCanMinimizePropertyId:
+		return AccessibleWidget->AsWindow()->SupportsDisplayState(IAccessibleWindow::EWindowDisplayState::Minimize);
+	case UIA_WindowIsModalPropertyId:
+		return AccessibleWidget->AsWindow()->IsModal();
+	case UIA_WindowIsTopmostPropertyId:
+		return false;
+	case UIA_WindowWindowInteractionStatePropertyId:
+		return static_cast<int32>(WindowInteractionState_Running);
+	case UIA_WindowWindowVisualStatePropertyId:
+		switch (AccessibleWidget->AsWindow()->GetDisplayState())
+		{
+		case IAccessibleWindow::EWindowDisplayState::Normal:
+			return static_cast<int32>(WindowVisualState_Normal);
+		case IAccessibleWindow::EWindowDisplayState::Minimize:
+			return static_cast<int32>(WindowVisualState_Minimized);
+		case IAccessibleWindow::EWindowDisplayState::Maximize:
+			return static_cast<int32>(WindowVisualState_Maximized);
+		}
+	}
+
+	return FVariant();
+}
+
+VARIANT FVariantToWindowsVariant(const FVariant& Value)
+{
+	VARIANT OutValue;
+	switch (Value.GetType())
+	{
+	case EVariantTypes::Bool:
+		OutValue.vt = VT_BOOL;
+		OutValue.boolVal = Value.GetValue<bool>() ? VARIANT_TRUE : VARIANT_FALSE;
+		break;
+	case EVariantTypes::Double:
+		OutValue.vt = VT_R8;
+		OutValue.dblVal = Value.GetValue<double>();
+		break;
+	case EVariantTypes::Float:
+		OutValue.vt = VT_R4;
+		OutValue.fltVal = Value.GetValue<float>();
+		break;
+	case EVariantTypes::Int32:
+		OutValue.vt = VT_I4;
+		// todo: verify correct behavior on different operating systems
+		OutValue.intVal = Value.GetValue<int32>();
+		break;
+	case EVariantTypes::String:
+		OutValue.vt = VT_BSTR;
+		OutValue.bstrVal = SysAllocString(*Value.GetValue<FString>());
+		break;
+	case EVariantTypes::Empty:
+	default:
+		OutValue.vt = VT_EMPTY;
+		break;
+	}
+
+	return OutValue;
+}
+
+}
+
+#endif
diff --git a/Engine/Source/Runtime/ApplicationCore/Private/Windows/Accessibility/WindowsUIAWidgetProvider.cpp b/Engine/Source/Runtime/ApplicationCore/Private/Windows/Accessibility/WindowsUIAWidgetProvider.cpp
new file mode 100644
--- /dev/null
+++ b/Engine/Source/Runtime/ApplicationCore/Private/Windows/Accessibility/WindowsUIAWidgetProvider.cpp
@@ -0,0 +1,689 @@
+// Copyright 1998-2019 Epic Games, Inc. All Rights Reserved.
+
+#if WITH_ACCESSIBILITY
+
+#include "Windows/Accessibility/WindowsUIAWidgetProvider.h"
+
+#include "GenericPlatform/GenericAccessibleInterfaces.h"
+#include "Windows/Accessibility/WindowsUIAControlProvider.h"
+#include "Windows/Accessibility/WindowsUIAManager.h"
+#include "Windows/Accessibility/WindowsUIAPropertyGetters.h"
+#include "Windows/WindowsWindow.h"
+
+DECLARE_CYCLE_STAT(TEXT("Windows Accessibility: Navigate"), STAT_AccessibilityWindowsNavigate, STATGROUP_Accessibility);
+DECLARE_CYCLE_STAT(TEXT("Windows Accessibility: GetProperty"), STAT_AccessibilityWindowsGetProperty, STATGROUP_Accessibility);
+
+#define LOCTEXT_NAMESPACE "SlateAccessibility"
+
+/** Convert our accessible widget type to a Windows control ID */
+ULONG WidgetTypeToControlType(TSharedRef<IAccessibleWidget> Widget)
+{
+	ULONG* Type = FWindowsUIAManager::WidgetTypeToWindowsTypeMap.Find(Widget->GetWidgetType());
+	if (Type)
+	{
+		return *Type;
+	}
+	else
+	{
+		return UIA_CustomControlTypeId;
+	}
+}
+
+/**
+ * Convert our accessible widget type to a human-readable localized string
+ * See https://docs.microsoft.com/en-us/windows/desktop/winauto/uiauto-automation-element-propids for rules.
+ */
+FString WidgetTypeToLocalizedString(TSharedRef<IAccessibleWidget> Widget)
+{
+	FText* Text = FWindowsUIAManager::WidgetTypeToTextMap.Find(Widget->GetWidgetType());
+	if (Text)
+	{
+		return Text->ToString();
+	}
+	else
+	{
+		static FText CustomControlTypeName = LOCTEXT("ControlTypeCustom", "custom");
+		return CustomControlTypeName.ToString();
+	}
+}
+
+// FWindowsUIAWidgetProvider methods
+
+FWindowsUIAWidgetProvider::FWindowsUIAWidgetProvider(FWindowsUIAManager& InManager, TSharedRef<IAccessibleWidget> InWidget)
+	: FWindowsUIABaseProvider(InManager, InWidget)
+{
+	//UpdateCachedProperties();
+}
+
+FWindowsUIAWidgetProvider::~FWindowsUIAWidgetProvider()
+{
+	if (UIAManager)
+	{
+		UIAManager->OnWidgetProviderRemoved(Widget);
+	}
+}
+
+//void FWindowsUIAWidgetProvider::UpdateCachedProperty(PROPERTYID PropertyId)
+//{
+//	FVariant& CachedValue = CachedPropertyValues.FindOrAdd(PropertyId);
+//	FVariant CurrentValue = WindowsUIAPropertyGetters::GetPropertyValue(Widget, PropertyId);
+//	if (CachedValue.IsEmpty())
+//	{
+//		CachedValue = CurrentValue;
+//	}
+//	else if (CachedValue != CurrentValue)
+//	{
+//		UE_LOG(LogAccessibility, VeryVerbose, TEXT("UIA Property Changed: %i"), PropertyId);
+//		UiaRaiseAutomationPropertyChangedEvent(
+//			static_cast<IRawElementProviderSimple*>(this), PropertyId,
+//			WindowsUIAPropertyGetters::FVariantToWindowsVariant(CachedValue),
+//			WindowsUIAPropertyGetters::FVariantToWindowsVariant(CurrentValue));
+//
+//		CachedValue = CurrentValue;
+//	}
+//}
+//
+//void FWindowsUIAWidgetProvider::UpdateCachedProperties()
+//{
+//	if (IsValid())
+//	{
+//		if (SupportsInterface(UIA_RangeValuePatternId))
+//		{
+//			UpdateCachedProperty(UIA_RangeValueIsReadOnlyPropertyId);
+//			UpdateCachedProperty(UIA_RangeValueValuePropertyId);
+//			UpdateCachedProperty(UIA_RangeValueIsReadOnlyPropertyId);
+//			UpdateCachedProperty(UIA_RangeValueMinimumPropertyId);
+//			UpdateCachedProperty(UIA_RangeValueMaximumPropertyId);
+//			UpdateCachedProperty(UIA_RangeValueLargeChangePropertyId);
+//			UpdateCachedProperty(UIA_RangeValueSmallChangePropertyId);
+//		}
+//
+//		if (SupportsInterface(UIA_TogglePatternId))
+//		{
+//			UpdateCachedProperty(UIA_ToggleToggleStatePropertyId);
+//		}
+//
+//		if (SupportsInterface(UIA_ValuePatternId))
+//		{
+//			UpdateCachedProperty(UIA_ValueIsReadOnlyPropertyId);
+//			UpdateCachedProperty(UIA_ValueValuePropertyId);
+//		}
+//
+//		if (SupportsInterface(UIA_WindowPatternId))
+//		{
+//			UpdateCachedProperty(UIA_WindowCanMaximizePropertyId);
+//			UpdateCachedProperty(UIA_WindowCanMinimizePropertyId);
+//			UpdateCachedProperty(UIA_WindowIsModalPropertyId);
+//			UpdateCachedProperty(UIA_WindowIsTopmostPropertyId);
+//			UpdateCachedProperty(UIA_WindowWindowInteractionStatePropertyId);
+//			UpdateCachedProperty(UIA_WindowWindowVisualStatePropertyId);
+//			UpdateCachedProperty(UIA_TransformCanMovePropertyId);
+//			UpdateCachedProperty(UIA_TransformCanRotatePropertyId);
+//			UpdateCachedProperty(UIA_TransformCanResizePropertyId);
+//		}
+//	}
+//}
+
+HRESULT STDCALL FWindowsUIAWidgetProvider::QueryInterface(REFIID riid, void** ppInterface)
+{
+	if (riid == __uuidof(IUnknown))
+	{
+		*ppInterface = static_cast<IRawElementProviderSimple*>(this);
+	}
+	else if (riid == __uuidof(IRawElementProviderSimple))
+	{
+		*ppInterface = static_cast<IRawElementProviderSimple*>(this);
+	}
+	else if (riid == __uuidof(IRawElementProviderFragment))
+	{
+		*ppInterface = static_cast<IRawElementProviderFragment*>(this);
+	}
+	else
+	{
+		*ppInterface = nullptr;
+	}
+
+	if (*ppInterface)
+	{
+		// QueryInterface is the one exception where we need to call AddRef without going through GetWidgetProvider().
+		AddRef();
+		return S_OK;
+	}
+	else
+	{
+		return E_NOINTERFACE;
+	}
+}
+
+ULONG STDCALL FWindowsUIAWidgetProvider::AddRef()
+{
+	return FWindowsUIABaseProvider::IncrementRef();
+}
+
+ULONG STDCALL FWindowsUIAWidgetProvider::Release()
+{
+	return FWindowsUIABaseProvider::DecrementRef();
+}
+
+bool FWindowsUIAWidgetProvider::SupportsInterface(PATTERNID PatternId) const
+{
+	switch (PatternId)
+	{
+	case UIA_InvokePatternId:
+	{
+		IAccessibleActivatable* Activatable = Widget->AsActivatable();
+		// Toggle and Invoke are mutually exclusive
+		return Activatable && !Activatable->IsCheckable();
+	}
+	case UIA_RangeValuePatternId:
+	{
+		IAccessibleProperty* Property = Widget->AsProperty();
+		// Value and RangeValue are mutually exclusive
+		return Property && Property->GetStepSize() > 0.0f;
+	}
+	case UIA_TextPatternId:
+	{
+		return Widget->AsText() != nullptr;
+	}
+	case UIA_TogglePatternId:
+	{
+		IAccessibleActivatable* Activatable = Widget->AsActivatable();
+		return Activatable && Activatable->IsCheckable();
+	}
+	case UIA_ValuePatternId:
+	{
+		IAccessibleProperty* Property = Widget->AsProperty();
+		return Property && FMath::IsNearlyZero(Property->GetStepSize());
+	}
+	}
+	return false;
+}
+
+HRESULT STDCALL FWindowsUIAWidgetProvider::get_ProviderOptions(ProviderOptions* pRetVal)
+{
+	// ServerSideProvider means that we are creating the definition of the accessible widgets for Clients (eg screenreaders) to consume.
+	// UseComThreading is necessary to ensure that COM messages are properly routed to the main thread.
+	*pRetVal = static_cast<ProviderOptions>(ProviderOptions_ServerSideProvider | ProviderOptions_UseComThreading);
+	return S_OK;
+}
+
+HRESULT STDCALL FWindowsUIAWidgetProvider::GetPatternProvider(PATTERNID patternId, IUnknown** pRetVal)
+{
+	if (IsValid())
+	{
+		if (SupportsInterface(patternId))
+		{
+			// FWindowsUIAControlProvider implements all possible control providers that we support.
+			FWindowsUIAControlProvider* ControlProvider = new FWindowsUIAControlProvider(*UIAManager, Widget);
+			switch (patternId)
+			{
+			case UIA_InvokePatternId:
+				*pRetVal = static_cast<IInvokeProvider*>(ControlProvider);
+				break;
+			case UIA_RangeValuePatternId:
+				*pRetVal = static_cast<IRangeValueProvider*>(ControlProvider);
+				break;
+			case UIA_TextPatternId:
+				*pRetVal = static_cast<ITextProvider*>(ControlProvider);
+				break;
+			case UIA_TogglePatternId:
+				*pRetVal = static_cast<IToggleProvider*>(ControlProvider);
+				break;
+			case UIA_ValuePatternId:
+				*pRetVal = static_cast<IValueProvider*>(ControlProvider);
+				break;
+			default:
+				UE_LOG(LogAccessibility, Error, TEXT("FWindowsUIAWidgetProvider::SupportsInterface() returned true, but was unhandled in GetPatternProvider(). PatternId = %i"), patternId);
+				*pRetVal = nullptr;
+				ControlProvider->Release();
+				break;
+			}
+		}
+		return S_OK;
+	}
+	else
+	{
+		return UIA_E_ELEMENTNOTAVAILABLE;
+	}
+}
+
+HRESULT STDCALL FWindowsUIAWidgetProvider::GetPropertyValue(PROPERTYID propertyId, VARIANT* pRetVal)
+{
+	SCOPE_CYCLE_COUNTER(STAT_AccessibilityWindowsGetProperty);
+
+	if (!IsValid())
+	{
+		return UIA_E_ELEMENTNOTAVAILABLE;
+	}
+
+	bool bValid = true;
+
+	// https://docs.microsoft.com/en-us/windows/desktop/winauto/uiauto-automation-element-propids
+	// potential other properties:
+	// - UIA_CenterPointPropertyId
+	// - UIA_ControllerForPropertyId (eg. to make clicking on a label also click a checkbox)
+	// - UIA_DescribedByPropertyId (same as LabeledBy? not sure)
+	// - UIA_IsDialogPropertyId (doesn't seem to exist in our version of UIA)
+	// - UIA_ItemStatusPropertyId (eg. busy, loading)
+	// - UIA_LabeledByPropertyId (eg. to describe a checkbox from a separate label)
+	// - UIA_NativeWindowHandlePropertyId (unclear if this is only for windows or all widgets)
+	// - UIA_PositionInSetPropertyId (position in list of items, could be useful)
+	// - UIA_SizeOfSetPropertyId (number of items in list, could be useful)
+	switch (propertyId)
+	{
+	//case UIA_AcceleratorKeyPropertyId:
+	//	pRetVal->vt = VT_BSTR;
+	//	// todo: hotkey, used with IInvokeProvider
+	//	pRetVal->bstrVal = SysAllocString(L"");
+	//	break;
+	//case UIA_AccessKeyPropertyId:
+	//	pRetVal->vt = VT_BSTR;
+	//	// todo: activates menu, like F in &File
+	//	pRetVal->bstrVal = SysAllocString(L"");
+	//	break;
+	//case UIA_AutomationIdPropertyId:
+	//	pRetVal->vt = VT_BSTR;
+	//	// todo: an identifiable name for the widget for automation testing, used like App.Find("MainMenuBar")
+	//	pRetVal->bstrVal = SysAllocString(L"");
+	//	break;
+	case UIA_BoundingRectanglePropertyId:
+		pRetVal->vt = VT_R8 | VT_ARRAY;
+		pRetVal->parray = SafeArrayCreateVector(VT_R8, 0, 4);
+		if (pRetVal->parray)
+		{
+			FBox2D Box = Widget->GetBounds();
+			LONG i = 0;
+			bValid &= (SafeArrayPutElement(pRetVal->parray, &i, &Box.Min.X) != S_OK);
+			i = 1;
+			bValid &= (SafeArrayPutElement(pRetVal->parray, &i, &Box.Max.X) != S_OK);
+			i = 2;
+			bValid &= (SafeArrayPutElement(pRetVal->parray, &i, &Box.Min.Y) != S_OK);
+			i = 3;
+			bValid &= (SafeArrayPutElement(pRetVal->parray, &i, &Box.Max.Y) != S_OK);
+		}
+		else
+		{
+			bValid = false;
+		}
+		break;
+	case UIA_ClassNamePropertyId:
+		pRetVal->vt = VT_BSTR;
+		pRetVal->bstrVal = SysAllocString(*Widget->GetClassName());
+		break;
+	case UIA_ControlTypePropertyId:
+		pRetVal->vt = VT_I4;
+		pRetVal->lVal = WidgetTypeToControlType(Widget);
+		break;
+	//case UIA_CulturePropertyId:
+	//	pRetVal->vt = VT_I4;
+	//	// todo: https://docs.microsoft.com/en-us/windows/desktop/Intl/language-identifier-constants-and-strings
+	//	pRetVal->lVal = 0;
+	//	break;
+	case UIA_FrameworkIdPropertyId:
+		pRetVal->vt = VT_BSTR;
+		// todo: figure out what goes here
+		pRetVal->bstrVal = SysAllocString(*LOCTEXT("Slate", "Slate").ToString());
+		break;
+	//case UIA_FullDescriptionPropertyId:
+	//	pRetVal->vt = VT_BSTR;
+	//	// todo: how is this different from help text?
+	//	pRetVal->bstrVal = SysAllocString(*Widget->GetHelpText());
+	//	break;
+	case UIA_HasKeyboardFocusPropertyId:
+		pRetVal->vt = VT_BOOL;
+		pRetVal->boolVal = Widget->HasFocus() ? VARIANT_TRUE : VARIANT_FALSE;
+		break;
+	case UIA_HelpTextPropertyId:
+		pRetVal->vt = VT_BSTR;
+		pRetVal->bstrVal = SysAllocString(*Widget->GetHelpText());
+		break;
+	case UIA_IsContentElementPropertyId:
+		pRetVal->vt = VT_BOOL;
+		// todo: https://docs.microsoft.com/en-us/windows/desktop/winauto/uiauto-treeoverview
+		pRetVal->boolVal = VARIANT_TRUE;
+		break;
+	case UIA_IsControlElementPropertyId:
+		pRetVal->vt = VT_BOOL;
+		// todo: https://docs.microsoft.com/en-us/windows/desktop/winauto/uiauto-treeoverview
+		pRetVal->boolVal = VARIANT_TRUE;
+		break;
+	case UIA_IsEnabledPropertyId:
+		pRetVal->vt = VT_BOOL;
+		pRetVal->boolVal = Widget->IsEnabled() ? VARIANT_TRUE : VARIANT_FALSE;
+		break;
+	case UIA_IsKeyboardFocusablePropertyId:
+		pRetVal->vt = VT_BOOL;
+		pRetVal->boolVal = Widget->SupportsFocus() ? VARIANT_TRUE : VARIANT_FALSE;
+		break;
+	case UIA_IsOffscreenPropertyId:
+		pRetVal->vt = VT_BOOL;
+		pRetVal->boolVal = Widget->IsHidden() ? VARIANT_TRUE : VARIANT_FALSE;
+		break;
+	case UIA_IsPasswordPropertyId:
+		if (Widget->AsProperty())
+		{
+			pRetVal->vt = VT_BOOL;
+			pRetVal->boolVal = Widget->AsProperty()->IsPassword() ? VARIANT_TRUE : VARIANT_FALSE;
+		}
+		break;
+//#if WINVER >= 0x0603 // Windows 8.1
+//	case UIA_IsPeripheralPropertyId:
+//		pRetVal->vt = VT_BOOL;
+//		// todo: see https://docs.microsoft.com/en-us/windows/desktop/winauto/uiauto-automation-element-propids for list of control types
+//		pRetVal->boolVal = VARIANT_FALSE;
+//		break;
+//#endif
+//	case UIA_ItemTypePropertyId:
+//		pRetVal->vt = VT_BSTR;
+//		// todo: friendly name of what's in a listview
+//		pRetVal->bstrVal = SysAllocString(L"");
+//		break;
+//#if WINVER >= 0x0602 // Windows 8
+//	case UIA_LiveSettingPropertyId:
+//		pRetVal->vt = VT_I4;
+//		// todo: "politeness" setting
+//		pRetVal->lVal = 0;
+//		break;
+//#endif
+	case UIA_LocalizedControlTypePropertyId:
+		pRetVal->vt = VT_BSTR;
+		pRetVal->bstrVal = SysAllocString(*WidgetTypeToLocalizedString(Widget));
+		break;
+	case UIA_NamePropertyId:
+		pRetVal->vt = VT_BSTR;
+		// todo: slate widgets don't have names, screen reader may read this as accessible text
+		pRetVal->bstrVal = SysAllocString(*Widget->GetWidgetName());
+		break;
+	//case UIA_OrientationPropertyId:
+	//	pRetVal->vt = VT_I4;
+	//	// todo: sliders, scroll bars, layouts
+	//	pRetVal->lVal = OrientationType_None;
+	//	break;
+	case UIA_ProcessIdPropertyId:
+		pRetVal->vt = VT_I4;
+		pRetVal->lVal = ::GetCurrentProcessId();
+		break;
+	case UIA_SizePropertyId:
+		pRetVal->vt = VT_R8 | VT_ARRAY;
+		pRetVal->parray = SafeArrayCreateVector(VT_R8, 0, 2);
+		if (pRetVal->parray)
+		{
+			FBox2D Box = Widget->GetBounds();
+			LONG i = 0;
+			float Width = Box.Max.X - Box.Min.X;
+			float Height = Box.Max.Y - Box.Min.Y;
+			bValid &= (SafeArrayPutElement(pRetVal->parray, &i, &Width) != S_OK);
+			i = 1;
+			bValid &= (SafeArrayPutElement(pRetVal->parray, &i, &Height) != S_OK);
+		}
+		else
+		{
+			bValid = false;
+		}
+		break;
+	default:
+		pRetVal->vt = VT_EMPTY;
+		break;
+	}
+
+	if (!bValid)
+	{
+		pRetVal->vt = VT_EMPTY;
+		return E_FAIL;
+	}
+
+	return S_OK;
+}
+
+HRESULT STDCALL FWindowsUIAWidgetProvider::get_HostRawElementProvider(IRawElementProviderSimple** pRetVal)
+{
+	// Only return host provider for native windows
+	*pRetVal = nullptr;
+	return S_OK;
+}
+
+HRESULT STDCALL FWindowsUIAWidgetProvider::Navigate(NavigateDirection direction, IRawElementProviderFragment** pRetVal)
+{
+	SCOPE_CYCLE_COUNTER(STAT_AccessibilityWindowsNavigate);
+
+	if (!IsValid())
+	{
+		return UIA_E_ELEMENTNOTAVAILABLE;
+	}
+
+	TSharedPtr<IAccessibleWidget> Relative = nullptr;
+	switch (direction)
+	{
+	case NavigateDirection_Parent:
+		Relative = Widget->GetParent();
+		break;
+	case NavigateDirection_NextSibling:
+		Relative = Widget->GetNextSibling();
+		break;
+	case NavigateDirection_PreviousSibling:
+		Relative = Widget->GetPreviousSibling();
+		break;
+	case NavigateDirection_FirstChild:
+		if (Widget->GetNumberOfChildren() > 0)
+		{
+			Relative = Widget->GetChildAt(0);
+		}
+		break;
+	case NavigateDirection_LastChild:
+	{
+		const int32 NumChildren = Widget->GetNumberOfChildren();
+		if (NumChildren > 0)
+		{
+			Relative = Widget->GetChildAt(NumChildren - 1);
+		}
+		break;
+	}
+	}
+
+	if (Relative.IsValid())
+	{
+		*pRetVal = static_cast<IRawElementProviderFragment*>(&UIAManager->GetWidgetProvider(Relative.ToSharedRef()));
+	}
+	else
+	{
+		*pRetVal = nullptr;
+	}
+	return S_OK;
+}
+
+HRESULT STDCALL FWindowsUIAWidgetProvider::GetRuntimeId(SAFEARRAY** pRetVal)
+{
+	if (IsValid())
+	{
+		int rtId[] = { UiaAppendRuntimeId, UIAManager->GetRuntimeId(this) };
+		*pRetVal = SafeArrayCreateVector(VT_I4, 0, 2);
+		if (*pRetVal)
+		{
+			for (LONG i = 0; i < 2; ++i)
+			{
+				if (SafeArrayPutElement(*pRetVal, &i, &rtId[i]) != S_OK)
+				{
+					return E_FAIL;
+				}
+			}
+		};
+		return S_OK;
+	}
+	else
+	{
+		return UIA_E_ELEMENTNOTAVAILABLE;
+	}
+}
+
+HRESULT STDCALL FWindowsUIAWidgetProvider::get_BoundingRectangle(UiaRect* pRetVal)
+{
+	if (IsValid())
+	{
+		FBox2D Box = Widget->GetBounds();
+		pRetVal->left = Box.Min.X;
+		pRetVal->top = Box.Min.Y;
+		pRetVal->width = Box.Max.X - Box.Min.X;
+		pRetVal->height = Box.Max.Y - Box.Min.Y;
+		return S_OK;
+	}
+	else
+	{
+		return UIA_E_ELEMENTNOTAVAILABLE;
+	}
+}
+HRESULT STDCALL FWindowsUIAWidgetProvider::GetEmbeddedFragmentRoots(SAFEARRAY** pRetVal)
+{
+	// This would technically only be valid in our case for a window within a window
+	*pRetVal = nullptr;
+	return S_OK;
+}
+
+HRESULT STDCALL FWindowsUIAWidgetProvider::SetFocus()
+{
+	if (IsValid())
+	{
+		if (Widget->SupportsFocus())
+		{
+			Widget->SetFocus();
+			return S_OK;
+		}
+		else
+		{
+			return UIA_E_NOTSUPPORTED;
+		}
+	}
+	else
+	{
+		return UIA_E_ELEMENTNOTAVAILABLE;
+	}
+}
+
+HRESULT STDCALL FWindowsUIAWidgetProvider::get_FragmentRoot(IRawElementProviderFragmentRoot** pRetVal)
+{
+	if (IsValid())
+	{
+		TSharedPtr<IAccessibleWidget> Window = Widget->GetTopLevelWindow();
+		if (Window.IsValid())
+		{
+			*pRetVal = static_cast<IRawElementProviderFragmentRoot*>(&static_cast<FWindowsUIAWindowProvider&>(UIAManager->GetWidgetProvider(Window.ToSharedRef())));
+			return S_OK;
+		}
+	}
+
+	return UIA_E_ELEMENTNOTAVAILABLE;
+}
+
+// ~
+
+// FWindowsUIAWindowProvider methods
+
+FWindowsUIAWindowProvider::FWindowsUIAWindowProvider(FWindowsUIAManager& InManager, TSharedRef<IAccessibleWidget> InWidget)
+	: FWindowsUIAWidgetProvider(InManager, InWidget)
+{
+	ensure(InWidget->AsWindow() != nullptr);
+}
+
+FWindowsUIAWindowProvider::~FWindowsUIAWindowProvider()
+{
+}
+
+HRESULT STDCALL FWindowsUIAWindowProvider::QueryInterface(REFIID riid, void** ppInterface)
+{
+	if (riid == __uuidof(IRawElementProviderFragmentRoot))
+	{
+		*ppInterface = static_cast<IRawElementProviderFragmentRoot*>(this);
+		AddRef();
+		return S_OK;
+	}
+	else
+	{
+		return FWindowsUIAWidgetProvider::QueryInterface(riid, ppInterface);
+	}
+}
+
+ULONG STDCALL FWindowsUIAWindowProvider::AddRef()
+{
+	return FWindowsUIABaseProvider::IncrementRef();
+}
+
+ULONG STDCALL FWindowsUIAWindowProvider::Release()
+{
+	return FWindowsUIABaseProvider::DecrementRef();
+}
+
+HRESULT STDCALL FWindowsUIAWindowProvider::get_HostRawElementProvider(IRawElementProviderSimple** pRetVal)
+{
+	if (Widget->IsValid())
+	{
+		TSharedPtr<FGenericWindow> NativeWindow = Widget->AsWindow()->GetNativeWindow();
+		if (NativeWindow.IsValid())
+		{
+			HWND Hwnd = static_cast<HWND>(NativeWindow->GetOSWindowHandle());
+			if (Hwnd != nullptr)
+			{
+				return UiaHostProviderFromHwnd(Hwnd, pRetVal);
+			}
+		}
+		return UIA_E_INVALIDOPERATION;
+	}
+	
+	return UIA_E_ELEMENTNOTAVAILABLE;
+}
+
+HRESULT STDCALL FWindowsUIAWindowProvider::GetPatternProvider(PATTERNID patternId, IUnknown** pRetVal)
+{
+	if (IsValid())
+	{
+		switch (patternId)
+		{
+		case UIA_WindowPatternId:
+			*pRetVal = static_cast<IWindowProvider*>(new FWindowsUIAControlProvider(*UIAManager, Widget));
+			return S_OK;
+		default:
+			return FWindowsUIAWidgetProvider::GetPatternProvider(patternId, pRetVal);
+		}
+	}
+
+	return UIA_E_ELEMENTNOTAVAILABLE;
+}
+
+HRESULT STDCALL FWindowsUIAWindowProvider::ElementProviderFromPoint(double x, double y, IRawElementProviderFragment** pRetVal)
+{
+	if (IsValid())
+	{
+		TSharedPtr<IAccessibleWidget> Child = Widget->AsWindow()->GetChildAtPosition(x, y);
+		if (Child.IsValid())
+		{
+			*pRetVal = static_cast<IRawElementProviderFragment*>(&UIAManager->GetWidgetProvider(Child.ToSharedRef()));
+		}
+		else
+		{
+			*pRetVal = nullptr;
+		}
+		return S_OK;
+	}
+
+	return UIA_E_ELEMENTNOTAVAILABLE;
+}
+
+HRESULT STDCALL FWindowsUIAWindowProvider::GetFocus(IRawElementProviderFragment** pRetVal)
+{
+	*pRetVal = nullptr;
+	if (IsValid())
+	{
+		TSharedPtr<IAccessibleWidget> Focus = Widget->AsWindow()->GetFocusedWidget();
+		if (Focus.IsValid())
+		{
+			*pRetVal = static_cast<IRawElementProviderFragment*>(&UIAManager->GetWidgetProvider(Focus.ToSharedRef()));
+		}
+		return S_OK;
+	}
+
+	return UIA_E_ELEMENTNOTAVAILABLE;
+}
+
+// ~
+
+#undef LOCTEXT_NAMESPACE
+
+#endif
diff --git a/Engine/Source/Runtime/ApplicationCore/Private/Windows/WindowsApplication.cpp b/Engine/Source/Runtime/ApplicationCore/Private/Windows/WindowsApplication.cpp
--- a/Engine/Source/Runtime/ApplicationCore/Private/Windows/WindowsApplication.cpp
+++ b/Engine/Source/Runtime/ApplicationCore/Private/Windows/WindowsApplication.cpp
@@ -24,6 +24,12 @@
 #include "Developer/SourceCodeAccess/Public/ISourceCodeAccessModule.h"
 #endif
 
+#if WITH_ACCESSIBILITY
+#include "Windows/Accessibility/WindowsUIAManager.h"
+#include "Windows/Accessibility/WindowsUIAWidgetProvider.h"
+#include <UIAutomation.h>
+#endif
+
 // Allow Windows Platform types in the entire file.
 #include "Windows/AllowWindowsPlatformTypes.h"
 THIRD_PARTY_INCLUDES_START
@@ -92,6 +98,9 @@ FWindowsApplication::FWindowsApplication( const HINSTANCE HInstance, const HICON
 		bAllowedToDeferMessageProcessing,
 		TEXT( "Whether windows message processing is deferred until tick or if they are processed immediately" ) )
 	, bInModalSizeLoop( false )
+#if WITH_ACCESSIBILITY
+	, UIAManager(new FWindowsUIAManager(*this))
+#endif
 
 {
 	FMemory::Memzero(ModifierKeyState, EModifierKey::Count);
@@ -313,6 +322,14 @@ void FWindowsApplication::SetMessageHandler( const TSharedRef< FGenericApplicati
 
 }
 
+#if WITH_ACCESSIBILITY
+void FWindowsApplication::SetAccessibleMessageHandler(const TSharedRef<FGenericAccessibleMessageHandler>& InAccessibleMessageHandler)
+{
+	GenericApplication::SetAccessibleMessageHandler(InAccessibleMessageHandler);
+	UIAManager->OnAccessibleMessageHandlerChanged();
+}
+#endif
+
 bool FWindowsApplication::IsGamepadAttached() const
 {
 	if (bForceNoGamepads)
@@ -1354,6 +1371,13 @@ int32 FWindowsApplication::ProcessMessage( HWND hwnd, uint32 msg, WPARAM wParam,
 		case WM_DESTROY:
 			{
 				Windows.Remove( CurrentNativeEventWindow );
+#if WITH_ACCESSIBILITY
+				// Tell UIA that the window no longer exists so that it can release some resources
+				if (GetAccessibleMessageHandler()->ApplicationIsAccessible())
+				{
+					UiaReturnRawElementProvider(hwnd, 0, 0, nullptr);
+				}
+#endif
 				return 0;
 			}
 			break;
@@ -1517,6 +1541,19 @@ int32 FWindowsApplication::ProcessMessage( HWND hwnd, uint32 msg, WPARAM wParam,
 			}
 			break;
 
+#if WITH_ACCESSIBILITY
+		case WM_GETOBJECT:
+		{
+			if (GetAccessibleMessageHandler()->ApplicationIsAccessible())
+			{
+				FScopedWidgetProvider Provider(UIAManager->GetWindowProvider(CurrentNativeEventWindow));
+				LRESULT Result = UiaReturnRawElementProvider(hwnd, wParam, lParam, &Provider.Provider);
+				return Result;
+			}
+			break;
+		}
+#endif
+
 		default:
 			if (bMessageExternallyHandled)
 			{
diff --git a/Engine/Source/Runtime/ApplicationCore/Public/GenericPlatform/GenericAccessibleInterfaces.h b/Engine/Source/Runtime/ApplicationCore/Public/GenericPlatform/GenericAccessibleInterfaces.h
new file mode 100644
--- /dev/null
+++ b/Engine/Source/Runtime/ApplicationCore/Public/GenericPlatform/GenericAccessibleInterfaces.h
@@ -0,0 +1,506 @@
+// Copyright 1998-2019 Epic Games, Inc. All Rights Reserved.
+
+#pragma once
+
+#include "CoreMinimal.h"
+
+/** Whether a widget should be included in accessibility, and if so, how its text should be retrieved. */
+enum class EAccessibleBehavior : uint8
+{
+	/** Not accessible. */
+	NotAccessible,
+	/** Accessible, for the implementing library to decide what it means. Given all data about a particular widget, it should try to choose the most-relevant text automatically. */
+	Auto,
+	/** Accessible, and traverse all child widgets and concat their summary text together. */
+	Summary,
+	/** Accessible, and retrieve manually-assigned text from a TAttribute. */
+	Custom,
+	/** Accessible, and use the tooltip's accessible text. */
+	ToolTip
+};
+
+DECLARE_LOG_CATEGORY_EXTERN(LogAccessibility, Log, All);
+
+#if WITH_ACCESSIBILITY
+
+#include "Misc/Variant.h"
+#include "Stats/Stats.h"
+
+class FGenericWindow;
+class IAccessibleWidget;
+
+DECLARE_STATS_GROUP(TEXT("Accessibility"), STATGROUP_Accessibility, STATCAT_Advanced);
+
+/** What kind of widget to tell the operating system this is. This may be translated to a different type depending on the platform. */
+enum class EAccessibleWidgetType : uint8
+{
+	Unknown,
+	Button,
+	CheckBox,
+	ComboBox,
+	Hyperlink,
+	Image,
+	Layout,
+	ScrollBar,
+	Slider,
+	Text,
+	TextEdit,
+	Window
+};
+
+/** Events that can be raised from accessible widgets to report back to the platform */
+enum class EAccessibleEvent : uint8
+{
+	/** A widget has become focused or unfocused. */
+	FocusChange,
+	/** A widget has been clicked, checked, or otherwise activated */
+	Activate,
+	/** Notify the user that something has happened. The user is not guaranteed to get this message. */
+	Notification,
+	/** Right before a widget is removed from its parent widget. */
+	BeforeRemoveFromParent,
+	/** Right after a widget is added to its parent widget. */
+	AfterAddToParent,
+	/** The widget was removed from the UI tree or deleted. */
+	WidgetRemoved
+};
+
+/**
+ * An accessible window corresponds to a native OS window. Fake windows that are embedded
+ * within other widgets that simply look and feel like windows are not IAccessibleWindows.
+ */
+class IAccessibleWindow
+{
+public:
+	/** The different states a window can be in to describe its screen anchors */
+	enum class EWindowDisplayState
+	{
+		Normal,
+		Minimize,
+		Maximize
+	};
+
+	/**
+	 * Retrieve the native OS window backing this accessible window. This can be safely
+	 * casted if you know what OS you're in (ie FWindowsWindows on Windows platform).
+	 *
+	 * @return The native window causing this accessible window to exist
+	 */
+	virtual TSharedPtr<FGenericWindow> GetNativeWindow() const = 0;
+
+	/**
+	 * Finds the deepest accessible widget in the hierarchy at the specified coordinates. The window
+	 * may return a pointer to itself in the case where there are no accessible children at the position.
+	 * This could return nullptr in the case where the coordinates are outside the window bounds.
+	 *
+	 * @param X The X coordinate in absolute screen space
+	 * @param Y The Y coordinate in absolute screen space
+	 * @return The deepest widget in the UI heirarchy at X,Y
+	 */
+	virtual TSharedPtr<IAccessibleWidget> GetChildAtPosition(int32 X, int32 Y) = 0;
+	/**
+	 * Retrieves the currently-focused widget, if it is accessible.
+	 *
+	 * @return The widget that has focus, or nullptr if the focused widget is not accessible.
+	 */
+	virtual TSharedPtr<IAccessibleWidget> GetFocusedWidget() const = 0;
+	/**
+	 * Request that the window closes itself. This may not happen immediately.
+	 */
+	virtual void Close() = 0;
+
+	/**
+	 * Check if the window can be minimized or maximized.
+	 *
+	 * @param State Whether to check for minimize or maximize.
+	 * @return True if the display state can be switched to, otherwise false.
+	 */
+	virtual bool SupportsDisplayState(EWindowDisplayState State) const = 0;
+	/**
+	 * Gets the current state minimize/maximize state of the window.
+	 *
+	 * @return The display state corresponding to how the window is displayed.
+	 */
+	virtual EWindowDisplayState GetDisplayState() const = 0;
+	/**
+	 * Sets a window to be minimized, maximized, or restored to normal.
+	 *
+	 * @param State What to change the window's display to.
+	 */
+	virtual void SetDisplayState(EWindowDisplayState State) = 0;
+	/**
+	 * Whether or not the window is modal.
+	 *
+	 * @return true if the window is modal, otherwise false.
+	 */
+	virtual bool IsModal() const = 0;
+};
+
+/**
+ * A widget that can be triggered to fire an event, such as buttons or checkboxes.
+ */
+class IAccessibleActivatable
+{
+public:
+	/** Trigger the widget */
+	virtual void Activate() = 0;
+	/**
+	 * Check whether this widget can be toggled between various check states
+	 *
+	 * @return true if this widget supports being in different states
+	 */
+	virtual bool IsCheckable() const { return false; }
+	/**
+	 * If IsCheckable() is true, this gets the current state that the widget is in.
+	 * //todo: return ECheckState
+	 *
+	 * @return true if the current state is considered "on" or "checked"
+	 */
+	virtual bool GetCheckedState() const { return false; }
+};
+
+/**
+ * An accessible widget that stores an arbitrary value of any type capable of being serialized into a string.
+ * Optional overrides add support for slider-like functionality.
+ */
+class IAccessibleProperty
+{
+public:
+	/**
+	 * Whether the widget is in read-only mode, which could be different than IsEnabled().
+	 *
+	 * @return true if the widget is in read-only mode.
+	 */
+	virtual bool IsReadOnly() const { return true; }
+	/**
+	 * Check if this text is storing password data, indicating that it may need special handling to presenting itself to the user.
+	 *
+	 * @return true if the text is storing a password or otherwise senstive data that should be hidden.
+	 */
+	virtual bool IsPassword() const { return false; }
+	/**
+	 * How much the value should increase/decrease when the user attempts to modify the value using UI controls.
+	 * Note: This should always return a positive value. The caller is responsible for negating it when attempting to decrease.
+	 *
+	 * @return A number suggesting how much to modify GetValue() by when the user wants to increase/decrease the value.
+	 */
+	virtual float GetStepSize() const { return 0.0f; }
+	/**
+	 * The maximum allowed value for this property. This should only be used if GetStepSize is not 0.
+	 *
+	 * @return The maximum value that this property can be assigned when using step sizes.
+	 */
+	virtual float GetMaximum() const { return 0.0f; }
+	/**
+	 * The minimum allowed value for this property. This should only be used if GetStepSize is not 0.
+	 *
+	 * @return The minimum value that this property can be assigned when using step sizes.
+	 */
+	virtual float GetMinimum() const { return 0.0f; }
+	/**
+	 * The current value stored by the widget. Even if the underlying value is not a String, it should be serialized to one
+	 * in order to match the return type.
+	 *
+	 * @return A string representing the serialized value stored in the widget.
+	 */
+	virtual FString GetValue() const = 0;
+	/*
+	 * Set the value stored by the widget. While this function accepts a String, there is no way to know
+	 * what the underlying data is stored as. The platform layer must retain some additional information
+	 * about what kind of widget this is, and ensure it's being called with valid arguments.
+	 *
+	 * @param Value The new value to assign to the widget, which may need to be converted before assigning to a variable.
+	 */
+	virtual void SetValue(const FString& Value) {}
+};
+
+/**
+ * A widget that contains text, with the potential ability to select sections, read specific words/paragraphs, etc.
+ * Note: This class is currently incomplete.
+ */
+class IAccessibleText
+{
+public:
+	/**
+	 * Get the full text contained in this widget, even if some if it is clipped.
+	 *
+	 * @return All the text held by this widget.
+	 */
+	virtual const FString& GetText() const = 0;
+};
+
+/**
+ * Provides the core set of accessible data that is necessary in order for widget traversal and TTS to be implemented.
+ * In order to support functionality beyond this, subclasses must implement the other accessible interfaces and
+ * then override the As*() functions.
+ */
+class IAccessibleWidget : public TSharedFromThis<IAccessibleWidget>
+{
+public:
+	IAccessibleWidget() {}
+	virtual ~IAccessibleWidget() {}
+
+	/**
+	 * Whether or not the underlying widget backing this interface still exists
+	 *
+	 * @return true if functions can be called on this interface and should return valid results
+	 */
+	virtual bool IsValid() const = 0;
+	/**
+	 * Returns the window at the top of this widget's hierarchy. This function may return itself for accessible windows,
+	 * and could return nullptr in cases where the widget is not currently part of a hierarchy.
+	 *
+	 * @return The root window in this widget's widget tree, or nullptr if there is no window.
+	 */
+	virtual TSharedPtr<IAccessibleWidget> GetTopLevelWindow() const = 0;
+	/**
+	 * Retrieving the bounding rect in absolute coordinates for this widget. On some platforms this may be used for hit testing.
+	 *
+	 * @return The bounds of the widget.
+	 */
+	virtual FBox2D GetBounds() const = 0;
+	/**
+	 * Get the accessible parent of this widget. This may be nullptr if this widget is a window, or if the widget is
+	 * currently disconnected from the UI tree.
+	 *
+	 * @return The accessible parent widget of this widget.
+	 */
+	virtual TSharedPtr<IAccessibleWidget> GetParent() = 0;
+	/**
+	 * Retrieves the widget after this one in the parent's list of children. This should return nullptr for the last widget.
+	 *
+	 * @return The next widget on the same level of the UI hierarchy.
+	 */
+	virtual TSharedPtr<IAccessibleWidget> GetNextSibling() = 0;
+	/**
+	 * Retrieves the widget before this one in the parent's list of children. This should return nullptr for the first widget.
+	 *
+	 * @return The previous widget on the same level of the UI hierarchy.
+	 */
+	virtual TSharedPtr<IAccessibleWidget> GetPreviousSibling() = 0;
+	/**
+	 * Retrieves the accessible child widget at a certain index. This should return nullptr if Index < 0 or Index >= GetNumberOfChildren().
+	 *
+	 * @param The index of the child widget to get
+	 * @return The accessible child widget at the specified index.
+	 */
+	virtual TSharedPtr<IAccessibleWidget> GetChildAt(int32 Index) = 0;
+	/**
+	 * How many accessible children this widget has.
+	 *
+	 * @return The number of accessible children that exist for this widget.
+	 */
+	virtual int32 GetNumberOfChildren() = 0;
+	/**
+	 * What type of accessible widget the underlying widget should be treated as. A widget may be capable of presenting itself
+	 * as multiple different types of widgets, but only one can be reported back to the platform.
+	 *
+	 * @return Which type of widget the platform layer should treat this as.
+	 */
+	virtual EAccessibleWidgetType GetWidgetType() const = 0;
+
+	/**
+	 * The name of the underlying class that this accessible widget represents.
+	 *
+	 * @return The class name of the underlying widget.
+	 */
+	virtual FString GetClassName() const = 0;
+	/**
+	 * The name of the widget to report to the platform layer. For screen readers, this is often the text that will be spoken.
+	 *
+	 * @return Ideally, a human-readable name that represents what the widget does.
+	 */
+	virtual FString GetWidgetName() const = 0;
+	/**
+	 * Additional information a user may need in order to effectively interact or use the widget, such as a tooltip.
+	 *
+	 * @return A more-detailed description of what the widget is or how its used.
+	 */
+	virtual FString GetHelpText() const = 0;
+	
+	/**
+	 * Whether the widget is enabled and can be interacted with.
+	 *
+	 * @return true if the widget is enabled.
+	 */
+	virtual bool IsEnabled() const = 0;
+	/**
+	 * Whether the widget is being rendered on screen or not.
+	 *
+	 * @return true if the widget is hidden off screen, collapsed, or something similar.
+	 */
+	virtual bool IsHidden() const = 0;
+	/**
+	 * Whether the widget supports keyboard focus or not.
+	 *
+	 * @return true if the widget can receive keyboard focus.
+	 */
+	virtual bool SupportsFocus() const = 0;
+	/**
+	 * Whether the widget has keyboard focus or not.
+	 *
+	 * @return true if the widget currently has keyboard focus.
+	 */
+	virtual bool HasFocus() const = 0;
+	/** Assign keyboard focus to this widget, if it supports it. If not, focus should not be affected. */
+	virtual void SetFocus() = 0;
+
+	/**
+	 * Attempt to cast this to an IAccessibleWindow
+	 *
+	 * @return 'this' as an IAccessibleWindow if possible, otherwise nullptr
+	 */
+	virtual IAccessibleWindow* AsWindow() { return nullptr; }
+	/**
+	 * Attempt to cast this to an IAccessibleActivatable
+	 *
+	 * @return 'this' as an IAccessibleActivatable if possible, otherwise nullptr
+	 */
+	virtual IAccessibleActivatable* AsActivatable() { return nullptr; }
+	/**
+	 * Attempt to cast this to an IAccessibleProperty
+	 *
+	 * @return 'this' as an IAccessibleProperty if possible, otherwise nullptr
+	 */
+	virtual IAccessibleProperty* AsProperty() { return nullptr; }
+	/**
+	 * Attempt to cast this to an IAccessibleText
+	 *
+	 * @return 'this' as an IAccessibleText if possible, otherwise nullptr
+	 */
+	virtual IAccessibleText* AsText() { return nullptr; }
+};
+
+/**
+ * Platform and application-agnostic messaging system for accessible events. The message handler
+ * lives in GenericApplication and any subclass that wishes to support accessibility should subclass
+ * this and use GenericAppliation::SetAccessibleMessageHandler to enable functionality.
+ *
+ * GetAccessibleWindow() is tne entry point to all accessible widgets. Once the window is retrieved, it
+ * can be queried for children in various ways. RaiseEvent() allows messages to bubble back up to the
+ * native OS through anything bound to the AccessibleEventDelegate.
+ *
+ * Callers can use ApplicationIsAccessible() to see if accessibility is supported or not. Alternatively,
+ * calling GetAccessibleWindow and seeing if the result is valid should provide the same information.
+ */
+class FGenericAccessibleMessageHandler
+{
+public:
+	/** A widget raised an event to pass to the native OS implementation. */
+	DECLARE_DELEGATE_FourParams(
+		FAccessibleEvent,
+		/** The accessible widget that generated the event */
+		TSharedRef<IAccessibleWidget>,
+		/** The type of event generated */
+		EAccessibleEvent,
+		/** If this was a property changed event, the 'before' value */
+		FVariant,
+		/** If this was a property changed event, the 'after' value. This may also be set for other events such as Notification. */
+		FVariant);
+
+	FGenericAccessibleMessageHandler() : bIsActive(false) {}
+
+	virtual ~FGenericAccessibleMessageHandler()
+	{
+		if (AccessibleEventDelegate.IsBound())
+		{
+			AccessibleEventDelegate.Unbind();
+		}
+	}
+
+	/**
+	 * Subclasses should return true to indicate that they support accessibility.
+	 *
+	 * @return true if the application intends to return valid accessible widgets when queried.
+	 */
+	virtual bool ApplicationIsAccessible() const { return false; }
+
+	/**
+	 * Checks if accessibility is enabled in the application. Usually this happens when screen-reading software is turned on.
+	 * Note: On some platforms, there is no way to deactivate this after enabling it. 
+	 *
+	 * @return The last value SetActive() was called with.
+	 */
+	bool IsActive() const { return bIsActive; }
+
+	/**
+	 * Notify the application to start or stop processing accessible messages from the platform layer.
+	 *
+	 * @param bActive Whether to enable to disable the message handler.
+	 */
+	void SetActive(bool bActive)
+	{
+		if (bActive != bIsActive)
+		{
+			bIsActive = bActive;
+
+			if (bIsActive)
+			{
+				UE_LOG(LogAccessibility, Log, TEXT("Enabling Accessibility"));
+				OnActivate();
+			}
+			else
+			{
+				OnDeactivate();
+				UE_LOG(LogAccessibility, Log, TEXT("Accessibility Disabled"));
+			}
+		}
+	}
+
+	/**
+	 * Creates or retrieves an accessible object for a native OS window.
+	 * todo: Behavior for non-native windows (virtual or others) is currently undefined.
+	 *
+	 * @param InWindow The native window to find the accessible window for
+	 * @return The accessible object corresponding to the supplied native window
+	 */
+	virtual TSharedPtr<IAccessibleWidget> GetAccessibleWindow(const TSharedRef<FGenericWindow>& InWindow) const { return nullptr; }
+
+	/**
+	 * Push an event from an accessible widget back to the platform layer.
+	 *
+	 * @param Widget The widget raising the event
+	 * @param Event The type of event being raised
+	 * @param OldValue If this is a property changed event, the old value of the property
+	 * @param NewValue If this is a property changed event, the current value of the property. This may also be used by non-property events that require data.
+	 */
+	void RaiseEvent(TSharedRef<IAccessibleWidget> Widget, EAccessibleEvent Event, FVariant OldValue, FVariant NewValue)
+	{
+		AccessibleEventDelegate.ExecuteIfBound(Widget, Event, OldValue, NewValue);
+	}
+	/**
+	 * Push an event from an accessible widget back to the platform layer.
+	 * This is a convenience function if OldValue and NewValue are irrelevant for this event type.
+	 *
+	 * @param Widget The widget raising the event
+	 * @param Event The type of event being raised
+	 */
+	void RaiseEvent(TSharedRef<IAccessibleWidget> Widget, EAccessibleEvent Event)
+	{
+		RaiseEvent(Widget, Event, FVariant(), FVariant());
+	}
+
+	/**
+	 * Assign a function to be called whenever an accessible event is raised.
+	 *
+	 * @param Delegate The delegate to execute when an event is raised.
+	 */
+	void SetAccessibleEventDelegate(const FAccessibleEvent& Delegate)
+	{
+		AccessibleEventDelegate = Delegate;
+	}
+
+protected:
+	/** Triggered when bIsActive changes from false to true. */
+	virtual void OnActivate() {}
+	/** Triggered when bIsActive changes from true to false. */
+	virtual void OnDeactivate() {}
+
+private:
+	/** Whether or not accessibility is currently enabled in the application */
+	bool bIsActive;
+	/** Delegate for the platform layer to listen to widget events */
+	FAccessibleEvent AccessibleEventDelegate;
+};
+
+#endif
diff --git a/Engine/Source/Runtime/ApplicationCore/Public/GenericPlatform/GenericApplication.h b/Engine/Source/Runtime/ApplicationCore/Public/GenericPlatform/GenericApplication.h
--- a/Engine/Source/Runtime/ApplicationCore/Public/GenericPlatform/GenericApplication.h
+++ b/Engine/Source/Runtime/ApplicationCore/Public/GenericPlatform/GenericApplication.h
@@ -9,6 +9,9 @@
 #include "Math/Vector4.h"
 #include "Templates/SharedPointer.h"
 #include "Delegates/Delegate.h"
+#if WITH_ACCESSIBILITY
+#include "GenericPlatform/GenericAccessibleInterfaces.h"
+#endif
 #include "GenericPlatform/GenericApplicationMessageHandler.h"
 #include "GenericPlatform/GenericWindowDefinition.h"
 #include "GenericPlatform/GenericWindow.h"
@@ -428,6 +431,9 @@ public:
 	GenericApplication( const TSharedPtr< ICursor >& InCursor )
 		: Cursor( InCursor )
 		, MessageHandler( MakeShareable( new FGenericApplicationMessageHandler() ) )
+#if WITH_ACCESSIBILITY
+		, AccessibleMessageHandler(MakeShareable(new FGenericAccessibleMessageHandler()))
+#endif
 	{
 
 	}
@@ -438,6 +444,11 @@ public:
 
 	TSharedRef< FGenericApplicationMessageHandler > GetMessageHandler() { return MessageHandler; }
 
+#if WITH_ACCESSIBILITY
+	virtual void SetAccessibleMessageHandler(const TSharedRef<FGenericAccessibleMessageHandler>& InAccessibleMessageHandler) { AccessibleMessageHandler = InAccessibleMessageHandler; }
+	TSharedRef<FGenericAccessibleMessageHandler> GetAccessibleMessageHandler() const { return AccessibleMessageHandler; }
+#endif
+
 	virtual void PollGameDeviceState( const float TimeDelta ) { }
 
 	virtual void PumpMessages( const float TimeDelta ) { }
@@ -543,6 +554,10 @@ public:
 protected:
 
 	TSharedRef< class FGenericApplicationMessageHandler > MessageHandler;
+
+#if WITH_ACCESSIBILITY
+	TSharedRef<FGenericAccessibleMessageHandler> AccessibleMessageHandler;
+#endif
 	
 	/** Trigger the OnDisplayMetricsChanged event with the argument 'InMetrics' */
 	void BroadcastDisplayMetricsChanged( const FDisplayMetrics& InMetrics ){ OnDisplayMetricsChangedEvent.Broadcast( InMetrics ); }
diff --git a/Engine/Source/Runtime/ApplicationCore/Public/Windows/Accessibility/WindowsUIABaseProvider.h b/Engine/Source/Runtime/ApplicationCore/Public/Windows/Accessibility/WindowsUIABaseProvider.h
new file mode 100644
--- /dev/null
+++ b/Engine/Source/Runtime/ApplicationCore/Public/Windows/Accessibility/WindowsUIABaseProvider.h
@@ -0,0 +1,55 @@
+// Copyright 1998-2019 Epic Games, Inc. All Rights Reserved.
+
+#pragma once
+
+#if WITH_ACCESSIBILITY
+
+#include "Windows/AllowWindowsPlatformTypes.h"
+#include <Ole2.h>
+#include "Windows/HideWindowsPlatformTypes.h"
+#include <UIAutomation.h>
+
+#include "Templates/SharedPointer.h"
+
+class FWindowsUIAManager;
+class IAccessibleWidget;
+
+/**
+ * Base class for all Windows UIA Providers to inherit from. Provides implementation of IUnknown's
+ * reference counting functions as well as integration with the UIA Manager class.
+ */
+class FWindowsUIABaseProvider
+{
+public:
+	/** Notify the Provider that the backing application no longer exists. */
+	void OnUIAManagerDestroyed();
+
+protected:
+	FWindowsUIABaseProvider(FWindowsUIAManager& InManager, TSharedRef<IAccessibleWidget> InWidget);
+	virtual ~FWindowsUIABaseProvider();
+
+	/**
+	 * Check whether this Provider can have operations called on it. This might return false if the
+	 * application is being destroyed, or if the underlying accessible widget is no longer valid.
+	 */
+	bool IsValid() const;
+
+	/** Add one to RefCount, and return the new RefCount */
+	uint32 IncrementRef();
+	/** Subtract one from RefCount, and return the new RefCount. If RefCount is 0, deletes the Provider and returns 0. */
+	uint32 DecrementRef();
+
+	/** A pointer to the UIA Manager which is guaranteed to be valid so long as the application is still running. */
+	FWindowsUIAManager* UIAManager;
+	/**
+	 * All Providers must have a valid accessible widget, even if the accessible widget itself does not
+	 * have any valid data backing it. In this case, the accessible widget will return some default values.
+	 */
+	TSharedRef<IAccessibleWidget> Widget;
+
+private:
+	/** Counter for the number of strong references to this object. Many of these will be from external applications. */
+	uint32 RefCount;
+};
+
+#endif
diff --git a/Engine/Source/Runtime/ApplicationCore/Public/Windows/Accessibility/WindowsUIAControlProvider.h b/Engine/Source/Runtime/ApplicationCore/Public/Windows/Accessibility/WindowsUIAControlProvider.h
new file mode 100644
--- /dev/null
+++ b/Engine/Source/Runtime/ApplicationCore/Public/Windows/Accessibility/WindowsUIAControlProvider.h
@@ -0,0 +1,152 @@
+// Copyright 1998-2019 Epic Games, Inc. All Rights Reserved.
+
+#pragma once
+
+#if WITH_ACCESSIBILITY
+
+#include "Windows/Accessibility/WindowsUIABaseProvider.h"
+#include "Templates/SharedPointer.h"
+
+class FWindowsUIAControlProvider;
+class IAccessibleWidget;
+
+/**
+ * A TextRange Provider represents a start and end position within an ITextProvider. This can reference
+ * a word, a paragraph, a selection, etc. Multiple text ranges can exist for a single piece of text.
+ */
+class FWindowsUIATextRangeProvider
+	: public FWindowsUIABaseProvider
+	, public ITextRangeProvider
+{
+public:
+	FWindowsUIATextRangeProvider(FWindowsUIAManager& InManager, TSharedRef<IAccessibleWidget> InWidget, FTextRange InRange);
+
+	// IUnknown
+	HRESULT STDCALL QueryInterface(REFIID riid, void** ppInterface) override;
+	ULONG STDCALL AddRef() override;
+	ULONG STDCALL Release() override;
+	// ~
+
+	// ITextRangeProvider
+	virtual HRESULT STDCALL Clone(ITextRangeProvider** pRetVal) override;
+	virtual HRESULT STDCALL Compare(ITextRangeProvider* range, BOOL* pRetVal) override;
+	virtual HRESULT STDCALL CompareEndpoints(TextPatternRangeEndpoint endpoint, ITextRangeProvider* targetRange, TextPatternRangeEndpoint targetEndpoint, int* pRetVal) override;
+	virtual HRESULT STDCALL ExpandToEnclosingUnit(TextUnit unit) override;
+	virtual HRESULT STDCALL FindAttribute(TEXTATTRIBUTEID attributeId, VARIANT val, BOOL backward, ITextRangeProvider** pRetVal) override;
+	virtual HRESULT STDCALL FindText(BSTR text, BOOL backward, BOOL ignoreCase, ITextRangeProvider** pRetVal) override;
+	virtual HRESULT STDCALL GetAttributeValue(TEXTATTRIBUTEID attributeId, VARIANT* pRetVal) override;
+	virtual HRESULT STDCALL GetBoundingRectangles(SAFEARRAY** pRetVal) override;
+	virtual HRESULT STDCALL GetEnclosingElement(IRawElementProviderSimple** pRetVal) override;
+	virtual HRESULT STDCALL GetText(int maxLength, BSTR* pRetVal) override;
+	virtual HRESULT STDCALL Move(TextUnit unit, int count, int* pRetVal) override;
+	virtual HRESULT STDCALL MoveEndpointByUnit(TextPatternRangeEndpoint endpoint, TextUnit unit, int count, int* pRetVal) override;
+	virtual HRESULT STDCALL MoveEndpointByRange(TextPatternRangeEndpoint endpoint, ITextRangeProvider* targetRange, TextPatternRangeEndpoint targetEndpoint) override;
+	virtual HRESULT STDCALL Select() override;
+	virtual HRESULT STDCALL AddToSelection() override;
+	virtual HRESULT STDCALL RemoveFromSelection() override;
+	virtual HRESULT STDCALL ScrollIntoView(BOOL alignToTop) override;
+	virtual HRESULT STDCALL GetChildren(SAFEARRAY** pRetVal) override;
+	// ~
+
+protected:
+	/**
+	 * Gets the substring that this text range Provider represents.
+	 *
+	 * @return A substring of the underlying widget's text from TextRange's begin/end indices.
+	 */
+	FString TextFromTextRange();
+	static FString TextFromTextRange(const FString& InString, const FTextRange& InRange);
+
+	/** The range that this Provider represents in the underlying text. */
+	FTextRange TextRange;
+
+private:
+	virtual ~FWindowsUIATextRangeProvider();
+};
+
+/**
+ * The control Provider handles all control pattern-related functionality for the widget Provider. Control Providers should only
+ * be generated through widget Providers in their GetPatternProvider() function. In doing this, a control Pattern is only guaranteed
+ * to support functions that match the pattern Provider that it was generated for. This means that if GetPatternProvider() is called
+ * for UIA_ValuePatternId, the control Provider is only guaranteed to work with functions that implement IValueProvider. It is up
+ * to the caller to ensure that they are making valid function calls.
+ */
+class FWindowsUIAControlProvider
+	: public FWindowsUIABaseProvider
+	, public IInvokeProvider
+	, public IRangeValueProvider
+	, public ITextProvider
+	, public IToggleProvider
+	, public ITransformProvider
+	, public IValueProvider
+	, public IWindowProvider
+{
+public:
+	FWindowsUIAControlProvider(FWindowsUIAManager& Manager, TSharedRef<IAccessibleWidget> InWidget);
+
+	// IUnknown
+	HRESULT STDCALL QueryInterface(REFIID riid, void** ppInterface) override;
+	ULONG STDCALL AddRef() override;
+	ULONG STDCALL Release() override;
+	// ~
+
+	// IInvokeProvider
+	virtual HRESULT STDCALL Invoke() override;
+	// ~
+
+	// IRangeValueProvider
+	virtual HRESULT STDCALL SetValue(double val) override;
+	virtual HRESULT STDCALL get_Value(double* pRetVal) override;
+	virtual HRESULT STDCALL get_IsReadOnly(BOOL* pRetVal) override;
+	virtual HRESULT STDCALL get_Maximum(double* pRetVal) override;
+	virtual HRESULT STDCALL get_Minimum(double* pRetVal) override;
+	virtual HRESULT STDCALL get_LargeChange(double* pRetVal) override;
+	virtual HRESULT STDCALL get_SmallChange(double* pRetVal) override;
+	// ~
+
+	// ITextProvider
+	virtual HRESULT STDCALL get_DocumentRange(ITextRangeProvider** pRetVal) override;
+	virtual HRESULT STDCALL get_SupportedTextSelection(SupportedTextSelection* pRetVal) override;
+	virtual HRESULT STDCALL GetSelection(SAFEARRAY** pRetVal) override;
+	virtual HRESULT STDCALL GetVisibleRanges(SAFEARRAY** pRetVal) override;
+	virtual HRESULT STDCALL RangeFromChild(IRawElementProviderSimple* childElement, ITextRangeProvider** pRetVal) override;
+	virtual HRESULT STDCALL RangeFromPoint(UiaPoint point, ITextRangeProvider** pRetVal) override;
+	// ~
+
+	// IToggleState
+	virtual HRESULT STDCALL get_ToggleState(ToggleState* pRetVal) override;
+	virtual HRESULT STDCALL Toggle() override;
+	// ~
+
+	// ITransformProvider
+	virtual HRESULT STDCALL get_CanMove(BOOL *pRetVal) override;
+	virtual HRESULT STDCALL get_CanResize(BOOL *pRetVal) override;
+	virtual HRESULT STDCALL get_CanRotate(BOOL *pRetVal) override;
+	virtual HRESULT STDCALL Move(double x, double y) override;
+	virtual HRESULT STDCALL Resize(double width, double height) override;
+	virtual HRESULT STDCALL Rotate(double degrees) override;
+	// ~
+
+	// IValueProvider
+	virtual HRESULT STDCALL SetValue(LPCWSTR val) override;
+	virtual HRESULT STDCALL get_Value(BSTR* pRetVal) override;
+	//virtual HRESULT STDCALL get_IsReadOnly(BOOL* pRetVal) override; // Duplicate of IRangeValueProvider
+	// ~
+
+	// IWindowProvider
+	virtual HRESULT STDCALL Close() override;
+	virtual HRESULT STDCALL get_CanMaximize(BOOL* pRetVal) override;
+	virtual HRESULT STDCALL get_CanMinimize(BOOL* pRetVal) override;
+	virtual HRESULT STDCALL get_IsModal(BOOL* pRetVal) override;
+	virtual HRESULT STDCALL get_IsTopmost(BOOL* pRetVal) override;
+	virtual HRESULT STDCALL get_WindowInteractionState(WindowInteractionState* pRetVal) override;
+	virtual HRESULT STDCALL get_WindowVisualState(WindowVisualState* pRetVal) override;
+	virtual HRESULT STDCALL SetVisualState(WindowVisualState state) override;
+	virtual HRESULT STDCALL WaitForInputIdle(int milliseconds, BOOL* pRetVal) override;
+	// ~
+
+private:
+	virtual ~FWindowsUIAControlProvider();
+};
+
+#endif
diff --git a/Engine/Source/Runtime/ApplicationCore/Public/Windows/Accessibility/WindowsUIAManager.h b/Engine/Source/Runtime/ApplicationCore/Public/Windows/Accessibility/WindowsUIAManager.h
new file mode 100644
--- /dev/null
+++ b/Engine/Source/Runtime/ApplicationCore/Public/Windows/Accessibility/WindowsUIAManager.h
@@ -0,0 +1,100 @@
+// Copyright 1998-2019 Epic Games, Inc. All Rights Reserved.
+
+#pragma once
+
+#if WITH_ACCESSIBILITY
+
+#include "Templates/SharedPointer.h"
+#include "GenericPlatform/GenericAccessibleInterfaces.h"
+
+class FWindowsApplication;
+class FWindowsWindow;
+class FWindowsUIAControlProvider;
+class FWindowsUIAWidgetProvider;
+class FWindowsUIAWindowProvider;
+class IAccessibleWidget;
+class FVariant;
+
+/**
+ * Manager for Windows implementation of UE4's accessibility API, utilizing Microsoft's UI Automation API. It provides a
+ * central location for message passing to/from the WindowsApplication, AccessibleMessageHandler, and Windows UIA Providers.
+ *
+ * This class only accepts and returns references, so callers should ensure their objects are valid before using it.
+ */
+class FWindowsUIAManager
+{
+	/** Allow FWindowsUIABaseProvider to register and unregister itself from ProviderList automatically */
+	friend class FWindowsUIABaseProvider;
+public:
+	FWindowsUIAManager(const FWindowsApplication& InApplication);
+	~FWindowsUIAManager();
+
+	/**
+	 * Create a Windows UIA IRawElementProviderSimple from a given accessible widget. Providers are stored in a local cache,
+	 * and this function will allocate a new Provider if one does not already exist for the given widget. Using this function
+	 * will always increase the RefCount of the Provider by one.
+	 *
+	 * If InWidget->IsWindow() returns true, a FWindowsUIAWindowProvider will be created instead of a FWindowsUIAWidgetProvider.
+	 *
+	 * @param InWidget A non-null reference to an accessible widget
+	 * @return A reference to the cached Provider
+	 */
+	FWindowsUIAWidgetProvider& GetWidgetProvider(TSharedRef<IAccessibleWidget> InWidget);
+
+	/**
+	 * Create a Windows UIA IRawElementProviderSimple from a given native window handle. Providers are stored in a local cache,
+	 * and this function will allocate a new Provider if one does not already exist for the given widget. Using this function
+	 * will always increase the RefCount of the Provider by one.
+	 *
+	 * If the Provider cache is empty upon calling this function, it will activate accessibility in the entire application.
+	 *
+	 * @param InWindow A non-null reference to a native window
+	 * @return A reference to the cached Provider
+	 */
+	FWindowsUIAWindowProvider& GetWindowProvider(TSharedRef<FWindowsWindow> InWindow);
+
+	/**
+	 * Generate an application-unique ID for this Provider. A widget must be able to be uniquely
+	 * identified in the case where it gets deleted, but only within the context of a specific hwnd.
+	 * It is OK for widgets to overlap IDs in separate running instances of the same application.
+	 * This returns an int because it's passed directly to Windows as one.
+	 *
+	 * @return An application-unique identifier for this Provider
+	 */
+	int GetRuntimeId(FWindowsUIAWidgetProvider* InProvider);
+
+	/**
+	 * Notify the Manager that the RefCount for a Provider has reached 0, which will cause it to be removed from the cache.
+	 *
+	 * If the Provider cache is empty after removing the Provider, accessibility will be disabled through the entire application.
+	 *
+	 * @param InWidget The accessible widget backing the Provider that was removed
+	 */
+	void OnWidgetProviderRemoved(TSharedRef<IAccessibleWidget> InWidget);
+
+	/**
+	 * Notify the Manager that the accessible message handler for the application has changed, in order to relink to the accessible event delegate.
+	 */
+	void OnAccessibleMessageHandlerChanged();
+
+
+	static TMap<EAccessibleWidgetType, ULONG> WidgetTypeToWindowsTypeMap;
+	static TMap<EAccessibleWidgetType, FText> WidgetTypeToTextMap;
+
+private:
+	/** Callback function for processing events raised from the AccessibleMessageHandler */
+	void OnEventRaised(TSharedRef<IAccessibleWidget> Widget, EAccessibleEvent Event, FVariant OldValue, FVariant NewValue);
+
+	/** Cache of all Providers with a RefCount of at least 1 that map to an accessible widget. */
+	TMap<TSharedRef<IAccessibleWidget>, FWindowsUIAWidgetProvider*> CachedWidgetProviders;
+	TMap<FWindowsUIAWidgetProvider*, int> RuntimeIds;
+	/**
+	 * A set of all Providers with a RefCount of at least 1. This also includes non-widget Providers such as text ranges.
+	 * When the application is closed, this is used to notify Providers held by external applications that they are invalid.
+	 */
+	TSet<FWindowsUIABaseProvider*> ProviderList;
+	/** A reference back to the owning application that can be used to access the AccessibleMessageHandler */
+	const FWindowsApplication& WindowsApplication;
+};
+
+#endif
diff --git a/Engine/Source/Runtime/ApplicationCore/Public/Windows/Accessibility/WindowsUIAPropertyGetters.h b/Engine/Source/Runtime/ApplicationCore/Public/Windows/Accessibility/WindowsUIAPropertyGetters.h
new file mode 100644
--- /dev/null
+++ b/Engine/Source/Runtime/ApplicationCore/Public/Windows/Accessibility/WindowsUIAPropertyGetters.h
@@ -0,0 +1,29 @@
+// Copyright 1998-2019 Epic Games, Inc. All Rights Reserved.
+
+#pragma once
+
+#if WITH_ACCESSIBILITY
+
+#include "Windows/AllowWindowsPlatformTypes.h"
+#include <Ole2.h>
+#include "Windows/HideWindowsPlatformTypes.h"
+#include <UIAutomation.h>
+
+#include "Templates/SharedPointer.h"
+
+class FVariant;
+class IAccessibleWidget;
+
+namespace WindowsUIAPropertyGetters
+{
+	VARIANT GetPropertyValueWindows(TSharedRef<IAccessibleWidget> AccessibleWidget, PROPERTYID WindowsPropertyId);
+	FVariant GetPropertyValue(TSharedRef<IAccessibleWidget> AccessibleWidget, PROPERTYID WindowsPropertyId);
+
+	/**
+	 * Convert an FVariant to a Windows VARIANT.
+	 * Only necessary conversions are implemented.
+	 */
+	VARIANT FVariantToWindowsVariant(const FVariant& Value);
+}
+
+#endif
diff --git a/Engine/Source/Runtime/ApplicationCore/Public/Windows/Accessibility/WindowsUIAWidgetProvider.h b/Engine/Source/Runtime/ApplicationCore/Public/Windows/Accessibility/WindowsUIAWidgetProvider.h
new file mode 100644
--- /dev/null
+++ b/Engine/Source/Runtime/ApplicationCore/Public/Windows/Accessibility/WindowsUIAWidgetProvider.h
@@ -0,0 +1,119 @@
+// Copyright 1998-2019 Epic Games, Inc. All Rights Reserved.
+
+#pragma once
+
+#if WITH_ACCESSIBILITY
+
+#include "Windows/Accessibility/WindowsUIABaseProvider.h"
+#include "Misc/Variant.h"
+
+class FWindowsUIAManager;
+class IAccessibleWidget;
+
+/**
+ * Windows UIA Provider implementation for all non-Window widgets.
+ *
+ * WidgetProvider does not inherit from IRawElementProviderFragmentRoot and thus cannot be queried
+ * by external applications for ElementProviderFromPoint(). This should work fine but there could
+ * be unforeseen consquences for cases like modal popup dialogs which are not necessarily windows.
+ */
+class FWindowsUIAWidgetProvider
+	: public FWindowsUIABaseProvider
+	, public IRawElementProviderSimple
+	, public IRawElementProviderFragment
+{
+	friend class FWindowsUIAManager;
+public:
+	// IUnknown
+	virtual HRESULT STDCALL QueryInterface(REFIID riid, void** ppInterface) override;
+	virtual ULONG STDCALL AddRef() override;
+	virtual ULONG STDCALL Release() override;
+	// ~
+
+	// IRawElementProviderSimple
+	virtual HRESULT STDCALL get_ProviderOptions(ProviderOptions* pRetVal) override;
+	virtual HRESULT STDCALL GetPatternProvider(PATTERNID patternId, IUnknown** pRetVal) override;
+	virtual HRESULT STDCALL GetPropertyValue(PROPERTYID propertyId, VARIANT* pRetVal) override;
+	virtual HRESULT STDCALL get_HostRawElementProvider(IRawElementProviderSimple** pRetVal) override;
+	// ~
+
+	// IRawElementProviderFragment
+	virtual HRESULT STDCALL Navigate(NavigateDirection direction, IRawElementProviderFragment** pRetVal) override;
+	virtual HRESULT STDCALL GetRuntimeId(SAFEARRAY** pRetVal) override;
+	virtual HRESULT STDCALL get_BoundingRectangle(UiaRect* pRetVal) override;
+	virtual HRESULT STDCALL GetEmbeddedFragmentRoots(SAFEARRAY** pRetVal) override;
+	virtual HRESULT STDCALL SetFocus() override;
+	virtual HRESULT STDCALL get_FragmentRoot(IRawElementProviderFragmentRoot** pRetVal) override;
+	// ~
+
+	/**
+	 * Check if this Provider implements a specific control pattern. A FWindowsUIAControlProvider can
+	 * then be created of the pattern is supported. Since this is a helper function, the caller should
+	 * separately ensure IsValid() before doing anything with the result.
+	 *
+	 * @param PatternId The control pattern to check for
+	 * @return true if a control Provider can be created for this Provider for this control pattern
+	 */
+	bool SupportsInterface(PATTERNID PatternId) const;
+
+protected:
+	FWindowsUIAWidgetProvider(FWindowsUIAManager& InManager, TSharedRef<IAccessibleWidget> InWidget);
+	virtual ~FWindowsUIAWidgetProvider();
+
+//	void UpdateCachedProperties();
+//private:
+//	void UpdateCachedProperty(PROPERTYID PropertyId);
+//	TMap<int32, FVariant> CachedPropertyValues;
+};
+
+/**
+ * Windows UIA Provider implementation for all Window widgets. Widget->AsWindow() must return a valid pointer.
+ */
+class FWindowsUIAWindowProvider
+	: public FWindowsUIAWidgetProvider
+	, public IRawElementProviderFragmentRoot
+{
+	friend class FWindowsUIAManager;
+public:
+	// IUnknown
+	virtual HRESULT STDCALL QueryInterface(REFIID riid, void** ppInterface) override;
+	virtual ULONG STDCALL AddRef() override;
+	virtual ULONG STDCALL Release() override;
+	// ~
+
+	// IRawElementProviderSimple
+	virtual HRESULT STDCALL get_HostRawElementProvider(IRawElementProviderSimple** pRetVal) override;
+	virtual HRESULT STDCALL GetPatternProvider(PATTERNID patternId, IUnknown** pRetVal) override;
+	// ~
+
+	// IRawElementProviderFragmentRoot
+	virtual HRESULT STDCALL ElementProviderFromPoint(double x, double y, IRawElementProviderFragment** pRetVal) override;
+	virtual HRESULT STDCALL GetFocus(IRawElementProviderFragment** pRetVal) override;
+	// ~
+
+protected:
+	/** Note: InWidget must also be an IAccessibleWindow */
+	FWindowsUIAWindowProvider(FWindowsUIAManager& InManager, TSharedRef<IAccessibleWidget> InWidget);
+	virtual ~FWindowsUIAWindowProvider();
+};
+
+/**
+ * Helper class which can be used in conjunction with FWindowsUIAManager::GetWidgetProvider, which will
+ * automatically calls Release on the Provider when the variable goes out of scope. Note that GetWidgetProvider
+ * already increments the ref count when its called, so the scoped provider does not have to also increase it.
+ */
+class FScopedWidgetProvider
+{
+public:
+	FScopedWidgetProvider(FWindowsUIAWidgetProvider& InProvider)
+		: Provider(InProvider)
+	{
+	}
+	~FScopedWidgetProvider()
+	{
+		Provider.Release();
+	}
+	FWindowsUIAWidgetProvider& Provider;
+};
+
+#endif
diff --git a/Engine/Source/Runtime/ApplicationCore/Public/Windows/WindowsApplication.h b/Engine/Source/Runtime/ApplicationCore/Public/Windows/WindowsApplication.h
--- a/Engine/Source/Runtime/ApplicationCore/Public/Windows/WindowsApplication.h
+++ b/Engine/Source/Runtime/ApplicationCore/Public/Windows/WindowsApplication.h
@@ -342,6 +342,9 @@ public:
 	// GenericApplication overrides
 
 	virtual void SetMessageHandler( const TSharedRef< class FGenericApplicationMessageHandler >& InMessageHandler ) override;
+#if WITH_ACCESSIBILITY
+	virtual void SetAccessibleMessageHandler(const TSharedRef<FGenericAccessibleMessageHandler>& InAccessibleMessageHandler) override;
+#endif
 	virtual void PollGameDeviceState( const float TimeDelta ) override;
 	virtual void PumpMessages( const float TimeDelta ) override;
 	virtual void ProcessDeferredEvents( const float TimeDelta ) override;
@@ -507,6 +510,11 @@ private:
 
 	TSharedPtr<FTaskbarList> TaskbarList;
 
+#if WITH_ACCESSIBILITY
+	/** Handler for WM_GetObject messages that come in */
+	TUniquePtr<class FWindowsUIAManager> UIAManager;
+#endif
+
 	// Accessibility shortcut keys
 	STICKYKEYS							StartupStickyKeys;
 	TOGGLEKEYS							StartupToggleKeys;
diff --git a/Engine/Source/Runtime/Core/Public/Misc/Build.h b/Engine/Source/Runtime/Core/Public/Misc/Build.h
--- a/Engine/Source/Runtime/Core/Public/Misc/Build.h
+++ b/Engine/Source/Runtime/Core/Public/Misc/Build.h
@@ -91,6 +91,13 @@
 	#error UBT should always define WITH_PLUGIN_SUPPORT to be 0 or 1
 #endif
 
+/**
+ * Whether we are compiling with Slate accessibility and automation support
+ */
+#ifndef WITH_ACCESSIBILITY
+	#define WITH_ACCESSIBILITY		1
+#endif
+
  /** Enable perf counters */
 #ifndef WITH_PERFCOUNTERS
 	#define WITH_PERFCOUNTERS		0
diff --git a/Engine/Source/Runtime/Slate/Private/Framework/Application/SlateApplication.cpp b/Engine/Source/Runtime/Slate/Private/Framework/Application/SlateApplication.cpp
--- a/Engine/Source/Runtime/Slate/Private/Framework/Application/SlateApplication.cpp
+++ b/Engine/Source/Runtime/Slate/Private/Framework/Application/SlateApplication.cpp
@@ -25,6 +25,9 @@
 #include "Input/HittestGrid.h"
 #include "Framework/Application/HardwareCursor.h"
 #include "HAL/PlatformApplicationMisc.h"
+#if WITH_ACCESSIBILITY
+#include "Widgets/Accessibility/SlateAccessibleMessageHandler.h"
+#endif
 
 #include "Framework/Application/IWidgetReflector.h"
 #include "Framework/Commands/GenericCommands.h"
@@ -874,9 +877,15 @@ void FPopupSupport::SendNotifications( const FWidgetPath& WidgetsUnderCursor )
 void FSlateApplication::SetPlatformApplication(const TSharedRef<class GenericApplication>& InPlatformApplication)
 {
 	PlatformApplication->SetMessageHandler(MakeShareable(new FGenericApplicationMessageHandler()));
+#if WITH_ACCESSIBILITY
+	PlatformApplication->SetAccessibleMessageHandler(MakeShareable(new FGenericAccessibleMessageHandler()));
+#endif
 
 	PlatformApplication = InPlatformApplication;
 	PlatformApplication->SetMessageHandler(CurrentApplication.ToSharedRef());
+#if WITH_ACCESSIBILITY
+	PlatformApplication->SetAccessibleMessageHandler(CurrentApplication->GetAccessibleMessageHandler().ToSharedRef());
+#endif
 }
 
 void FSlateApplication::OverridePlatformApplication(TSharedPtr<class GenericApplication> InPlatformApplication)
@@ -902,6 +911,9 @@ TSharedRef<FSlateApplication> FSlateApplication::Create(const TSharedRef<class G
 
 	PlatformApplication = InPlatformApplication;
 	PlatformApplication->SetMessageHandler( CurrentApplication.ToSharedRef() );
+#if WITH_ACCESSIBILITY
+	PlatformApplication->SetAccessibleMessageHandler(CurrentApplication->GetAccessibleMessageHandler().ToSharedRef());
+#endif
 
 	// The grid needs to know the size and coordinate system of the desktop.
 	// Some monitor setups have a primary monitor on the right and below the
@@ -954,6 +966,9 @@ FSlateApplication::FSlateApplication()
 	, bSlateWindowActive(true)
 	, Scale( 1.0f )
 	, DragTriggerDistance( 0 )
+#if WITH_ACCESSIBILITY
+	, AccessibleMessageHandler(new FSlateAccessibleMessageHandler())
+#endif
 	, CursorRadius( 0.0f )
 	, LastUserInteractionTime( 0.0 )
 	, LastUserInteractionTimeForThrottling( 0.0 )
@@ -2942,6 +2957,9 @@ bool FSlateApplication::SetUserFocus(FSlateUser* User, const FWidgetPath& InFocu
 
 		// Let previously-focused widget know that it's losing focus
 		OldFocusedWidget->OnFocusLost(FocusEvent);
+#if WITH_ACCESSIBILITY
+		GetAccessibleMessageHandler()->OnWidgetEventRaised(OldFocusedWidget.ToSharedRef(), EAccessibleEvent::FocusChange, true, false);
+#endif
 	}
 
 #if SLATE_HAS_WIDGET_REFLECTOR
@@ -2973,6 +2991,10 @@ bool FSlateApplication::SetUserFocus(FSlateUser* User, const FWidgetPath& InFocu
 		{
 			ProcessReply(InFocusPath, Reply, nullptr, nullptr, User->GetUserIndex());
 		}
+
+#if WITH_ACCESSIBILITY
+		GetAccessibleMessageHandler()->OnWidgetEventRaised(NewFocusedWidget.ToSharedRef(), EAccessibleEvent::FocusChange, false, true);
+#endif
 	}
 
 	return true;
diff --git a/Engine/Source/Runtime/Slate/Private/Widgets/Accessibility/SlateAccessibleWidgets.cpp b/Engine/Source/Runtime/Slate/Private/Widgets/Accessibility/SlateAccessibleWidgets.cpp
new file mode 100644
--- /dev/null
+++ b/Engine/Source/Runtime/Slate/Private/Widgets/Accessibility/SlateAccessibleWidgets.cpp
@@ -0,0 +1,160 @@
+// Copyright 1998-2019 Epic Games, Inc. All Rights Reserved.
+
+#if WITH_ACCESSIBILITY
+
+#include "Widgets/Accessibility/SlateAccessibleWidgets.h"
+#include "Widgets/Input/SButton.h"
+#include "Widgets/Input/SCheckBox.h"
+#include "Widgets/Input/SEditableText.h"
+#include "Widgets/Input/SSlider.h"
+#include "Widgets/Text/STextBlock.h"
+
+void FSlateAccessibleButton::Activate()
+{
+	if (Widget.IsValid())
+	{
+		StaticCastSharedPtr<SButton>(Widget.Pin())->ExecuteOnClick();
+	}
+}
+
+void FSlateAccessibleCheckBox::Activate()
+{
+	if (Widget.IsValid())
+	{
+		StaticCastSharedPtr<SCheckBox>(Widget.Pin())->ToggleCheckedState();
+	}
+}
+
+bool FSlateAccessibleCheckBox::IsCheckable() const {
+	return true;
+}
+
+bool FSlateAccessibleCheckBox::GetCheckedState() const {
+	if (Widget.IsValid())
+	{
+		return StaticCastSharedPtr<SCheckBox>(Widget.Pin())->IsChecked();
+	}
+	return false;
+}
+
+FString FSlateAccessibleEditableText::GetWidgetName() const
+{
+	return FSlateAccessibleWidget::GetWidgetName().Left(50);
+}
+
+const FString& FSlateAccessibleEditableText::GetText() const
+{
+	if (Widget.IsValid())
+	{
+		return StaticCastSharedPtr<SEditableText>(Widget.Pin())->GetText().ToString();
+	}
+	else
+	{
+		static FString EmptyString;
+		return EmptyString;
+	}
+}
+
+bool FSlateAccessibleEditableText::IsPassword() const
+{
+	if (Widget.IsValid())
+	{
+		return StaticCastSharedPtr<SEditableText>(Widget.Pin())->IsTextPassword();
+	}
+	return true;
+}
+
+bool FSlateAccessibleEditableText::IsReadOnly() const
+{
+	if (Widget.IsValid())
+	{
+		return StaticCastSharedPtr<SEditableText>(Widget.Pin())->IsTextReadOnly();
+	}
+	return true;
+}
+
+FString FSlateAccessibleEditableText::GetValue() const
+{
+	return GetText();
+}
+
+void FSlateAccessibleEditableText::SetValue(const FString& Value)
+{
+	if (Widget.IsValid())
+	{
+		StaticCastSharedPtr<SEditableText>(Widget.Pin())->SetText(FText::FromString(Value));
+	}
+}
+
+bool FSlateAccessibleSlider::IsReadOnly() const
+{
+	if (Widget.IsValid())
+	{
+		return !Widget.Pin()->IsEnabled() || !Widget.Pin()->IsInteractable();
+	}
+	return true;
+}
+
+float FSlateAccessibleSlider::GetStepSize() const
+{
+	if (Widget.IsValid())
+	{
+		return StaticCastSharedPtr<SSlider>(Widget.Pin())->GetStepSize();
+	}
+	return 0.0f;
+}
+
+float FSlateAccessibleSlider::GetMaximum() const
+{
+	if (Widget.IsValid())
+	{
+		return StaticCastSharedPtr<SSlider>(Widget.Pin())->GetMaxValue();
+	}
+	return 1.0f;
+}
+
+float FSlateAccessibleSlider::GetMinimum() const
+{
+	if (Widget.IsValid())
+	{
+		return StaticCastSharedPtr<SSlider>(Widget.Pin())->GetMinValue();
+	}
+	return 0.0f;
+}
+
+FString FSlateAccessibleSlider::GetValue() const
+{
+	if (Widget.IsValid())
+	{
+		return FString::SanitizeFloat(StaticCastSharedPtr<SSlider>(Widget.Pin())->GetValue());
+	}
+	return FString::SanitizeFloat(0.0f);
+}
+
+void FSlateAccessibleSlider::SetValue(const FString& Value)
+{
+	if (Widget.IsValid())
+	{
+		StaticCastSharedPtr<SSlider>(Widget.Pin())->SetValue(FCString::Atof(*Value));
+	}
+}
+
+FString FSlateAccessibleTextBlock::GetWidgetName() const
+{
+	return FSlateAccessibleWidget::GetWidgetName().Left(50);
+}
+
+const FString& FSlateAccessibleTextBlock::GetText() const
+{
+	if (Widget.IsValid())
+	{
+		return StaticCastSharedPtr<STextBlock>(Widget.Pin())->GetText().ToString();
+	}
+	else
+	{
+		static const FString EmptyString;
+		return EmptyString;
+	}
+}
+
+#endif
diff --git a/Engine/Source/Runtime/Slate/Private/Widgets/Input/SButton.cpp b/Engine/Source/Runtime/Slate/Private/Widgets/Input/SButton.cpp
--- a/Engine/Source/Runtime/Slate/Private/Widgets/Input/SButton.cpp
+++ b/Engine/Source/Runtime/Slate/Private/Widgets/Input/SButton.cpp
@@ -4,7 +4,10 @@
 #include "Rendering/DrawElements.h"
 #include "Framework/Application/SlateApplication.h"
 #include "Widgets/Text/STextBlock.h"
-
+#if WITH_ACCESSIBILITY
+#include "Widgets/Accessibility/SlateAccessibleWidgets.h"
+#include "Widgets/Accessibility/SlateAccessibleMessageHandler.h"
+#endif
 
 static FName SButtonTypeName("SButton");
 
@@ -161,7 +164,7 @@ FReply SButton::OnKeyDown( const FGeometry& MyGeometry, const FKeyEvent& InKeyEv
 		if (PressMethod == EButtonPressMethod::ButtonPress)
 		{
 			//execute our "OnClicked" delegate, and get the reply
-			Reply = OnClicked.IsBound() ? OnClicked.Execute() : FReply::Handled();
+			Reply = ExecuteOnClick();
 
 			//You should ALWAYS handle the OnClicked event.
 			ensure(Reply.IsEventHandled() == true);
@@ -195,7 +198,7 @@ FReply SButton::OnKeyUp(const FGeometry& MyGeometry, const FKeyEvent& InKeyEvent
 		if ( PressMethod == EButtonPressMethod::ButtonRelease || ( PressMethod == EButtonPressMethod::DownAndUp && bWasPressed ) )
 		{
 			//execute our "OnClicked" delegate, and get the reply
-			Reply = OnClicked.IsBound() ? OnClicked.Execute() : FReply::Handled();
+			Reply = ExecuteOnClick();
 
 			//You should ALWAYS handle the OnClicked event.
 			ensure(Reply.IsEventHandled() == true);
@@ -223,7 +226,7 @@ FReply SButton::OnMouseButtonDown( const FGeometry& MyGeometry, const FPointerEv
 		if(InputClickMethod == EButtonClickMethod::MouseDown)
 		{
 			//get the reply from the execute function
-			Reply = OnClicked.IsBound() ? OnClicked.Execute() : FReply::Handled();
+			Reply = ExecuteOnClick();
 
 			//You should ALWAYS handle the OnClicked event.
 			ensure(Reply.IsEventHandled() == true);
@@ -289,9 +292,9 @@ FReply SButton::OnMouseButtonUp( const FGeometry& MyGeometry, const FPointerEven
 					// pressed the button down first, then we'll allow the click to proceed without an active capture
 					const bool bTriggerForMouseEvent = (InputClickMethod == EButtonClickMethod::MouseUp || HasMouseCapture() );
 
-					if ( ( bTriggerForTouchEvent || bTriggerForMouseEvent ) && OnClicked.IsBound() == true )
+					if ( ( bTriggerForTouchEvent || bTriggerForMouseEvent ) )
 					{
-						Reply = OnClicked.Execute();
+						Reply = ExecuteOnClick();
 					}
 				}
 			}
@@ -363,6 +366,22 @@ void SButton::OnMouseCaptureLost(const FCaptureLostEvent& CaptureLostEvent)
 	Release();
 }
 
+FReply SButton::ExecuteOnClick()
+{
+	if (OnClicked.IsBound())
+	{
+		FReply Reply = OnClicked.Execute();
+#if WITH_ACCESSIBILITY
+		FSlateApplicationBase::Get().GetAccessibleMessageHandler()->OnWidgetEventRaised(AsShared(), EAccessibleEvent::Activate);
+#endif
+		return Reply;
+	}
+	else
+	{
+		return FReply::Handled();
+	}
+}
+
 void SButton::Press()
 {
 	if ( !bIsPressed )
@@ -500,4 +519,11 @@ void SButton::SetTouchMethod(EButtonTouchMethod::Type InTouchMethod)
 void SButton::SetPressMethod(EButtonPressMethod::Type InPressMethod)
 {
 	PressMethod = InPressMethod;
-}
\ No newline at end of file
+}
+
+#if WITH_ACCESSIBILITY
+TSharedPtr<FSlateAccessibleWidget> SButton::CreateAccessibleWidget()
+{
+	return MakeShareable<FSlateAccessibleWidget>(new FSlateAccessibleButton(SharedThis(this)));
+}
+#endif
diff --git a/Engine/Source/Runtime/Slate/Private/Widgets/Input/SCheckBox.cpp b/Engine/Source/Runtime/Slate/Private/Widgets/Input/SCheckBox.cpp
--- a/Engine/Source/Runtime/Slate/Private/Widgets/Input/SCheckBox.cpp
+++ b/Engine/Source/Runtime/Slate/Private/Widgets/Input/SCheckBox.cpp
@@ -7,7 +7,10 @@
 #include "Framework/Application/SlateApplication.h"
 #include "Widgets/Layout/SBorder.h"
 #include "Widgets/Images/SImage.h"
-
+#if WITH_ACCESSIBILITY
+#include "Widgets/Accessibility/SlateAccessibleWidgets.h"
+#include "Widgets/Accessibility/SlateAccessibleMessageHandler.h"
+#endif
 
 /**
  * Construct this widget
@@ -297,6 +300,10 @@ void SCheckBox::ToggleCheckedState()
 		// The state of the check box changed.  Execute the delegate to notify users
 		OnCheckStateChanged.ExecuteIfBound( ECheckBoxState::Checked );
 	}
+
+#if WITH_ACCESSIBILITY
+	FSlateApplicationBase::Get().GetAccessibleMessageHandler()->OnWidgetEventRaised(AsShared(), EAccessibleEvent::Activate, State == ECheckBoxState::Checked, IsCheckboxChecked.Get() == ECheckBoxState::Checked);
+#endif
 }
 
 void SCheckBox::SetIsChecked(TAttribute<ECheckBoxState> InIsChecked)
@@ -509,3 +516,10 @@ const FSlateBrush* SCheckBox::GetUndeterminedPressedImage() const
 {
 	return UndeterminedPressedImage ? UndeterminedPressedImage : &Style->UndeterminedPressedImage;
 }
+
+#if WITH_ACCESSIBILITY
+TSharedPtr<FSlateAccessibleWidget> SCheckBox::CreateAccessibleWidget()
+{
+	return MakeShareable<FSlateAccessibleWidget>(new FSlateAccessibleCheckBox(SharedThis(this)));
+}
+#endif
diff --git a/Engine/Source/Runtime/Slate/Private/Widgets/Input/SEditableText.cpp b/Engine/Source/Runtime/Slate/Private/Widgets/Input/SEditableText.cpp
--- a/Engine/Source/Runtime/Slate/Private/Widgets/Input/SEditableText.cpp
+++ b/Engine/Source/Runtime/Slate/Private/Widgets/Input/SEditableText.cpp
@@ -5,6 +5,9 @@
 #include "Framework/Text/PlainTextLayoutMarshaller.h"
 #include "Widgets/Text/SlateEditableTextLayout.h"
 #include "Types/ReflectionMetadata.h"
+#if WITH_ACCESSIBILITY
+#include "Widgets/Accessibility/SlateAccessibleWidgets.h"
+#endif
 
 SEditableText::SEditableText()
 {
@@ -598,3 +601,16 @@ float SEditableText::UpdateAndClampVerticalScrollBar(const float InViewOffset, c
 {
 	return 0.0f;
 }
+
+#if WITH_ACCESSIBILITY
+TSharedPtr<FSlateAccessibleWidget> SEditableText::CreateAccessibleWidget()
+{
+	return MakeShareable<FSlateAccessibleWidget>(new FSlateAccessibleEditableText(SharedThis(this)));
+}
+
+void SEditableText::SetDefaultAccessibleText(EAccessibleType AccessibleType)
+{
+	TAttribute<FText>& Text = (AccessibleType == EAccessibleType::Main) ? AccessibleData.AccessibleText : AccessibleData.AccessibleSummaryText;
+	Text.Bind(this, &SEditableText::GetHintText);
+}
+#endif
diff --git a/Engine/Source/Runtime/Slate/Private/Widgets/Input/SSlider.cpp b/Engine/Source/Runtime/Slate/Private/Widgets/Input/SSlider.cpp
--- a/Engine/Source/Runtime/Slate/Private/Widgets/Input/SSlider.cpp
+++ b/Engine/Source/Runtime/Slate/Private/Widgets/Input/SSlider.cpp
@@ -2,6 +2,9 @@
 
 #include "Widgets/Input/SSlider.h"
 #include "Rendering/DrawElements.h"
+#if WITH_ACCESSIBILITY
+#include "Widgets/Accessibility/SlateAccessibleWidgets.h"
+#endif
 
 void SSlider::Construct( const SSlider::FArguments& InDeclaration )
 {
@@ -15,7 +18,8 @@ void SSlider::Construct( const SSlider::FArguments& InDeclaration )
 	LockedAttribute = InDeclaration._Locked;
 	Orientation = InDeclaration._Orientation;
 	StepSize = InDeclaration._StepSize;
-	ValueAttribute = InDeclaration._Value;
+	MinValue = InDeclaration._MinValue;
+	MaxValue = InDeclaration._MaxValue;
 	SliderBarColor = InDeclaration._SliderBarColor;
 	SliderHandleColor = InDeclaration._SliderHandleColor;
 	bIsFocusable = InDeclaration._IsFocusable;
@@ -46,7 +50,7 @@ int32 SSlider::OnPaint( const FPaintArgs& Args, const FGeometry& AllottedGeometr
 	const float Indentation = IndentHandle.Get() ? HandleSize.X : 0.0f;
 
 	const float SliderLength = AllottedWidth - (Indentation + HandleSize.X);
-	const float SliderPercent = ValueAttribute.Get();
+	const float SliderPercent = GetNormalizedValue();
 	const float SliderHandleOffset = SliderPercent * SliderLength;
 	const float SliderY = 0.5f * AllottedHeight;
 
@@ -201,7 +205,7 @@ FReply SSlider::OnKeyDown(const FGeometry& MyGeometry, const FKeyEvent& InKeyEve
 				}
 			}
 
-			CommitValue(FMath::Clamp(NewValue, 0.0f, 1.0f));
+			CommitValue(FMath::Clamp(NewValue, MinValue, MaxValue));
 			if (!Reply.IsEventHandled())
 			{
 				Reply = SLeafWidget::OnKeyDown(MyGeometry, InKeyEvent);
@@ -339,6 +343,8 @@ FReply SSlider::OnTouchEnded(const FGeometry& MyGeometry, const FPointerEvent& I
 
 void SSlider::CommitValue(float NewValue)
 {
+	const float OldValue = GetValue();
+
 	if (!ValueAttribute.IsBound())
 	{
 		ValueAttribute.Set(NewValue);
@@ -369,17 +375,17 @@ float SSlider::PositionToValue( const FGeometry& MyGeometry, const FVector2D& Ab
 		RelativeValue = (Denominator != 0.f) ? ((MyGeometry.Size.Y - LocalPosition.Y) - HalfIndentation) / Denominator : 0.f;
 	}
 
-	RelativeValue = FMath::Clamp(RelativeValue, 0.0f, 1.0f);
+	RelativeValue = FMath::Clamp(RelativeValue, 0.0f, 1.0f) * (MaxValue - MinValue) + MinValue;
 	if (bMouseUsesStep)
 	{
 		float direction = ValueAttribute.Get() - RelativeValue;
 		if (direction > StepSize.Get() / 2.0f)
 		{
-			return FMath::Clamp(ValueAttribute.Get() - StepSize.Get(), 0.0f, 1.0f);
+			return FMath::Clamp(ValueAttribute.Get() - StepSize.Get(), MinValue, MaxValue);
 		}
 		else if (direction < StepSize.Get() / -2.0f)
 		{
-			return FMath::Clamp(ValueAttribute.Get() + StepSize.Get(), 0.0f, 1.0f);
+			return FMath::Clamp(ValueAttribute.Get() + StepSize.Get(), MinValue, MaxValue);
 		}
 		else
 		{
@@ -424,11 +430,33 @@ float SSlider::GetValue() const
 	return ValueAttribute.Get();
 }
 
+float SSlider::GetNormalizedValue() const
+{
+	if (MaxValue == MinValue)
+	{
+		return 1.0f;
+	}
+	else
+	{
+		return (ValueAttribute.Get() - MinValue) / (MaxValue - MinValue);
+	}
+}
+
 void SSlider::SetValue(const TAttribute<float>& InValueAttribute)
 {
 	ValueAttribute = InValueAttribute;
 }
 
+void SSlider::SetMinAndMaxValues(float InMinValue, float InMaxValue)
+{
+	MinValue = InMinValue;
+	MaxValue = InMaxValue;
+	if (MinValue > MaxValue)
+	{
+		MaxValue = MinValue;
+	}
+}
+
 void SSlider::SetIndentHandle(const TAttribute<bool>& InIndentHandle)
 {
 	IndentHandle = InIndentHandle;
@@ -472,3 +500,9 @@ void SSlider::SetRequiresControllerLock(bool RequiresControllerLock) {
 	bRequiresControllerLock = RequiresControllerLock;
 }
 
+#if WITH_ACCESSIBILITY
+TSharedPtr<FSlateAccessibleWidget> SSlider::CreateAccessibleWidget()
+{
+	return MakeShareable<FSlateAccessibleWidget>(new FSlateAccessibleSlider(SharedThis(this)));
+}
+#endif
diff --git a/Engine/Source/Runtime/Slate/Private/Widgets/Layout/SWindowTitleBarArea.cpp b/Engine/Source/Runtime/Slate/Private/Widgets/Layout/SWindowTitleBarArea.cpp
--- a/Engine/Source/Runtime/Slate/Private/Widgets/Layout/SWindowTitleBarArea.cpp
+++ b/Engine/Source/Runtime/Slate/Private/Widgets/Layout/SWindowTitleBarArea.cpp
@@ -79,6 +79,7 @@ void SWindowTitleBarArea::Construct( const FArguments& InArgs )
 		[
 			SNew(SImage)
 			.Image(this, &SWindowTitleBarArea::GetMinimizeImage)
+			.AccessibleText(NSLOCTEXT("WindowTitleBar", "Minimize", "Minimize"))
 		];
 
 	MaximizeRestoreButton = SNew(SButton)
@@ -91,6 +92,7 @@ void SWindowTitleBarArea::Construct( const FArguments& InArgs )
 		[
 			SNew(SImage)
 			.Image(this, &SWindowTitleBarArea::GetMaximizeRestoreImage)
+			.AccessibleText(NSLOCTEXT("WindowTitleBar", "Maximize", "Maximize"))
 		];
 
 	CloseButton = SNew(SButton)
@@ -103,6 +105,7 @@ void SWindowTitleBarArea::Construct( const FArguments& InArgs )
 		[
 			SNew(SImage)
 			.Image(this, &SWindowTitleBarArea::GetCloseImage)
+			.AccessibleText(NSLOCTEXT("WindowTitleBar", "Close", "Close"))
 		];
 
 	ChildSlot
diff --git a/Engine/Source/Runtime/Slate/Private/Widgets/Text/STextBlock.cpp b/Engine/Source/Runtime/Slate/Private/Widgets/Text/STextBlock.cpp
--- a/Engine/Source/Runtime/Slate/Private/Widgets/Text/STextBlock.cpp
+++ b/Engine/Source/Runtime/Slate/Private/Widgets/Text/STextBlock.cpp
@@ -8,6 +8,9 @@
 #include "Rendering/DrawElements.h"
 #include "Framework/Application/SlateApplication.h"
 #include "Fonts/FontMeasure.h"
+#if WITH_ACCESSIBILITY
+#include "Widgets/Accessibility/SlateAccessibleWidgets.h"
+#endif
 
 DECLARE_CYCLE_STAT(TEXT("STextBlock::SetText Time"), Stat_SlateTextBlockSetText, STATGROUP_SlateVerbose)
 DECLARE_CYCLE_STAT(TEXT("STextBlock::OnPaint Time"), Stat_SlateTextBlockOnPaint, STATGROUP_SlateVerbose)
@@ -451,3 +454,16 @@ FTextBlockStyle STextBlock::GetComputedTextStyle() const
 	ComputedStyle.SetHighlightShape( *GetHighlightShape() );
 	return ComputedStyle;
 }
+
+#if WITH_ACCESSIBILITY
+TSharedPtr<FSlateAccessibleWidget> STextBlock::CreateAccessibleWidget()
+{
+	return MakeShareable<FSlateAccessibleWidget>(new FSlateAccessibleTextBlock(SharedThis(this)));
+}
+
+void STextBlock::SetDefaultAccessibleText(EAccessibleType AccessibleType)
+{
+	TAttribute<FText>& Text = (AccessibleType == EAccessibleType::Main) ? AccessibleData.AccessibleText : AccessibleData.AccessibleSummaryText;
+	Text.Bind(this, &STextBlock::GetTextCopy);
+}
+#endif
diff --git a/Engine/Source/Runtime/Slate/Private/Widgets/Views/SHeaderRow.cpp b/Engine/Source/Runtime/Slate/Private/Widgets/Views/SHeaderRow.cpp
--- a/Engine/Source/Runtime/Slate/Private/Widgets/Views/SHeaderRow.cpp
+++ b/Engine/Source/Runtime/Slate/Private/Widgets/Views/SHeaderRow.cpp
@@ -48,7 +48,8 @@ public:
 	{
 		check(InArgs._Style);
 
-		SWidget::Construct( InArgs._ToolTipText, InArgs._ToolTip, InArgs._Cursor, InArgs._IsEnabled, InArgs._Visibility, InArgs._RenderOpacity, InArgs._RenderTransform, InArgs._RenderTransformPivot, InArgs._Tag, InArgs._ForceVolatile, InArgs._Clipping, InArgs._FlowDirectionPreference, InArgs.MetaData );
+		SWidget::Construct( InArgs._ToolTipText, InArgs._ToolTip, InArgs._Cursor, InArgs._IsEnabled, InArgs._Visibility, InArgs._RenderOpacity, InArgs._RenderTransform, InArgs._RenderTransformPivot, InArgs._Tag, InArgs._ForceVolatile, InArgs._Clipping,
+							InArgs._FlowDirectionPreference, InArgs._AccessibleParams, InArgs.MetaData );
 
 		Style = InArgs._Style;
 		ColumnId = Column.ColumnId;
@@ -414,7 +415,8 @@ void SHeaderRow::Construct( const FArguments& InArgs )
 {
 	check(InArgs._Style);
 
-	SWidget::Construct( InArgs._ToolTipText, InArgs._ToolTip, InArgs._Cursor, InArgs._IsEnabled, InArgs._Visibility, InArgs._RenderOpacity, InArgs._RenderTransform, InArgs._RenderTransformPivot, InArgs._Tag, InArgs._ForceVolatile, InArgs._Clipping, InArgs._FlowDirectionPreference, InArgs.MetaData );
+	SWidget::Construct( InArgs._ToolTipText, InArgs._ToolTip, InArgs._Cursor, InArgs._IsEnabled, InArgs._Visibility, InArgs._RenderOpacity, InArgs._RenderTransform, InArgs._RenderTransformPivot, InArgs._Tag, InArgs._ForceVolatile, InArgs._Clipping,
+						InArgs._FlowDirectionPreference, InArgs._AccessibleParams, InArgs.MetaData );
 
 	ScrollBarThickness = FVector2D::ZeroVector;
 	ScrollBarVisibility = EVisibility::Collapsed;
diff --git a/Engine/Source/Runtime/Slate/Public/Framework/Application/SWindowTitleBar.h b/Engine/Source/Runtime/Slate/Public/Framework/Application/SWindowTitleBar.h
--- a/Engine/Source/Runtime/Slate/Public/Framework/Application/SWindowTitleBar.h
+++ b/Engine/Source/Runtime/Slate/Public/Framework/Application/SWindowTitleBar.h
@@ -43,7 +43,9 @@ class SAppIconWidget
 {
 	SLATE_BEGIN_ARGS( SAppIconWidget )
 		: _IconColorAndOpacity( FLinearColor::White )
-	{}
+	{
+		_AccessibleText = NSLOCTEXT("AppIconWidget", "System", "System Menu");
+	}
 
 	/** Icon color and opacity */
 	SLATE_ATTRIBUTE( FSlateColor, IconColorAndOpacity )
@@ -207,6 +209,7 @@ protected:
 						SNew(SImage)
 							.Image(this, &SWindowTitleBar::GetMinimizeImage)
 							.ColorAndOpacity(this, &SWindowTitleBar::GetWindowTitleContentColor)
+							.AccessibleText(NSLOCTEXT("WindowTitleBar", "Minimize", "Minimize"))
 					]
 				;
 
@@ -222,6 +225,7 @@ protected:
 						SNew(SImage)
 							.Image(this, &SWindowTitleBar::GetMaximizeRestoreImage)
 							.ColorAndOpacity(this, &SWindowTitleBar::GetWindowTitleContentColor)
+							.AccessibleText(NSLOCTEXT("WindowTitleBar", "Maximize", "Maximize"))
 					]
 				;
 
@@ -237,6 +241,7 @@ protected:
 						SNew(SImage)
 							.Image(this, &SWindowTitleBar::GetCloseImage)
 							.ColorAndOpacity(this, &SWindowTitleBar::GetWindowTitleContentColor)
+							.AccessibleText(NSLOCTEXT("WindowTitleBar", "Close", "Close"))
 					]
 				;
 		}
diff --git a/Engine/Source/Runtime/Slate/Public/Framework/Application/SlateApplication.h b/Engine/Source/Runtime/Slate/Public/Framework/Application/SlateApplication.h
--- a/Engine/Source/Runtime/Slate/Public/Framework/Application/SlateApplication.h
+++ b/Engine/Source/Runtime/Slate/Public/Framework/Application/SlateApplication.h
@@ -30,6 +30,9 @@
 #include "Framework/Application/GestureDetector.h"
 
 class FNavigationConfig;
+#if WITH_ACCESSIBILITY
+class FSlateAccessibleMessageHandler;
+#endif
 class IInputInterface;
 class IInputProcessor;
 class IPlatformTextField;
@@ -1413,6 +1416,10 @@ public:
 	virtual void SetAllUserFocus(const FWidgetPath& InFocusPath, const EFocusCause InCause) override;
 	virtual void SetAllUserFocusAllowingDescendantFocus(const FWidgetPath& InFocusPath, const EFocusCause InCause) override;
 	virtual TSharedPtr<SWidget> GetUserFocusedWidget(uint32 UserIndex) const override;
+#if WITH_ACCESSIBILITY
+	virtual TSharedPtr<FSlateAccessibleMessageHandler> GetAccessibleMessageHandler() const override { return AccessibleMessageHandler; }
+#endif
+	virtual const TArray<TSharedRef<SWindow>> GetTopLevelWindows() const override { return SlateWindows; }
 
 	DECLARE_EVENT_OneParam(FSlateApplication, FApplicationActivationStateChangedEvent, const bool /*IsActive*/)
 	virtual FApplicationActivationStateChangedEvent& OnApplicationActivationStateChanged() { return ApplicationActivationStateChangedEvent; }
@@ -1741,7 +1748,10 @@ private:
 
 	/** These windows will be destroyed next tick. */
 	TArray< TSharedRef<SWindow> > WindowDestroyQueue;
-	
+#if WITH_ACCESSIBILITY
+	/** Manager for widgets and application to interact with accessibility API */
+	TSharedRef<FSlateAccessibleMessageHandler> AccessibleMessageHandler;
+#endif
 	/** The stack of menus that are open */
 	FMenuStack MenuStack;
 
diff --git a/Engine/Source/Runtime/Slate/Public/SlateExtras.h b/Engine/Source/Runtime/Slate/Public/SlateExtras.h
--- a/Engine/Source/Runtime/Slate/Public/SlateExtras.h
+++ b/Engine/Source/Runtime/Slate/Public/SlateExtras.h
@@ -27,7 +27,6 @@ MONOLITHIC_HEADER_BOILERPLATE()
 #include "Framework/Commands/Commands.h"
 #include "Framework/Commands/UICommandList.h"
 #include "Widgets/SWeakWidget.h"
-#include "Framework/Text/TextRange.h"
 #include "Framework/Text/TextRunRenderer.h"
 #include "Framework/Text/TextLineHighlight.h"
 #include "Framework/Text/TextHitPoint.h"
diff --git a/Engine/Source/Runtime/Slate/Public/Widgets/Accessibility/SlateAccessibleWidgets.h b/Engine/Source/Runtime/Slate/Public/Widgets/Accessibility/SlateAccessibleWidgets.h
new file mode 100644
--- /dev/null
+++ b/Engine/Source/Runtime/Slate/Public/Widgets/Accessibility/SlateAccessibleWidgets.h
@@ -0,0 +1,136 @@
+// Copyright 1998-2019 Epic Games, Inc. All Rights Reserved.
+
+#pragma once
+
+#if WITH_ACCESSIBILITY
+
+#include "CoreMinimal.h"
+#include "Widgets/Accessibility/SlateCoreAccessibleWidgets.h"
+
+// SButton
+class SLATE_API FSlateAccessibleButton
+	: public FSlateAccessibleWidget
+	, public IAccessibleActivatable
+{
+public:
+	FSlateAccessibleButton(TWeakPtr<SWidget> InWidget) : FSlateAccessibleWidget(InWidget, EAccessibleWidgetType::Button) {}
+	virtual ~FSlateAccessibleButton() {}
+
+	virtual IAccessibleActivatable* AsActivatable() override { return this; }
+
+	// IAccessibleActivatable
+	virtual void Activate() override;
+	// ~
+};
+// ~
+
+// SCheckBox
+class SLATE_API FSlateAccessibleCheckBox
+	: public FSlateAccessibleWidget
+	, public IAccessibleActivatable
+{
+public:
+	FSlateAccessibleCheckBox(TWeakPtr<SWidget> InWidget) : FSlateAccessibleWidget(InWidget, EAccessibleWidgetType::CheckBox) {}
+	virtual ~FSlateAccessibleCheckBox() {}
+
+	virtual IAccessibleActivatable* AsActivatable() override { return this; }
+
+	// IAccessibleActivatable
+	virtual void Activate() override;
+	virtual bool IsCheckable() const override;
+	virtual bool GetCheckedState() const override;
+	// ~
+};
+// ~
+
+// SEditableText
+class SLATE_API FSlateAccessibleEditableText
+	: public FSlateAccessibleWidget
+	, public IAccessibleText
+	, public IAccessibleProperty
+{
+public:
+	FSlateAccessibleEditableText(TWeakPtr<SWidget> InWidget) : FSlateAccessibleWidget(InWidget, EAccessibleWidgetType::TextEdit) {}
+
+	// IAccessibleWidget
+	virtual FString GetWidgetName() const override;
+	virtual IAccessibleText* AsText() override { return this; }
+	virtual IAccessibleProperty* AsProperty() override { return this; }
+	// ~
+
+	// IAccessibleText
+	virtual const FString& GetText() const override;
+	// ~
+
+	// IAccessibleProperty
+	virtual bool IsReadOnly() const override;
+	virtual bool IsPassword() const override;
+	virtual FString GetValue() const override;
+	virtual void SetValue(const FString& Value) override;
+	// ~
+};
+// ~
+
+// SHyperlink
+class SLATE_API FSlateAccessibleHyperlink
+	: public FSlateAccessibleButton
+{
+public:
+	FSlateAccessibleHyperlink(TWeakPtr<SWidget> InWidget) : FSlateAccessibleButton(InWidget) { WidgetType = EAccessibleWidgetType::Hyperlink; }
+	
+	// todo: add way to get URL for external hyperlinks? may need to go with SHyperlinkLaunchURL
+};
+// ~
+
+// Layouts
+class SLATE_API FSlateAccessibleLayout
+	: public FSlateAccessibleWidget
+{
+public:
+	FSlateAccessibleLayout(TWeakPtr<SWidget> InWidget) : FSlateAccessibleWidget(InWidget, EAccessibleWidgetType::Layout) {}
+};
+// ~
+
+// SSlider
+class SLATE_API FSlateAccessibleSlider
+	: public FSlateAccessibleWidget
+	, public IAccessibleProperty
+{
+public:
+	FSlateAccessibleSlider(TWeakPtr<SWidget> InWidget) : FSlateAccessibleWidget(InWidget, EAccessibleWidgetType::Slider) {}
+
+	// IAccessibleWidget
+	virtual IAccessibleProperty* AsProperty() override { return this; }
+	// ~
+
+	// IAccessibleProperty
+	virtual bool IsReadOnly() const override;
+	virtual float GetStepSize() const override;
+	virtual float GetMaximum() const override;
+	virtual float GetMinimum() const override;
+	virtual FString GetValue() const override;
+	virtual void SetValue(const FString& Value) override;
+	// ~
+};
+// ~
+
+// STextBlock
+class SLATE_API FSlateAccessibleTextBlock
+	: public FSlateAccessibleWidget
+	//, public IAccessibleText // Disabled until we have better support for text. JAWS will not read these properly as-is.
+{
+public:
+	FSlateAccessibleTextBlock(TWeakPtr<SWidget> InWidget) : FSlateAccessibleWidget(InWidget, EAccessibleWidgetType::Text) {}
+
+	// IAccessibleWidget
+	virtual FString GetWidgetName() const override;
+	//virtual IAccessibleText* AsText() override { return this; }
+	// ~
+
+	// IAccessibleText
+	virtual const FString& GetText() const /*override*/;
+	// ~
+};
+// ~
+
+#endif
diff --git a/Engine/Source/Runtime/Slate/Public/Widgets/Input/SButton.h b/Engine/Source/Runtime/Slate/Public/Widgets/Input/SButton.h
--- a/Engine/Source/Runtime/Slate/Public/Widgets/Input/SButton.h
+++ b/Engine/Source/Runtime/Slate/Public/Widgets/Input/SButton.h
@@ -26,6 +26,10 @@ enum class ETextShapingMethod : uint8;
 class SLATE_API SButton
 	: public SBorder
 {
+#if WITH_ACCESSIBILITY
+	// Allow the accessible button to "click" this button
+	friend class FSlateAccessibleButton;
+#endif
 public:
 
 	SLATE_BEGIN_ARGS( SButton )
@@ -45,6 +49,7 @@ public:
 		, _ForegroundColor( FCoreStyle::Get().GetSlateColor( "InvertedForeground" ) )
 		, _IsFocusable( true )
 		{
+			_AccessibleParams = FAccessibleWidgetData(EAccessibleBehavior::Summary, EAccessibleBehavior::Auto, false);
 		}
 
 		/** Slot for this button's content (optional) */
@@ -178,9 +183,13 @@ public:
 	virtual void OnMouseCaptureLost(const FCaptureLostEvent& CaptureLostEvent) override;
 	virtual bool IsInteractable() const override;
 	virtual bool ComputeVolatility() const override;
+#if WITH_ACCESSIBILITY
+	virtual TSharedPtr<FSlateAccessibleWidget> CreateAccessibleWidget() override;
+#endif
 	// SWidget
 
 protected:
+	FReply ExecuteOnClick();
 
 	/** @return combines the user-specified margin and the button's internal margin. */
 	FMargin GetCombinedPadding() const;
diff --git a/Engine/Source/Runtime/Slate/Public/Widgets/Input/SCheckBox.h b/Engine/Source/Runtime/Slate/Public/Widgets/Input/SCheckBox.h
--- a/Engine/Source/Runtime/Slate/Public/Widgets/Input/SCheckBox.h
+++ b/Engine/Source/Runtime/Slate/Public/Widgets/Input/SCheckBox.h
@@ -52,7 +52,9 @@ public:
 		, _UndeterminedImage( nullptr )
 		, _UndeterminedHoveredImage( nullptr )
 		, _UndeterminedPressedImage( nullptr )
-		{}
+	{
+		_AccessibleParams = FAccessibleWidgetData(EAccessibleBehavior::Summary, EAccessibleBehavior::Auto, false);
+	}
 
 		/** Content to be placed next to the check box, or for a toggle button, the content to be placed inside the button */
 		SLATE_DEFAULT_SLOT( FArguments, Content )
@@ -146,6 +148,9 @@ public:
 	virtual void OnMouseEnter( const FGeometry& MyGeometry, const FPointerEvent& MouseEvent ) override;
 	virtual void OnMouseLeave( const FPointerEvent& MouseEvent ) override;
 	virtual bool IsInteractable() const override;
+#if WITH_ACCESSIBILITY
+	virtual TSharedPtr<FSlateAccessibleWidget> CreateAccessibleWidget() override;
+#endif
 	// End of SWidget interface
 
 	/**
diff --git a/Engine/Source/Runtime/Slate/Public/Widgets/Input/SEditableText.h b/Engine/Source/Runtime/Slate/Public/Widgets/Input/SEditableText.h
--- a/Engine/Source/Runtime/Slate/Public/Widgets/Input/SEditableText.h
+++ b/Engine/Source/Runtime/Slate/Public/Widgets/Input/SEditableText.h
@@ -65,6 +65,7 @@ public:
 		, _TextFlowDirection()
 		{
 			_Clipping = EWidgetClipping::ClipToBounds;
+			_AccessibleParams = FAccessibleWidgetData(EAccessibleBehavior::Auto, EAccessibleBehavior::Auto, false);
 		}
 
 		/** Sets the text content for this editable text widget */
@@ -352,6 +353,10 @@ protected:
 	virtual const FSlateBrush* GetFocusBrush() const;
 	virtual bool IsInteractable() const override;
 	virtual bool ComputeVolatility() const override;
+#if WITH_ACCESSIBILITY
+	virtual TSharedPtr<FSlateAccessibleWidget> CreateAccessibleWidget() override;
+	virtual void SetDefaultAccessibleText(EAccessibleType AccessibleType = EAccessibleType::Main) override;
+#endif
 	//~ End SWidget Interface
 
 protected:
diff --git a/Engine/Source/Runtime/Slate/Public/Widgets/Input/SHyperlink.h b/Engine/Source/Runtime/Slate/Public/Widgets/Input/SHyperlink.h
--- a/Engine/Source/Runtime/Slate/Public/Widgets/Input/SHyperlink.h
+++ b/Engine/Source/Runtime/Slate/Public/Widgets/Input/SHyperlink.h
@@ -12,6 +12,9 @@
 #include "Styling/SlateTypes.h"
 #include "Styling/CoreStyle.h"
 #include "Widgets/Input/SButton.h"
+#if WITH_ACCESSIBILITY
+#include "Widgets/Accessibility/SlateAccessibleWidgets.h"
+#endif
 
 enum class ETextFlowDirection : uint8;
 enum class ETextShapingMethod : uint8;
@@ -81,6 +84,12 @@ public:
 	{
 		return FCursorReply::Cursor( EMouseCursor::Hand );
 	}
+#if WITH_ACCESSIBILITY
+	virtual TSharedPtr<FSlateAccessibleWidget> CreateAccessibleWidget() override
+	{
+		return MakeShareable<FSlateAccessibleWidget>(new FSlateAccessibleHyperlink(SharedThis(this)));
+	}
+#endif
 
 protected:
 
diff --git a/Engine/Source/Runtime/Slate/Public/Widgets/Input/SSlider.h b/Engine/Source/Runtime/Slate/Public/Widgets/Input/SSlider.h
--- a/Engine/Source/Runtime/Slate/Public/Widgets/Input/SSlider.h
+++ b/Engine/Source/Runtime/Slate/Public/Widgets/Input/SSlider.h
@@ -35,11 +35,15 @@ public:
 		, _Style(&FCoreStyle::Get().GetWidgetStyle<FSliderStyle>("Slider"))
 		, _StepSize(0.01f)
 		, _Value(1.f)
+		, _MinValue(0.0f)
+		, _MaxValue(1.0f)
 		, _IsFocusable(true)
 		, _OnMouseCaptureBegin()
 		, _OnMouseCaptureEnd()
 		, _OnValueChanged()
-		{ }
+		{
+			_AccessibleParams = FAccessibleWidgetData(EAccessibleBehavior::Summary, EAccessibleBehavior::Auto, false);
+		}
 
 		/** Whether the slidable area should be indented to fit the handle. */
 		SLATE_ATTRIBUTE( bool, IndentHandle )
@@ -71,6 +75,11 @@ public:
 		/** A value that drives where the slider handle appears. Value is normalized between 0 and 1. */
 		SLATE_ATTRIBUTE( float, Value )
 
+		/** The minimum value that can be specified by using the slider. */
+		SLATE_ARGUMENT(float, MinValue)
+		/** The maximum value that can be specified by using the slider. */
+		SLATE_ARGUMENT(float, MaxValue)
+
 		/** Sometimes a slider should only be mouse-clickable and never keyboard focusable. */
 		SLATE_ARGUMENT(bool, IsFocusable)
 
@@ -98,37 +107,49 @@ public:
 	 */
 	void Construct( const SSlider::FArguments& InDeclaration );
 
-	/** See the Value attribute */
+	/** Get the MinValue attribute */
+	float GetMinValue() const { return MinValue; }
+
+	/** Get the MaxValue attribute */
+	float GetMaxValue() const { return MaxValue; }
+
+	/** Get the Value attribute */
 	float GetValue() const;
 
-	/** See the Value attribute */
+	/** Get the Value attribute scaled from 0 to 1 */
+	float GetNormalizedValue() const;
+
+	/** Set the Value attribute */
 	void SetValue(const TAttribute<float>& InValueAttribute);
+
+	/** Set the MinValue and MaxValue attributes. If the new MinValue is more than the new MaxValue, MaxValue will be changed to equal MinValue. */
+	void SetMinAndMaxValues(float InMinValue, float InMaxValue);
 	
-	/** See the IndentHandle attribute */
+	/** Set the IndentHandle attribute */
 	void SetIndentHandle(const TAttribute<bool>& InIndentHandle);
 	
-	/** See the Locked attribute */
+	/** Set the Locked attribute */
 	void SetLocked(const TAttribute<bool>& InLocked);
 
-	/** See the Orientation attribute */
+	/** Set the Orientation attribute */
 	void SetOrientation(EOrientation InOrientation);
 	
-	/** See the SliderBarColor attribute */
+	/** Set the SliderBarColor attribute */
 	void SetSliderBarColor(FSlateColor InSliderBarColor);
 	
-	/** See the SliderHandleColor attribute */
+	/** Set the SliderHandleColor attribute */
 	void SetSliderHandleColor(FSlateColor InSliderHandleColor);
 
 	/** Get the StepSize attribute */
 	float GetStepSize() const;
 
-	/** See the StepSize attribute */
+	/** Set the StepSize attribute */
 	void SetStepSize(const TAttribute<float>& InStepSize);
 
-	/** See the MouseUsesStep attribute */
+	/** Set the MouseUsesStep attribute */
 	void SetMouseUsesStep(bool MouseUsesStep);
 
-	/** See the RequiresControllerLock attribute */
+	/** Set the RequiresControllerLock attribute */
 	void SetRequiresControllerLock(bool RequiresControllerLock);
 
 public:
@@ -149,6 +170,9 @@ public:
 
 	virtual bool SupportsKeyboardFocus() const override;
 	virtual bool IsInteractable() const override;
+#if WITH_ACCESSIBILITY
+	virtual TSharedPtr<FSlateAccessibleWidget> CreateAccessibleWidget() override;
+#endif
 
 	/** @return Is the handle locked or not? Defaults to false */
 	bool IsLocked() const;
@@ -203,6 +227,9 @@ protected:
 	/** Holds the amount to adjust the value by when using a controller or keyboard */
 	TAttribute<float> StepSize;
 
+	float MinValue;
+	float MaxValue;
+
 	// Holds a flag indicating whether a controller/keyboard is manipulating the slider's value. 
 	// When true, navigation away from the widget is prevented until a new value has been accepted or canceled. 
 	bool bControllerInputCaptured;
diff --git a/Engine/Source/Runtime/Slate/Public/Widgets/Text/STextBlock.h b/Engine/Source/Runtime/Slate/Public/Widgets/Text/STextBlock.h
--- a/Engine/Source/Runtime/Slate/Public/Widgets/Text/STextBlock.h
+++ b/Engine/Source/Runtime/Slate/Public/Widgets/Text/STextBlock.h
@@ -65,6 +65,7 @@ public:
 		, _SimpleTextMode(false)
 		{
 			_Clipping = EWidgetClipping::OnDemand;
+			_AccessibleParams = FAccessibleWidgetData(EAccessibleBehavior::Auto, EAccessibleBehavior::Auto, false);
 		}
 
 		/** The text displayed in this text block */
@@ -163,7 +164,13 @@ public:
 	{
 		return BoundText.Get();
 	}
+private:
+	FText GetTextCopy() const
+	{
+		return BoundText.Get();
+	}
 	
+public:
 	/**
 	 * Sets the text for this text block
 	 *
@@ -233,6 +240,10 @@ public:
 	// SWidget interface
 	virtual int32 OnPaint( const FPaintArgs& Args, const FGeometry& AllottedGeometry, const FSlateRect& MyCullingRect, FSlateWindowElementList& OutDrawElements, int32 LayerId, const FWidgetStyle& InWidgetStyle, bool bParentEnabled ) const override;
 	virtual FVector2D ComputeDesiredSize(float) const override;
+#if WITH_ACCESSIBILITY
+	virtual TSharedPtr<FSlateAccessibleWidget> CreateAccessibleWidget() override;
+	virtual void SetDefaultAccessibleText(EAccessibleType AccessibleType = EAccessibleType::Main) override;
+#endif
 	// End of SWidget interface
 
 protected:
diff --git a/Engine/Source/Runtime/SlateCore/Private/Widgets/Accessibility/SlateAccessibleMessageHandler.cpp b/Engine/Source/Runtime/SlateCore/Private/Widgets/Accessibility/SlateAccessibleMessageHandler.cpp
new file mode 100644
--- /dev/null
+++ b/Engine/Source/Runtime/SlateCore/Private/Widgets/Accessibility/SlateAccessibleMessageHandler.cpp
@@ -0,0 +1,156 @@
+// Copyright 1998-2019 Epic Games, Inc. All Rights Reserved.
+
+#if WITH_ACCESSIBILITY
+
+#include "Widgets/Accessibility/SlateAccessibleMessageHandler.h"
+#include "Widgets/Accessibility/SlateAccessibleWidgetCache.h"
+#include "Application/SlateApplicationBase.h"
+#include "Application/SlateWindowHelper.h"
+#include "Widgets/SWidget.h"
+#include "Widgets/SWindow.h"
+#include "Input/HittestGrid.h"
+
+DECLARE_CYCLE_STAT(TEXT("Slate Accessibility: Parent Updated"), STAT_AccessibilitySlateParentUpdated, STATGROUP_Accessibility);
+DECLARE_CYCLE_STAT(TEXT("Slate Accessibility: Behavior Changed"), STAT_AccessibilitySlateBehaviorChanged, STATGROUP_Accessibility);
+DECLARE_CYCLE_STAT(TEXT("Slate Accessibility: Event Raised"), STAT_AccessibilitySlateEventRaised, STATGROUP_Accessibility);
+
+void FSlateAccessibleMessageHandler::OnActivate()
+{
+	// widgets are initialized when their accessibility window is created
+}
+
+void FSlateAccessibleMessageHandler::OnDeactivate()
+{
+	FSlateAccessibleWidgetCache::Get().ClearAll();
+}
+
+TSharedPtr<IAccessibleWidget> FSlateAccessibleMessageHandler::GetAccessibleWindow(const TSharedRef<FGenericWindow>& InWindow) const
+{
+	if (IsActive())
+	{
+		TSharedPtr<SWindow> SlateWindow = FSlateWindowHelper::FindWindowByPlatformWindow(FSlateApplicationBase::Get().GetTopLevelWindows(), InWindow);
+		if (SlateWindow.IsValid())
+		{
+			return FSlateAccessibleWidgetCache::Get().GetAccessibleWidget(SlateWindow);
+		}
+	}
+	return nullptr;
+}
+
+void FSlateAccessibleMessageHandler::OnWidgetRemoved(SWidget* Widget)
+{
+	if (IsActive())
+	{
+		TSharedPtr<FSlateAccessibleWidget> RemovedWidget = FSlateAccessibleWidgetCache::Get().RemoveWidget(Widget);
+		if (RemovedWidget.IsValid())
+		{
+			RaiseEvent(StaticCastSharedPtr<IAccessibleWidget>(RemovedWidget).ToSharedRef(), EAccessibleEvent::WidgetRemoved);
+		}
+	}
+}
+
+void FSlateAccessibleMessageHandler::OnWidgetParentChanged(TSharedRef<SWidget> Widget)
+{
+	if (IsActive())
+	{
+		SCOPE_CYCLE_COUNTER(STAT_AccessibilitySlateParentUpdated);
+
+		TSharedPtr<SWidget> Parent = Widget->GetParentWidget();
+		while (Parent.IsValid() && !Parent->IsAccessible())
+		{
+			Parent = Parent->GetParentWidget();
+		}
+		
+		if (Widget->IsAccessible())
+		{
+			if (Parent.IsValid())
+			{
+				FSlateAccessibleWidgetCache::Get().GetAccessibleWidget(Widget)->UpdateParent(FSlateAccessibleWidgetCache::Get().GetAccessibleWidget(Parent));
+			}
+			else
+			{
+				FSlateAccessibleWidgetCache::Get().GetAccessibleWidget(Widget)->UpdateParent(nullptr);
+			}
+		}
+		else if (Parent.IsValid() && Parent->CanChildrenBeAccessible())
+		{
+			TSharedPtr<IAccessibleWidget> AccessibleParent = FSlateAccessibleWidgetCache::Get().GetAccessibleWidget(Parent);
+			TArray<TSharedRef<SWidget>> AccessibleChildren = FSlateAccessibleWidget::GetAccessibleChildren(Widget);
+			for (int32 i = 0; i < AccessibleChildren.Num(); ++i)
+			{
+				FSlateAccessibleWidgetCache::Get().GetAccessibleWidget(AccessibleChildren[i])->UpdateParent(AccessibleParent);
+			}
+		}
+	}
+}
+
+void FSlateAccessibleMessageHandler::OnWidgetAccessibleBehaviorChanged(TSharedRef<SWidget> Widget)
+{
+	if (IsActive())
+	{
+		SCOPE_CYCLE_COUNTER(STAT_AccessibilitySlateBehaviorChanged);
+
+		TSharedPtr<SWidget> Parent = Widget->GetParentWidget();
+		while (Parent.IsValid())
+		{
+			if (Parent->IsAccessible())
+			{
+				break;
+			}
+
+			Parent = Parent->GetParentWidget();
+		}
+
+		if (Parent.IsValid())
+		{
+			TSharedPtr<IAccessibleWidget> AccessibleParent = FSlateAccessibleWidgetCache::Get().GetAccessibleWidget(Parent);
+			TArray<TSharedRef<SWidget>> AccessibleChildren = FSlateAccessibleWidget::GetAccessibleChildren(Widget);
+			if (Widget->IsAccessible())
+			{
+				TSharedPtr<FSlateAccessibleWidget> AccessibleWidget = StaticCastSharedPtr<FSlateAccessibleWidget>(FSlateAccessibleWidgetCache::Get().GetAccessibleWidget(Widget));
+				for (int32 i = 0; i < AccessibleChildren.Num(); ++i)
+				{
+					FSlateAccessibleWidgetCache::Get().GetAccessibleWidget(AccessibleChildren[i])->UpdateParent(AccessibleWidget);
+				}
+				AccessibleWidget->UpdateParent(AccessibleParent);
+			}
+			else
+			{
+				for (int32 i = 0; i < AccessibleChildren.Num(); ++i)
+				{
+					FSlateAccessibleWidgetCache::Get().GetAccessibleWidget(AccessibleChildren[i])->UpdateParent(AccessibleParent);
+				}
+			}
+		}
+	}
+}
+
+void FSlateAccessibleMessageHandler::OnWidgetEventRaised(TSharedRef<SWidget> Widget, EAccessibleEvent Event)
+{
+	OnWidgetEventRaised(Widget, Event, FVariant(), FVariant());
+}
+
+void FSlateAccessibleMessageHandler::OnWidgetEventRaised(TSharedRef<SWidget> Widget, EAccessibleEvent Event, FVariant OldValue, FVariant NewValue)
+{
+	if (IsActive())
+	{
+		SCOPE_CYCLE_COUNTER(STAT_AccessibilitySlateEventRaised);
+		// todo: not sure what to do for a case like focus changed to not-accessible widget. maybe pass through a nullptr?
+		if (Widget->IsAccessible())
+		{
+			TSharedPtr<SWidget> Parent = Widget->GetParentWidget();
+			while (Parent.IsValid())
+			{
+				if (!Parent->CanChildrenBeAccessible())
+				{
+					return;
+				}
+				Parent = Parent->GetParentWidget();
+			}
+
+			FSlateAccessibleMessageHandler::RaiseEvent(FSlateAccessibleWidgetCache::Get().GetAccessibleWidget(Widget).ToSharedRef(), Event, OldValue, NewValue);
+		}
+	}
+}
+
+#endif
diff --git a/Engine/Source/Runtime/SlateCore/Private/Widgets/Accessibility/SlateCoreAccessibleWidgets.cpp b/Engine/Source/Runtime/SlateCore/Private/Widgets/Accessibility/SlateCoreAccessibleWidgets.cpp
new file mode 100644
--- /dev/null
+++ b/Engine/Source/Runtime/SlateCore/Private/Widgets/Accessibility/SlateCoreAccessibleWidgets.cpp
@@ -0,0 +1,502 @@
+// Copyright 1998-2019 Epic Games, Inc. All Rights Reserved.
+
+#if WITH_ACCESSIBILITY
+
+#include "Widgets/Accessibility/SlateCoreAccessibleWidgets.h"
+#include "Widgets/Accessibility/SlateAccessibleWidgetCache.h"
+#include "Widgets/Accessibility/SlateAccessibleMessageHandler.h"
+#include "Layout/WidgetPath.h"
+#include "Widgets/IToolTip.h"
+#include "Widgets/SWidget.h"
+#include "Widgets/SWindow.h"
+#include "Input/HittestGrid.h"
+#include "Application/SlateApplicationBase.h"
+#include "Application/SlateWindowHelper.h"
+
+DECLARE_CYCLE_STAT(TEXT("Slate Accessibility: Get Widget At Point"), STAT_AccessibilitySlateGetChildAtPosition, STATGROUP_Accessibility);
+
+FSlateAccessibleWidget::FSlateAccessibleWidget(TWeakPtr<SWidget> InWidget, EAccessibleWidgetType InWidgetType)
+	: Widget(InWidget)
+	, WidgetType(InWidgetType)
+	, SiblingIndex(INDEX_NONE)
+	, bChildrenDirty(true)
+{
+}
+
+FSlateAccessibleWidget::~FSlateAccessibleWidget()
+{
+}
+
+bool FSlateAccessibleWidget::IsValid() const
+{
+	return Widget.IsValid();
+}
+
+TSharedPtr<SWindow> FSlateAccessibleWidget::GetTopLevelSlateWindow() const
+{
+	if (Widget.IsValid())
+	{
+		TSharedPtr<SWidget> WindowWidget = Widget.Pin();
+		// todo: fix this for nested windows
+		while (WindowWidget.IsValid())
+		{
+			if (WindowWidget->Advanced_IsWindow())
+			{
+				return StaticCastSharedPtr<SWindow>(WindowWidget);
+			}
+			WindowWidget = WindowWidget->GetParentWidget();
+		}
+	}
+	return nullptr;
+}
+
+TSharedPtr<IAccessibleWidget> FSlateAccessibleWidget::GetTopLevelWindow() const
+{
+	return FSlateAccessibleWidgetCache::Get().GetAccessibleWidget(GetTopLevelSlateWindow());
+}
+
+FBox2D FSlateAccessibleWidget::GetBounds() const
+{
+	if (Widget.IsValid())
+	{
+		const FGeometry& Geometry = Widget.Pin()->GetCachedGeometry();
+		return FBox2D(Geometry.GetAbsolutePosition(), Geometry.GetAbsolutePosition() + Geometry.GetAbsoluteSize());
+	}
+	return FBox2D();
+}
+
+FString FSlateAccessibleWidget::GetClassName() const
+{
+	if (Widget.IsValid())
+	{
+		// Note: this is technically debug code and not guaranteed to work
+		return Widget.Pin()->GetTypeAsString();
+	}
+	return FString();
+}
+
+FString FSlateAccessibleWidget::GetWidgetName() const
+{
+	if (Widget.IsValid())
+	{
+		FText AccessibleText = Widget.Pin()->GetAccessibleText();
+		if (AccessibleText.IsEmpty())
+		{
+			TSharedPtr<FTagMetaData> Tag = Widget.Pin()->GetMetaData<FTagMetaData>();
+			if (Tag.IsValid())
+			{
+				return Tag->Tag.ToString();
+			}
+			else
+			{
+				return GetClassName();
+			}
+		}
+		else
+		{
+			return AccessibleText.ToString();
+		}
+	}
+	return FString();
+}
+
+FString FSlateAccessibleWidget::GetHelpText() const
+{
+	if (Widget.IsValid())
+	{
+		TSharedPtr<IToolTip> ToolTip = Widget.Pin()->GetToolTip();
+		if (ToolTip.IsValid())
+		{
+			return ToolTip->GetContentWidget()->GetAccessibleText().ToString();
+		}
+	}
+	return FString();
+}
+
+bool FSlateAccessibleWidget::IsEnabled() const
+{
+	if (Widget.IsValid())
+	{
+		return Widget.Pin()->IsEnabled();
+	}
+	return false;
+}
+
+bool FSlateAccessibleWidget::IsHidden() const
+{
+	if (Widget.IsValid())
+	{
+		return !Widget.Pin()->GetVisibility().IsVisible();
+	}
+	return true;
+}
+
+bool FSlateAccessibleWidget::SupportsFocus() const
+{
+	if (Widget.IsValid())
+	{
+		return Widget.Pin()->SupportsKeyboardFocus();
+	}
+	return false;
+}
+
+bool FSlateAccessibleWidget::HasFocus() const
+{
+	if (Widget.IsValid())
+	{
+		return Widget.Pin()->HasKeyboardFocus();
+	}
+	return false;
+}
+
+void FSlateAccessibleWidget::SetFocus()
+{
+	if (SupportsFocus())
+	{
+		TSharedPtr<SWindow> WidgetWindow = GetTopLevelSlateWindow();
+		if (WidgetWindow.IsValid())
+		{
+			TArray<TSharedRef<SWindow>> WindowArray;
+			WindowArray.Add(WidgetWindow.ToSharedRef());
+			FWidgetPath WidgetPath;
+			if (FSlateWindowHelper::FindPathToWidget(WindowArray, Widget.Pin().ToSharedRef(), WidgetPath))
+			{
+				FSlateApplicationBase::Get().SetKeyboardFocus(WidgetPath, EFocusCause::SetDirectly);
+			}
+		}
+	}
+}
+
+void FSlateAccessibleWidget::MarkChildrenDirty()
+{
+	for (int32 i = 0; i < Children.Num(); ++i)
+	{
+		if (Children[i].IsValid())
+		{
+			Children[i].Pin()->SiblingIndex = INDEX_NONE;
+		}
+	}
+
+	Children.Reset();
+	bChildrenDirty = true;
+}
+
+void FSlateAccessibleWidget::UpdateAllChildren(bool bUpdateRecursively)
+{
+	if (bChildrenDirty)
+	{
+		bChildrenDirty = false;
+		if (Widget.IsValid())
+		{
+			TArray<TSharedRef<SWidget>> AccessibleChildren = GetAccessibleChildren(Widget.Pin().ToSharedRef());
+			Children.Reset(AccessibleChildren.Num());
+			for (int32 i = 0; i < AccessibleChildren.Num(); ++i)
+			{
+				TSharedPtr<FSlateAccessibleWidget> Child = FSlateAccessibleWidgetCache::Get().GetAccessibleWidget(AccessibleChildren[i]);
+				Children.Add(Child);
+				Child->Parent = StaticCastSharedRef<FSlateAccessibleWidget>(AsShared());
+				Child->SiblingIndex = i;
+
+				if (bUpdateRecursively)
+				{
+					Child->UpdateAllChildren(true);
+				}
+			}
+		}
+	}
+}
+
+void FSlateAccessibleWidget::UpdateParent(TSharedPtr<IAccessibleWidget> NewParent)
+{
+	if (Parent == NewParent)
+	{
+		return;
+	}
+
+	if (Parent.IsValid())
+	{
+		// Even though we're storing SiblingIndex, we have no guarantee
+		// that GetChildren will return the same order after a widget is
+		// added or removed, so we have to re-update all indices.
+		Parent.Pin()->MarkChildrenDirty();
+		FSlateApplicationBase::Get().GetAccessibleMessageHandler()->RaiseEvent(AsShared(), EAccessibleEvent::BeforeRemoveFromParent);
+	}
+
+	Parent = StaticCastSharedPtr<FSlateAccessibleWidget>(NewParent);
+
+	if (Parent.IsValid())
+	{
+		Parent.Pin()->MarkChildrenDirty();
+		FSlateApplicationBase::Get().GetAccessibleMessageHandler()->RaiseEvent(AsShared(), EAccessibleEvent::AfterAddToParent);
+	}
+	else
+	{
+		SiblingIndex = INDEX_NONE;
+	}
+}
+
+TSharedPtr<IAccessibleWidget> FSlateAccessibleWidget::GetParent()
+{
+	if (Parent.IsValid())
+	{
+		return Parent.Pin();
+	}
+	return nullptr;
+}
+
+TSharedPtr<IAccessibleWidget> FSlateAccessibleWidget::GetNextSibling()
+{
+	if (Parent.IsValid())
+	{
+		TSharedPtr<FSlateAccessibleWidget> SharedParent = Parent.Pin();
+		SharedParent->UpdateAllChildren();
+		if (SiblingIndex >= 0 && SiblingIndex < SharedParent->Children.Num() - 1)
+		{
+			const TWeakPtr<FSlateAccessibleWidget>& Child = SharedParent->Children[SiblingIndex + 1];
+			if (Child.IsValid())
+			{
+				return Child.Pin();
+			}
+		}
+	}
+	return nullptr;
+}
+
+TSharedPtr<IAccessibleWidget> FSlateAccessibleWidget::GetPreviousSibling()
+{
+	if (Parent.IsValid())
+	{
+		TSharedPtr<FSlateAccessibleWidget> SharedParent = Parent.Pin();
+		SharedParent->UpdateAllChildren();
+		if (SiblingIndex >= 1 && SiblingIndex < SharedParent->Children.Num())
+		{
+			const TWeakPtr<FSlateAccessibleWidget>& Child = SharedParent->Children[SiblingIndex - 1];
+			if (Child.IsValid())
+			{
+				return Child.Pin();
+			}
+		}
+	}
+	return nullptr;
+}
+
+TSharedPtr<IAccessibleWidget> FSlateAccessibleWidget::GetChildAt(int32 Index)
+{
+	UpdateAllChildren();
+	if (Index >= 0 && Index < Children.Num() && Widget.IsValid() && Widget.Pin()->CanChildrenBeAccessible())
+	{
+		const TWeakPtr<FSlateAccessibleWidget>& Child = Children[Index];
+		if (Child.IsValid())
+		{
+			return Child.Pin();
+		}
+
+	}
+	return nullptr;
+}
+
+int32 FSlateAccessibleWidget::GetNumberOfChildren()
+{
+	UpdateAllChildren();
+	if (Widget.IsValid() && Widget.Pin()->CanChildrenBeAccessible())
+	{
+		return Children.Num();
+	}
+	return 0;
+}
+
+TArray<TSharedRef<SWidget>> FSlateAccessibleWidget::GetAccessibleChildren(TSharedRef<SWidget> Widget)
+{
+	TArray<TSharedRef<SWidget>> AccessibleChildren;
+	if (!Widget->CanChildrenBeAccessible())
+	{
+		return AccessibleChildren;
+	}
+
+	FChildren* Children = Widget->GetChildren();
+	if (Children)
+	{
+		for (int32 i = 0; i < Children->Num(); ++i)
+		{
+			TSharedRef<SWidget> Child = Children->GetChildAt(i);
+			if (Child->GetAccessibleBehavior() != EAccessibleBehavior::NotAccessible)
+			{
+				AccessibleChildren.Add(Child);
+			}
+			else
+			{
+				AccessibleChildren.Append(GetAccessibleChildren(Child));
+			}
+		}
+	}
+	return AccessibleChildren;
+}
+
+TSharedPtr<IAccessibleWidget> FSlateAccessibleWidget::GetChildAtUsingGeometry(int32 X, int32 Y)
+{
+	// This is slow and we should use the HitTest grid when possible.
+	if (!IsHidden() && GetBounds().IsInside(FVector2D(X, Y)))
+	{
+		UpdateAllChildren();
+		for (int32 i = 0; i < Children.Num(); ++i)
+		{
+			if (Children[i].IsValid())
+			{
+				TSharedPtr<IAccessibleWidget> Child = Children[i].Pin()->GetChildAtUsingGeometry(X, Y);
+				if (Child.IsValid())
+				{
+					return Child;
+				}
+			}
+		}
+		return AsShared();
+	}
+	return nullptr;
+}
+
+// SWindow
+TSharedPtr<FGenericWindow> FSlateAccessibleWindow::GetNativeWindow() const
+{
+	if (Widget.IsValid())
+	{
+		return StaticCastSharedPtr<SWindow>(Widget.Pin())->GetNativeWindow();
+	}
+	return nullptr;
+}
+
+TSharedPtr<IAccessibleWidget> FSlateAccessibleWindow::GetChildAtPosition(int32 X, int32 Y)
+{
+	if (Widget.IsValid())
+	{
+		static const bool UseHitTestGrid = false;
+
+		SCOPE_CYCLE_COUNTER(STAT_AccessibilitySlateGetChildAtPosition);
+		if (UseHitTestGrid)
+		{
+			TSharedPtr<SWindow> SlateWindow = StaticCastSharedPtr<SWindow>(Widget.Pin());
+			TArray<FWidgetAndPointer> Hits = SlateWindow->GetHittestGrid()->GetBubblePath(FVector2D(X, Y), 0.0f, false);
+			TSharedPtr<SWidget> LastAccessibleWidget = nullptr;
+			for (int32 i = 0; i < Hits.Num(); ++i)
+			{
+				if (Hits[i].Widget->GetAccessibleBehavior() != EAccessibleBehavior::NotAccessible)
+				{
+					LastAccessibleWidget = Hits[i].Widget;
+				}
+				if (!Hits[i].Widget->CanChildrenBeAccessible())
+				{
+					break;
+				}
+			}
+			return FSlateAccessibleWidgetCache::Get().GetAccessibleWidget(LastAccessibleWidget);
+		}
+		else
+		{
+			return GetChildAtUsingGeometry(X, Y);
+		}
+	}
+
+	return nullptr;
+}
+
+TSharedPtr<IAccessibleWidget> FSlateAccessibleWindow::GetFocusedWidget() const
+{
+	return FSlateAccessibleWidgetCache::Get().GetAccessibleWidget(FSlateApplicationBase::Get().GetKeyboardFocusedWidget());
+}
+
+FString FSlateAccessibleWindow::GetWidgetName() const
+{
+	if (Widget.IsValid())
+	{
+		return StaticCastSharedPtr<SWindow>(Widget.Pin())->GetTitle().ToString();
+	}
+	else
+	{
+		return FSlateAccessibleWidget::GetWidgetName();
+	}
+}
+
+void FSlateAccessibleWindow::Close()
+{
+	if (Widget.IsValid())
+	{
+		StaticCastSharedPtr<SWindow>(Widget.Pin())->RequestDestroyWindow();
+	}
+}
+
+bool FSlateAccessibleWindow::SupportsDisplayState(EWindowDisplayState State) const
+{
+	if (Widget.IsValid())
+	{
+		switch (State)
+		{
+		case IAccessibleWindow::EWindowDisplayState::Normal:
+			return true;
+		case IAccessibleWindow::EWindowDisplayState::Minimize:
+			return StaticCastSharedPtr<SWindow>(Widget.Pin())->HasMinimizeBox();
+		case IAccessibleWindow::EWindowDisplayState::Maximize:
+			return StaticCastSharedPtr<SWindow>(Widget.Pin())->HasMaximizeBox();
+		}
+	}
+	return false;
+}
+
+IAccessibleWindow::EWindowDisplayState FSlateAccessibleWindow::GetDisplayState() const
+{
+	if (Widget.IsValid())
+	{
+		TSharedPtr<SWindow> Window = StaticCastSharedPtr<SWindow>(Widget.Pin());
+		if (Window->IsWindowMaximized())
+		{
+			return IAccessibleWindow::EWindowDisplayState::Maximize;
+		}
+		else if (Window->IsWindowMinimized())
+		{
+			return IAccessibleWindow::EWindowDisplayState::Minimize;
+		}
+		else
+		{
+			return IAccessibleWindow::EWindowDisplayState::Normal;
+		}
+	}
+	return IAccessibleWindow::EWindowDisplayState::Normal;
+}
+
+void FSlateAccessibleWindow::SetDisplayState(EWindowDisplayState State)
+{
+	if (Widget.IsValid() && GetDisplayState() != State)
+	{
+		switch (State)
+		{
+		case IAccessibleWindow::EWindowDisplayState::Normal:
+			StaticCastSharedPtr<SWindow>(Widget.Pin())->Restore();
+			break;
+		case IAccessibleWindow::EWindowDisplayState::Minimize:
+			StaticCastSharedPtr<SWindow>(Widget.Pin())->Minimize();
+			break;
+		case IAccessibleWindow::EWindowDisplayState::Maximize:
+			StaticCastSharedPtr<SWindow>(Widget.Pin())->Maximize();
+			break;
+		}
+	}
+}
+
+bool FSlateAccessibleWindow::IsModal() const
+{
+	if (Widget.IsValid())
+	{
+		return StaticCastSharedPtr<SWindow>(Widget.Pin())->IsModalWindow();
+	}
+	return false;
+}
+
+// ~
+
+// SImage
+FString FSlateAccessibleImage::GetHelpText() const
+{
+	// todo: See UIA_HelpTextPropertyId on https://docs.microsoft.com/en-us/windows/desktop/winauto/uiauto-supportimagecontroltype
+	return FString();
+}
+// ~
+
+#endif
diff --git a/Engine/Source/Runtime/SlateCore/Private/Widgets/Images/SImage.cpp b/Engine/Source/Runtime/SlateCore/Private/Widgets/Images/SImage.cpp
--- a/Engine/Source/Runtime/SlateCore/Private/Widgets/Images/SImage.cpp
+++ b/Engine/Source/Runtime/SlateCore/Private/Widgets/Images/SImage.cpp
@@ -2,6 +2,10 @@
 
 #include "Widgets/Images/SImage.h"
 #include "Rendering/DrawElements.h"
+#include "Widgets/IToolTip.h"
+#if WITH_ACCESSIBILITY
+#include "Widgets/Accessibility/SlateCoreAccessibleWidgets.h"
+#endif
 
 void SImage::Construct( const FArguments& InArgs )
 {
@@ -74,3 +78,10 @@ void SImage::SetImage(TAttribute<const FSlateBrush*> InImage)
 		Invalidate(EInvalidateWidget::LayoutAndVolatility);
 	}
 }
+
+#if WITH_ACCESSIBILITY
+TSharedPtr<FSlateAccessibleWidget> SImage::CreateAccessibleWidget()
+{
+	return MakeShareable<FSlateAccessibleWidget>(new FSlateAccessibleImage(SharedThis(this)));
+}
+#endif
diff --git a/Engine/Source/Runtime/SlateCore/Private/Widgets/SWidget.cpp b/Engine/Source/Runtime/SlateCore/Private/Widgets/SWidget.cpp
--- a/Engine/Source/Runtime/SlateCore/Private/Widgets/SWidget.cpp
+++ b/Engine/Source/Runtime/SlateCore/Private/Widgets/SWidget.cpp
@@ -15,6 +15,10 @@
 #include "Application/ActiveTimerHandle.h"
 #include "Input/HittestGrid.h"
 #include "Debugging/SlateDebugging.h"
+#if WITH_ACCESSIBILITY
+#include "Widgets/Accessibility/SlateCoreAccessibleWidgets.h"
+#include "Widgets/Accessibility/SlateAccessibleMessageHandler.h"
+#endif
 
 DECLARE_DWORD_COUNTER_STAT(TEXT("Widgets Created (Per Frame)"), STAT_SlateTotalWidgetsPerFrame, STATGROUP_Slate);
 DECLARE_DWORD_COUNTER_STAT(TEXT("SWidget::Paint (Count)"), STAT_SlateNumPaintedWidgets, STATGROUP_Slate);
@@ -142,6 +146,9 @@ SWidget::~SWidget()
 		{
 			FSlateApplicationBase::Get().UnRegisterActiveTimer(ActiveTimerHandle);
 		}
+#if WITH_ACCESSIBILITY
+		FSlateApplicationBase::Get().GetAccessibleMessageHandler()->OnWidgetRemoved(this);
+#endif
 	}
 
 	DEC_DWORD_STAT(STAT_SlateTotalWidgets);
@@ -161,6 +168,7 @@ void SWidget::Construct(
 	const bool InForceVolatile,
 	const EWidgetClipping InClipping,
 	const EFlowDirectionPreference InFlowPreference,
+	const FAccessibleWidgetData& InAccessibleData,
 	const TArray<TSharedRef<ISlateMetaData>>& InMetaData
 )
 {
@@ -191,11 +199,29 @@ void SWidget::Construct(
 	Clipping = InClipping;
 	FlowDirectionPreference = InFlowPreference;
 	MetaData = InMetaData;
+
+#if WITH_ACCESSIBILITY
+	AccessibleData.bCanChildrenBeAccessible = InAccessibleData.bCanChildrenBeAccessible;
+	// If custom text is provided, force behavior to custom. Otherwise, use the passed-in behavior and set their default text.
+	if (!InAccessibleData.AccessibleText.IsSet())
+	{
+		SetDefaultAccessibleText(EAccessibleType::Main);
+	}
+	if (!InAccessibleData.AccessibleSummaryText.IsSet())
+	{
+		SetDefaultAccessibleText(EAccessibleType::Summary);
+	}
+	SetAccessibleBehavior(InAccessibleData.AccessibleText.IsSet() ? EAccessibleBehavior::Custom : InAccessibleData.AccessibleBehavior, InAccessibleData.AccessibleText, EAccessibleType::Main);
+	SetAccessibleBehavior(InAccessibleData.AccessibleSummaryText.IsSet() ? EAccessibleBehavior::Custom : InAccessibleData.AccessibleSummaryBehavior, InAccessibleData.AccessibleSummaryText, EAccessibleType::Summary);
+#endif
 }
 
-void SWidget::SWidgetConstruct(const TAttribute<FText>& InToolTipText, const TSharedPtr<IToolTip>& InToolTip, const TAttribute< TOptional<EMouseCursor::Type> >& InCursor, const TAttribute<bool>& InEnabledState, const TAttribute<EVisibility>& InVisibility, const float InRenderOpacity, const TAttribute<TOptional<FSlateRenderTransform>>& InTransform, const TAttribute<FVector2D>& InTransformPivot, const FName& InTag, const bool InForceVolatile, const EWidgetClipping InClipping, const EFlowDirectionPreference InFlowPreference, const TArray<TSharedRef<ISlateMetaData>>& InMetaData)
+void SWidget::SWidgetConstruct(const TAttribute<FText>& InToolTipText, const TSharedPtr<IToolTip>& InToolTip, const TAttribute< TOptional<EMouseCursor::Type> >& InCursor, const TAttribute<bool>& InEnabledState,
+							   const TAttribute<EVisibility>& InVisibility, const float InRenderOpacity, const TAttribute<TOptional<FSlateRenderTransform>>& InTransform, const TAttribute<FVector2D>& InTransformPivot,
+							   const FName& InTag, const bool InForceVolatile, const EWidgetClipping InClipping, const EFlowDirectionPreference InFlowPreference, const FAccessibleWidgetData& InAccessibleData,
+							   const TArray<TSharedRef<ISlateMetaData>>& InMetaData)
 {
-	Construct(InToolTipText, InToolTip, InCursor, InEnabledState, InVisibility, InRenderOpacity, InTransform, InTransformPivot, InTag, InForceVolatile, InClipping, InFlowPreference, InMetaData);
+	Construct(InToolTipText, InToolTip, InCursor, InEnabledState, InVisibility, InRenderOpacity, InTransform, InTransformPivot, InTag, InForceVolatile, InClipping, InFlowPreference, InAccessibleData, InMetaData);
 }
 
 FReply SWidget::OnFocusReceived(const FGeometry& MyGeometry, const FFocusEvent& InFocusEvent)
@@ -562,6 +588,12 @@ void SWidget::AssignParentWidget(TSharedPtr<SWidget> InParent)
 #endif
 
 	ParentWidgetPtr = InParent;
+#if WITH_ACCESSIBILITY
+	if (FSlateApplicationBase::IsInitialized())
+	{
+		FSlateApplicationBase::Get().GetAccessibleMessageHandler()->OnWidgetParentChanged(AsShared());
+	}
+#endif
 	if (InParent.IsValid())
 	{
 		InParent->Invalidate(EInvalidateWidget::Layout);
@@ -578,6 +610,12 @@ bool SWidget::ConditionallyDetatchParentWidget(SWidget* InExpectedParent)
 	if (Parent.Get() == InExpectedParent)
 	{
 		ParentWidgetPtr.Reset();
+#if WITH_ACCESSIBILITY
+		if (FSlateApplicationBase::IsInitialized())
+		{
+			FSlateApplicationBase::Get().GetAccessibleMessageHandler()->OnWidgetParentChanged(AsShared());
+		}
+#endif
 
 		if (Parent.IsValid())
 		{
@@ -1232,6 +1270,139 @@ void SWidget::SetOnMouseLeave(FSimpleNoReplyPointerEventHandler EventHandler)
 	MouseLeaveHandler = EventHandler;
 }
 
+#if WITH_ACCESSIBILITY
+TSharedPtr<FSlateAccessibleWidget> SWidget::CreateAccessibleWidget()
+{
+	return MakeShareable<FSlateAccessibleWidget>(new FSlateAccessibleWidget(AsShared()));
+}
+
+void SWidget::SetAccessibleBehavior(EAccessibleBehavior InBehavior, const TAttribute<FText>& InText, EAccessibleType AccessibleType)
+{
+	EAccessibleBehavior& Behavior = (AccessibleType == EAccessibleType::Main) ? AccessibleData.AccessibleBehavior : AccessibleData.AccessibleSummaryBehavior;
+	if (Behavior != InBehavior)
+	{
+		// If switching off of custom, revert back to default text
+		if (Behavior == EAccessibleBehavior::Custom)
+		{
+			SetDefaultAccessibleText(AccessibleType);
+		}
+		else if (InBehavior == EAccessibleBehavior::Custom)
+		{
+			TAttribute<FText>& Text = (AccessibleType == EAccessibleType::Main) ? AccessibleData.AccessibleText : AccessibleData.AccessibleSummaryText;
+			Text = InText;
+		}
+		const bool bWasAccessible = IsAccessible();
+		Behavior = InBehavior;
+		if (AccessibleType == EAccessibleType::Main && bWasAccessible != IsAccessible())
+		{
+			FSlateApplicationBase::Get().GetAccessibleMessageHandler()->OnWidgetAccessibleBehaviorChanged(AsShared());
+		}
+	}
+}
+
+void SWidget::SetCanChildrenBeAccessible(bool InCanChildrenBeAccessible)
+{
+	AccessibleData.bCanChildrenBeAccessible = InCanChildrenBeAccessible;
+	// todo: emit notifications
+}
+
+void SWidget::SetDefaultAccessibleText(EAccessibleType AccessibleType)
+{
+	TAttribute<FText>& Text = (AccessibleType == EAccessibleType::Main) ? AccessibleData.AccessibleText : AccessibleData.AccessibleSummaryText;
+	Text = TAttribute<FText>();
+}
+
+FText SWidget::GetAccessibleText(EAccessibleType AccessibleType) const
+{
+	const EAccessibleBehavior Behavior = (AccessibleType == EAccessibleType::Main) ? AccessibleData.AccessibleBehavior : AccessibleData.AccessibleSummaryBehavior;
+	const EAccessibleBehavior OtherBehavior = (AccessibleType == EAccessibleType::Main) ? AccessibleData.AccessibleSummaryBehavior : AccessibleData.AccessibleBehavior;
+	const TAttribute<FText>& Text = (AccessibleType == EAccessibleType::Main) ? AccessibleData.AccessibleText : AccessibleData.AccessibleSummaryText;
+	const TAttribute<FText>& OtherText = (AccessibleType == EAccessibleType::Main) ? AccessibleData.AccessibleSummaryText : AccessibleData.AccessibleText;
+
+	switch (Behavior)
+	{
+	case EAccessibleBehavior::Custom:
+		return Text.Get(FText::GetEmpty());
+	case EAccessibleBehavior::Summary:
+		return GetAccessibleSummary();
+	case EAccessibleBehavior::ToolTip:
+		if (ToolTip.IsValid() && !ToolTip->IsEmpty())
+		{
+			return ToolTip->GetContentWidget()->GetAccessibleText(EAccessibleType::Main);
+		}
+		break;
+	case EAccessibleBehavior::Auto:
+		// Auto first checks if custom text was set. This should never happen with user-defined values as custom should be
+		// used instead in that case - however, this will be used for widgets with special default text such as TextBlocks.
+		// If no text is found, then it will attempt to use the other variable's text, so that a developer can do things like
+		// leave Summary on Auto, set Main to Custom, and have Summary automatically use Main's value without having to re-type it.
+		if (Text.IsSet())
+		{
+			return Text.Get(FText::GetEmpty());
+		}
+		switch (OtherBehavior)
+		{
+		case EAccessibleBehavior::Custom:
+		case EAccessibleBehavior::ToolTip:
+			return GetAccessibleText(AccessibleType == EAccessibleType::Main ? EAccessibleType::Summary : EAccessibleType::Main);
+		case EAccessibleBehavior::NotAccessible:
+		case EAccessibleBehavior::Summary:
+			return GetAccessibleSummary();
+		}
+		break;
+	}
+	return FText::GetEmpty();
+}
+
+FText SWidget::GetAccessibleSummary() const
+{
+	FTextBuilder Builder;
+	FChildren* Children = const_cast<SWidget*>(this)->GetChildren();
+	if (Children)
+	{
+		for (int32 i = 0; i < Children->Num(); ++i)
+		{
+			FText Text = Children->GetChildAt(i)->GetAccessibleText(EAccessibleType::Summary);
+			if (!Text.IsEmpty())
+			{
+				Builder.AppendLine(Text);
+			}
+		}
+	}
+	return Builder.ToText();
+}
+
+bool SWidget::IsAccessible() const
+{
+	if (AccessibleData.AccessibleBehavior == EAccessibleBehavior::NotAccessible)
+	{
+		return false;
+	}
+
+	TSharedPtr<SWidget> Parent = GetParentWidget();
+	while (Parent.IsValid())
+	{
+		if (!Parent->CanChildrenBeAccessible())
+		{
+			return false;
+		}
+		Parent = Parent->GetParentWidget();
+	}
+	return true;
+}
+
+EAccessibleBehavior SWidget::GetAccessibleBehavior(EAccessibleType AccessibleType) const
+{
+	return AccessibleType == EAccessibleType::Main ? AccessibleData.AccessibleBehavior : AccessibleData.AccessibleSummaryBehavior;
+}
+
+bool SWidget::CanChildrenBeAccessible() const
+{
+	return AccessibleData.bCanChildrenBeAccessible;
+}
+
+#endif
+
 #if SLATE_CULL_WIDGETS
 
 bool SWidget::IsChildWidgetCulled(const FSlateRect& MyCullingRect, const FArrangedWidget& ArrangedChild) const
diff --git a/Engine/Source/Runtime/SlateCore/Private/Widgets/SWindow.cpp b/Engine/Source/Runtime/SlateCore/Private/Widgets/SWindow.cpp
--- a/Engine/Source/Runtime/SlateCore/Private/Widgets/SWindow.cpp
+++ b/Engine/Source/Runtime/SlateCore/Private/Widgets/SWindow.cpp
@@ -6,7 +6,9 @@
 #include "Layout/WidgetPath.h"
 #include "Input/HittestGrid.h"
 #include "HAL/PlatformApplicationMisc.h"
-
+#if WITH_ACCESSIBILITY
+#include "Widgets/Accessibility/SlateCoreAccessibleWidgets.h"
+#endif
 
 FOverlayPopupLayer::FOverlayPopupLayer(const TSharedRef<SWindow>& InitHostWindow, const TSharedRef<SWidget>& InitPopupContent, TSharedPtr<SOverlay> InitOverlay)
 	: FPopupLayer(InitHostWindow, InitPopupContent)
@@ -906,7 +908,7 @@ void SWindow::ReshapeWindow( FVector2D NewPosition, FVector2D NewSize )
 
 void SWindow::ReshapeWindow( const FSlateRect& InNewShape )
 {
-	ReshapeWindow( FVector2D(InNewShape.Left, InNewShape.Top), FVector2D( InNewShape.Right - InNewShape.Left,  InNewShape.Bottom - InNewShape.Top) );
+	ReshapeWindow( FVector2D(500, 500), FVector2D( InNewShape.Right - InNewShape.Left,  InNewShape.Bottom - InNewShape.Top) );
 }
 
 void SWindow::Resize( FVector2D NewSize )
@@ -1317,7 +1319,7 @@ void SWindow::ShowWindow()
 		{
 			SlatePrepass( FSlateApplicationBase::Get().GetApplicationScale() * NativeWindow->GetDPIScaleFactor() );
 			const FVector2D WindowDesiredSizePixels = GetDesiredSizeDesktopPixels();
-			ReshapeWindow( InitialDesiredScreenPosition - (WindowDesiredSizePixels * 0.5f), WindowDesiredSizePixels);
+			ReshapeWindow( FVector2D(200, 200), WindowDesiredSizePixels);
 		}
 
 		// Set the window to be maximized if we need to.  Note that this won't actually show the window if its not
@@ -2067,3 +2069,16 @@ FScopedSwitchWorldHack::FScopedSwitchWorldHack( const FWidgetPath& WidgetPath )
 	}
 }
 #endif
+
+#if WITH_ACCESSIBILITY
+TSharedPtr<FSlateAccessibleWidget> SWindow::CreateAccessibleWidget()
+{
+	return MakeShareable<FSlateAccessibleWidget>(new FSlateAccessibleWindow(SharedThis(this)));
+}
+
+void SWindow::SetDefaultAccessibleText(EAccessibleType AccessibleType)
+{
+	TAttribute<FText>& Text = (AccessibleType == EAccessibleType::Main) ? AccessibleData.AccessibleText : AccessibleData.AccessibleSummaryText;
+	Text.Bind(this, &SWindow::GetTitle);
+}
+#endif
diff --git a/Engine/Source/Runtime/SlateCore/Public/Application/SlateApplicationBase.h b/Engine/Source/Runtime/SlateCore/Public/Application/SlateApplicationBase.h
--- a/Engine/Source/Runtime/SlateCore/Public/Application/SlateApplicationBase.h
+++ b/Engine/Source/Runtime/SlateCore/Public/Application/SlateApplicationBase.h
@@ -133,6 +133,13 @@ public:
 	 */
 	virtual TSharedPtr<SWindow> GetActiveTopLevelWindow() const = 0;
 
+	/**
+	 * Get a list of all top-level windows in the application, excluding virtual windows.
+	 *
+	 * @return An array of all current top-level SWindows.
+	 */
+	virtual const TArray<TSharedRef<SWindow>> GetTopLevelWindows() const = 0;
+
 	/**
 	 * Gets the global application icon.
 	 *
@@ -227,6 +234,15 @@ public:
 
 	/** @return	Returns true if there are any pop-up menus summoned */
 	virtual bool AnyMenusVisible() const = 0;
+
+#if WITH_ACCESSIBILITY
+	/** 
+	 * Subclasses should store their own accessibility handler and return it through this function.
+	 *
+	 * @return A pointer to the accessibility handler, if one exists
+	 */
+	virtual TSharedPtr<class FSlateAccessibleMessageHandler> GetAccessibleMessageHandler() const = 0;
+#endif
 protected:
 	/**
 	 * Implementation of GetMouseCaptor which can be overridden without warnings.
diff --git a/Engine/Source/Runtime/SlateCore/Public/Widgets/Accessibility/SlateAccessibleMessageHandler.h b/Engine/Source/Runtime/SlateCore/Public/Widgets/Accessibility/SlateAccessibleMessageHandler.h
new file mode 100644
--- /dev/null
+++ b/Engine/Source/Runtime/SlateCore/Public/Widgets/Accessibility/SlateAccessibleMessageHandler.h
@@ -0,0 +1,60 @@
+#pragma once
+
+#if WITH_ACCESSIBILITY
+
+#include "GenericPlatform/GenericAccessibleInterfaces.h"
+#include "Misc/Variant.h"
+
+class SWidget;
+
+/**
+ * Message handling system for Slate Accessibility API, dealing with both receiving events and pushing them back to the platform layer.
+ */
+class SLATECORE_API FSlateAccessibleMessageHandler : public FGenericAccessibleMessageHandler
+{
+public:
+	// FGenericAccessibleMessageHandler
+	virtual void OnActivate() override;
+	virtual void OnDeactivate() override;
+	virtual TSharedPtr<IAccessibleWidget> GetAccessibleWindow(const TSharedRef<FGenericWindow>& InWindow) const override;
+	virtual bool ApplicationIsAccessible() const override { return true; }
+	//~
+
+	/**
+	 * Callback for SWidget destructor. Removes the corresponding accessible widget for the Slate widget.
+	 *
+	 * @param Widget The widget that is being deleted.
+	 */
+	void OnWidgetRemoved(SWidget* Widget);
+	/**
+	 * Callback for a Slate widget's parent changing.
+	 *
+	 * @param Widget The widget whose parent changed.
+	 */
+	void OnWidgetParentChanged(TSharedRef<SWidget> Widget);
+	/**
+	 * Callback for a Slate widget's main accessible behavior changing. 
+	 *
+	 * @param Widget The widget whose behavior changed.
+	 */
+	void OnWidgetAccessibleBehaviorChanged(TSharedRef<SWidget> Widget);
+	/**
+	 * Callback for a Slate widget indicating that some event happened to it.
+	 *
+	 * @param Widget The widget raising the event.
+	 * @param Event The type of event being raised.
+	 */
+	void OnWidgetEventRaised(TSharedRef<SWidget> Widget, EAccessibleEvent Event);
+	/**
+	 * Callback for a Slate widget indicating that a property change occurred. This may also be used by certain events
+	 * such as Notification which don't have an 'OldValue'. Only NewValue should be set for those types of events.
+	 *
+	 * @param Widget The widget raising the event.
+	 * @param Event The type of event being raised.
+	 * @param OldValue The value of the property being changed before the change occurred.
+	 * @param NewValue The value of the property being changed after the change occurred, or any miscellaneous data for the event.
+	 */
+	void OnWidgetEventRaised(TSharedRef<SWidget> Widget, EAccessibleEvent Event, FVariant OldValue, FVariant NewValue);
+};
+
+#endif
diff --git a/Engine/Source/Runtime/SlateCore/Public/Widgets/Accessibility/SlateAccessibleWidgetCache.h b/Engine/Source/Runtime/SlateCore/Public/Widgets/Accessibility/SlateAccessibleWidgetCache.h
new file mode 100644
--- /dev/null
+++ b/Engine/Source/Runtime/SlateCore/Public/Widgets/Accessibility/SlateAccessibleWidgetCache.h
@@ -0,0 +1,95 @@
+// Copyright 1998-2019 Epic Games, Inc. All Rights Reserved.
+
+#pragma once
+
+#if WITH_ACCESSIBILITY
+
+#include "CoreMinimal.h"
+#include "Widgets/Accessibility/SlateCoreAccessibleWidgets.h"
+#include "Widgets/SWidget.h"
+
+/**
+ * Singleton used to retrieve accessible widgets for a given Slate widget. Accessible widgets will persist
+ * so longer as the OS is sending accessibility events, after which the platform layer should call ClearAll().
+ */
+class SLATECORE_API FSlateAccessibleWidgetCache
+{
+public:
+	static FSlateAccessibleWidgetCache& Get()
+	{
+		static FSlateAccessibleWidgetCache AccessibleWidgetCache;
+		return AccessibleWidgetCache;
+	}
+
+	/** Empty the cache and release all cached accessible widgets. This should generally only be done with accessibility is turned off. */
+	void ClearAll()
+	{
+		AccessibleWidgetMap.Empty();
+	}
+
+	/**
+	 * Callback for when an SWidget is deleted. FSlateAccessibleMessageHandler should be the only thing to call this.
+	 * The widget is removed from the cache and removed from the accessible tree.
+	 *
+	 * @param Widget the Slate widget that is being deleted
+	 */
+	TSharedPtr<FSlateAccessibleWidget> RemoveWidget(SWidget* Widget)
+	{
+		TSharedPtr<FSlateAccessibleWidget> RemovedWidget = nullptr;
+		TSharedPtr<FSlateAccessibleWidget>* AccessibleWidget = AccessibleWidgetMap.Find(Widget);
+		if (AccessibleWidget)
+		{
+			// Make sure to disconnect this widget from the tree
+			(*AccessibleWidget)->UpdateParent(nullptr);
+			AccessibleWidgetMap.RemoveAndCopyValue(Widget, RemovedWidget);
+		}
+		return RemovedWidget;
+	}
+
+	/**
+	 * Get a cached accessible widget for a given Slate widget, or create a new one is bCreateIfMissing is true.
+	 * This may return nullptr in the case where the Slate widget is not a valid accessible widget. A non-valid
+	 * accessible widget includes the case where the widget has accessible flags set but its parent cannot have
+	 * accessible children.
+	 *
+	 * @param Widget The Slate widget to get the accessible widget for
+	 * @param bCreateIfMissing If the widget is not already in the cache, whether or not it should be created
+	 * @return The accessible widget for the given Slate widget, or nullptr if the widget is not accessible.
+	 */
+	TSharedPtr<FSlateAccessibleWidget> GetAccessibleWidget(TSharedPtr<SWidget> Widget, bool bCreateIfMissing = true)
+	{
+		if (!Widget.IsValid() || !Widget->IsAccessible())
+		{
+			return nullptr;
+		}
+
+		TSharedPtr<FSlateAccessibleWidget>* AccessibleWidget = AccessibleWidgetMap.Find(Widget.Get());
+		if (AccessibleWidget)
+		{
+			return *AccessibleWidget;
+		}
+		else if (bCreateIfMissing)
+		{
+			TSharedPtr<FSlateAccessibleWidget> NewWidget = AccessibleWidgetMap.Add(Widget.Get(), Widget->CreateAccessibleWidget());
+			// todo: fix this for nested windows
+			if (NewWidget->AsWindow())
+			{
+				NewWidget->UpdateAllChildren(true);
+			}
+			return NewWidget;
+		}
+		return nullptr;
+	}
+
+private:
+	/**
+	 * Hold on to a shared pointer for each accessible widget until the SWidget is deleted. A raw pointer is used
+	 * because widgets are removed from the cache via the SWidget destructor.
+	 */
+	TMap<SWidget*, TSharedPtr<FSlateAccessibleWidget>> AccessibleWidgetMap;
+
+	FSlateAccessibleWidgetCache() {}
+	~FSlateAccessibleWidgetCache() {}
+};
+
+#endif
diff --git a/Engine/Source/Runtime/SlateCore/Public/Widgets/Accessibility/SlateCoreAccessibleWidgets.h b/Engine/Source/Runtime/SlateCore/Public/Widgets/Accessibility/SlateCoreAccessibleWidgets.h
new file mode 100644
--- /dev/null
+++ b/Engine/Source/Runtime/SlateCore/Public/Widgets/Accessibility/SlateCoreAccessibleWidgets.h
@@ -0,0 +1,148 @@
+// Copyright 1998-2019 Epic Games, Inc. All Rights Reserved.
+
+#pragma once
+
+#if WITH_ACCESSIBILITY
+
+#include "CoreMinimal.h"
+#include "GenericPlatform/GenericAccessibleInterfaces.h"
+
+class SWidget;
+class SWindow;
+
+/**
+ * The base implementation of IAccessibleWidget for all Slate widgets. Any new accessible widgets should
+ * inherit directly from FSlateAccessibleWidget, and optionally inherit from other IAccessible interfaces to
+ * provide more functionality.
+ */
+class SLATECORE_API FSlateAccessibleWidget : public IAccessibleWidget
+{
+public:
+	FSlateAccessibleWidget(TWeakPtr<SWidget> InWidget, EAccessibleWidgetType InWidgetType = EAccessibleWidgetType::Unknown);
+	virtual ~FSlateAccessibleWidget();
+
+	// IAccessibleWidget
+	virtual bool IsValid() const override final;
+	virtual TSharedPtr<IAccessibleWidget> GetTopLevelWindow() const override final;
+	virtual FBox2D GetBounds() const override final;
+	virtual TSharedPtr<IAccessibleWidget> GetParent() override final;
+	virtual TSharedPtr<IAccessibleWidget> GetNextSibling() override final;
+	virtual TSharedPtr<IAccessibleWidget> GetPreviousSibling() override final;
+	virtual TSharedPtr<IAccessibleWidget> GetChildAt(int32 Index) override final;
+	virtual int32 GetNumberOfChildren() override final;
+	virtual FString GetClassName() const override final;
+	virtual bool IsEnabled() const override final;
+	virtual bool IsHidden() const override final;
+	virtual bool SupportsFocus() const override final;
+	virtual bool HasFocus() const override final;
+	virtual void SetFocus() override final;
+
+	virtual EAccessibleWidgetType GetWidgetType() const override { return WidgetType; }
+	virtual FString GetWidgetName() const override;
+	virtual FString GetHelpText() const override;
+	// ~
+
+	/** Tell this widget to recompute its children the next time they are requested. */
+	void MarkChildrenDirty();
+	/**
+	 * Detach this widget from its current parent and attach it to a new parent. This will emit notifications back to the accessible message handler.
+	 *
+	 * @param NewParent The widget to assign as the new parent widget.
+	 */
+	void UpdateParent(TSharedPtr<IAccessibleWidget> NewParent);
+	/**
+	 * If MarkChildrenDirty() has been called, recalculate the list of all accessible widgets below this one.
+	 * Because SWidget->GetChildren() has no guarantees about what it returns and how it returns it, we can
+	 * never truly 100% guarantee that the accessible tree will be in sync with the Slate tree.
+	 *
+	 * We make a reasonable assumption that widgets are smart about implementing this function to return the
+	 * same widgets every time. However, we can't assume anything about when a child gets added or removed
+	 * with respect to the ordering of the children. Because of this, we have to recompute their indices
+	 * any time we suspect the hierarchy may have changed.
+	 *
+	 * @param bUpdateRecursively If true, calls UpdateAllChildren() on any children found.
+	 */
+	void UpdateAllChildren(bool bUpdateRecursively = false);
+
+	/**
+	 * Search the Slate hierarchy recursively and generate a list of all accessible widgets whose parent is this widget.
+	 *
+	 * @param AccessibleWidget The root widget to find children for.
+	 * @return All Slate widgets whose accessible parent is the passed-in widget.
+	 */
+	static TArray<TSharedRef<SWidget>> GetAccessibleChildren(TSharedRef<SWidget> AccessibleWidget);
+
+protected:
+	/**
+	 * Recursively find the accessible widget under the specified X,Y coordinates.
+	 *
+	 * @param X The X coordinate to search in absolute screen space.
+	 * @param Y The Y coordinate to search in absolute screen space.
+	 * @return The deepest accessible widget found.
+	 */
+	TSharedPtr<IAccessibleWidget> GetChildAtUsingGeometry(int32 X, int32 Y);
+
+	/** The underlying Slate widget backing this accessible widget. */
+	TWeakPtr<SWidget> Widget;
+	/** What type of widget the platform's accessibility API should treat this as. */
+	EAccessibleWidgetType WidgetType;
+	/** The accessible parent to this widget. This should usually be valid on widgets in the hierarchy, except for SWindows. */
+	TWeakPtr<FSlateAccessibleWidget> Parent;
+	/** All accessible widgets whose parent is this widget. This is not necessarily correct unless UpdateAllChildren() is called first. */
+	TArray<TWeakPtr<FSlateAccessibleWidget>> Children;
+	/** The index of this widget in its parent's list of children. */
+	int32 SiblingIndex;
+	/** Whether the contents of the Children array has changed and UpdateAllChildren() needs to be called. */
+	bool bChildrenDirty;
+
+private:
+	/**
+	 * Find the Slate window containing this widget's underlying Slate widget.
+	 *
+	 * @return The parent SWindow for the Slate widget referenced by this accessible widget.
+	 */
+	TSharedPtr<SWindow> GetTopLevelSlateWindow() const;
+};
+
+// SWindow
+class FSlateAccessibleWindow
+	: public FSlateAccessibleWidget
+	, public IAccessibleWindow
+{
+public:
+	FSlateAccessibleWindow(TWeakPtr<SWidget> InWidget) : FSlateAccessibleWidget(InWidget, EAccessibleWidgetType::Window) {}
+	virtual ~FSlateAccessibleWindow() {}
+
+	// IAccessibleWidget
+	virtual IAccessibleWindow* AsWindow() override { return this; }
+	virtual FString GetWidgetName() const override;
+	// ~
+
+	// IAccessibleWindow
+	virtual TSharedPtr<FGenericWindow> GetNativeWindow() const override;
+	virtual TSharedPtr<IAccessibleWidget> GetChildAtPosition(int32 X, int32 Y) override;
+	virtual TSharedPtr<IAccessibleWidget> GetFocusedWidget() const override;
+	virtual void Close() override;
+	virtual bool SupportsDisplayState(EWindowDisplayState State) const override;
+	virtual EWindowDisplayState GetDisplayState() const override;
+	virtual void SetDisplayState(EWindowDisplayState State) override;
+	virtual bool IsModal() const override;
+	// ~
+};
+// ~
+
+// SImage
+class SLATECORE_API FSlateAccessibleImage
+	: public FSlateAccessibleWidget
+{
+public:
+	FSlateAccessibleImage(TWeakPtr<SWidget> InWidget) : FSlateAccessibleWidget(InWidget, EAccessibleWidgetType::Image) {}
+	virtual ~FSlateAccessibleImage() {}
+
+	// IAccessibleWidget
+	virtual FString GetHelpText() const override;
+	// ~
+};
+// ~
+
+#endif
diff --git a/Engine/Source/Runtime/SlateCore/Public/Widgets/Accessibility/SlateWidgetAccessibleTypes.h b/Engine/Source/Runtime/SlateCore/Public/Widgets/Accessibility/SlateWidgetAccessibleTypes.h
new file mode 100644
--- /dev/null
+++ b/Engine/Source/Runtime/SlateCore/Public/Widgets/Accessibility/SlateWidgetAccessibleTypes.h
@@ -0,0 +1,29 @@
+// Copyright 1998-2019 Epic Games, Inc. All Rights Reserved.
+
+#pragma once
+
+#include "GenericPlatform/GenericAccessibleInterfaces.h"
+
+struct FAccessibleWidgetData
+{
+	FAccessibleWidgetData(EAccessibleBehavior InBehavior = EAccessibleBehavior::NotAccessible, EAccessibleBehavior InSummaryBehavior = EAccessibleBehavior::Auto, bool bInCanChildrenBeAccessible = true)
+		: bCanChildrenBeAccessible(bInCanChildrenBeAccessible)
+		, AccessibleBehavior(InBehavior)
+		, AccessibleSummaryBehavior(InSummaryBehavior)
+	{
+	}
+	FAccessibleWidgetData(const TAttribute<FText>& InAccessibleText, const TAttribute<FText>& InAccessibleSummaryText = TAttribute<FText>(), bool bInCanChildrenBeAccessible = true)
+		: bCanChildrenBeAccessible(bInCanChildrenBeAccessible)
+		, AccessibleBehavior(InAccessibleSummaryText.IsSet() ? EAccessibleBehavior::Custom : EAccessibleBehavior::NotAccessible)
+		, AccessibleSummaryBehavior(InAccessibleSummaryText.IsSet() ? EAccessibleBehavior::Custom : EAccessibleBehavior::Auto)
+		, AccessibleText(InAccessibleText)
+		, AccessibleSummaryText(InAccessibleSummaryText)
+	{
+	}
+
+	uint8 bCanChildrenBeAccessible : 1;
+	EAccessibleBehavior AccessibleBehavior;
+	EAccessibleBehavior AccessibleSummaryBehavior;
+	TAttribute<FText> AccessibleText;
+	TAttribute<FText> AccessibleSummaryText;
+};
diff --git a/Engine/Source/Runtime/SlateCore/Public/Widgets/DeclarativeSyntaxSupport.h b/Engine/Source/Runtime/SlateCore/Public/Widgets/DeclarativeSyntaxSupport.h
--- a/Engine/Source/Runtime/SlateCore/Public/Widgets/DeclarativeSyntaxSupport.h
+++ b/Engine/Source/Runtime/SlateCore/Public/Widgets/DeclarativeSyntaxSupport.h
@@ -11,6 +11,7 @@
 #include "GenericPlatform/ICursor.h"
 #include "Types/ISlateMetaData.h"
 #include "Widgets/SNullWidget.h"
+#include "Widgets/Accessibility/SlateWidgetAccessibleTypes.h"
 
 class IToolTip;
 class SUserWidget;
@@ -776,6 +777,8 @@ struct TSlateBaseNamedArgs
 	, _ForceVolatile( false )
 	, _Clipping( EWidgetClipping::Inherit )
 	, _FlowDirectionPreference( EFlowDirectionPreference::Inherit )
+	, _AccessibleParams()
+	, _AccessibleText()
 	{
 	}
 
@@ -820,6 +823,8 @@ struct TSlateBaseNamedArgs
 	SLATE_ARGUMENT( bool, ForceVolatile )
 	SLATE_ARGUMENT( EWidgetClipping, Clipping )
 	SLATE_ARGUMENT( EFlowDirectionPreference, FlowDirectionPreference)
+	SLATE_ARGUMENT(FAccessibleWidgetData, AccessibleParams)
+	SLATE_ATTRIBUTE(FText, AccessibleText)
 
 	TArray<TSharedRef<ISlateMetaData>> MetaData;
 };
@@ -1093,6 +1098,7 @@ struct TDecl
 			InArgs._ForceVolatile,
 			InArgs._Clipping,
 			InArgs._FlowDirectionPreference,
+			InArgs._AccessibleText.IsSet() ? FAccessibleWidgetData(InArgs._AccessibleText) : InArgs._AccessibleParams,
 			InArgs.MetaData );
 
 		_RequiredArgs.CallConstruct(_Widget, InArgs);
diff --git a/Engine/Source/Runtime/SlateCore/Public/Widgets/Images/SImage.h b/Engine/Source/Runtime/SlateCore/Public/Widgets/Images/SImage.h
--- a/Engine/Source/Runtime/SlateCore/Public/Widgets/Images/SImage.h
+++ b/Engine/Source/Runtime/SlateCore/Public/Widgets/Images/SImage.h
@@ -24,7 +24,7 @@ public:
 		: _Image( FCoreStyle::Get().GetDefaultBrush() )
 		, _ColorAndOpacity( FLinearColor::White )
 		, _FlipForRightToLeftFlowDirection( false )
-		{}
+		{ }
 
 		/** Image resource */
 		SLATE_ATTRIBUTE( const FSlateBrush*, Image )
@@ -68,6 +68,10 @@ public:
 
 	// SWidget overrides
 	virtual int32 OnPaint( const FPaintArgs& Args, const FGeometry& AllottedGeometry, const FSlateRect& MyCullingRect, FSlateWindowElementList& OutDrawElements, int32 LayerId, const FWidgetStyle& InWidgetStyle, bool bParentEnabled ) const override;
+#if WITH_ACCESSIBILITY
+	virtual TSharedPtr<FSlateAccessibleWidget> CreateAccessibleWidget() override;
+#endif
+
 protected:
 	// Begin SWidget overrides.
 	virtual FVector2D ComputeDesiredSize(float) const override;
diff --git a/Engine/Source/Runtime/SlateCore/Public/Widgets/SWidget.h b/Engine/Source/Runtime/SlateCore/Public/Widgets/SWidget.h
--- a/Engine/Source/Runtime/SlateCore/Public/Widgets/SWidget.h
+++ b/Engine/Source/Runtime/SlateCore/Public/Widgets/SWidget.h
@@ -24,6 +24,7 @@
 #include "Types/WidgetActiveTimerDelegate.h"
 #include "Textures/SlateShaderResource.h"
 #include "SlateGlobals.h"
+#include "Widgets/Accessibility/SlateWidgetAccessibleTypes.h"
 
 class FActiveTimerHandle;
 class FArrangedChildren;
@@ -38,6 +39,12 @@ class IToolTip;
 class SWidget;
 struct FSlateBrush;
 
+enum class EAccessibleType : uint8
+{
+	Main,
+	Summary
+};
+
 /** Widget update flags for fast path.  Work in progress, do not modify */
 enum class EWidgetUpdateFlags : uint8
 {
@@ -284,6 +291,7 @@ public:
 		const bool InForceVolatile,
 		const EWidgetClipping InClipping,
 		const EFlowDirectionPreference InFlowPreference,
+		const FAccessibleWidgetData& InAccessibleData,
 		const TArray<TSharedRef<ISlateMetaData>>& InMetaData);
 
 	void SWidgetConstruct(const TAttribute<FText>& InToolTipText,
@@ -298,6 +306,7 @@ public:
 		const bool InForceVolatile,
 		const EWidgetClipping InClipping,
 		const EFlowDirectionPreference InFlowPreference,
+		const FAccessibleWidgetData& InAccessibleData,
 		const TArray<TSharedRef<ISlateMetaData>>& InMetaData);
 
 	//
@@ -681,6 +690,9 @@ public:
 	 */
 	virtual EWindowZone::Type GetWindowZoneOverride() const;
 
+#if WITH_ACCESSIBILITY
+	virtual TSharedPtr<class FSlateAccessibleWidget> CreateAccessibleWidget();
+#endif
 
 public:
 	//	
@@ -962,6 +974,70 @@ public:
 		}
 	}
 
+#if WITH_ACCESSIBILITY
+	/**
+	 * Get the text that should be reported to the user when attempting to access this widget.
+	 *
+	 * @param AccessibleType Whether the widget is being accessed directly or through a summary query.
+	 * @return The text that should be conveyed to the user describing this widget.
+	 */
+	FText GetAccessibleText(EAccessibleType AccessibleType = EAccessibleType::Main) const;
+
+	/**
+	 * Traverse all child widgets and concat their results of GetAccessibleText(Summary).
+	 *
+	 * @return The combined text of all child widget's summary text.
+	 */
+	FText GetAccessibleSummary() const;
+
+	/**
+	 * Whether this widget is considered accessible or not. A widget is accessible if its behavior
+	 * is set to something other than NotAccessible, and all of its parent widgets support accessible children.
+	 *
+	 * @return true if an accessible widget should be created for this widget.
+	 */
+	bool IsAccessible() const;
+
+	/**
+	 * Get the behavior describing how the accessible text of this widget should be retrieved.
+	 *
+	 * @param AccessibleType Whether the widget is being accessed directly or through a summary query.
+	 * @return The accessible behavior of the widget.
+	 */
+	EAccessibleBehavior GetAccessibleBehavior(EAccessibleType AccessibleType = EAccessibleType::Main) const;
+
+	/**
+	 * Checks whether this widget allows its children to be accessible or not.
+	 *
+	 * @return true if children can be accessible.
+	 */
+	bool CanChildrenBeAccessible() const;
+
+	/**
+	 * Set a new accessible behavior, and if the behavior is custom, new accessible text to go along with it.
+	 *
+	 * @param InBehavior The new behavior for the widget. If the new behavior is custom, InText should also be set.
+	 * @param InText, If the new behavior is custom, this will be the custom text assigned to the widget.
+	 * @param AccessibleType Whether the widget is being accessed directly or through a summary query.
+	 */
+	void SetAccessibleBehavior(EAccessibleBehavior InBehavior, const TAttribute<FText>& InText = TAttribute<FText>(), EAccessibleType AccessibleType = EAccessibleType::Main);
+
+	/**
+	 * Sets whether children are allowed to be accessible or not.
+	 * Warning: Calling this function after accessibility is enabled will cause the accessibility tree to become unsynced.
+	 *
+	 * @param InCanChildrenBeAccessible Whether children should be accessible or not.
+	 */
+	void SetCanChildrenBeAccessible(bool InCanChildrenBeAccessible);
+
+	/**
+	 * Assign AccessibleText with a default value that can be used when AccessibleBehavior is set to Auto or Custom.
+	 *
+	 * @param AccessibleType Whether the widget is being accessed directly or through a summary query.
+	 */
+	virtual void SetDefaultAccessibleText(EAccessibleType AccessibleType = EAccessibleType::Main);
+#endif
+
 	/**
 	 * When performing a caching pass, volatile widgets are not cached as part of everything
 	 * else, instead they and their children are drawn as normal standard widgets and excluded
@@ -1578,6 +1654,11 @@ protected:
 	/** Render transform pivot of this widget (in normalized local space) */
 	TAttribute< FVector2D > RenderTransformPivot;
 
+#if WITH_ACCESSIBILITY
+	/** All variables surrounding how this widget is exposed to the platform's accessibility API. */
+	FAccessibleWidgetData AccessibleData;
+#endif
+
 protected:
 	/** Debugging information on the type of widget we're creating for the Widget Reflector. */
 	FName TypeOfWidget;
diff --git a/Engine/Source/Runtime/SlateCore/Public/Widgets/SWindow.h b/Engine/Source/Runtime/SlateCore/Public/Widgets/SWindow.h
--- a/Engine/Source/Runtime/SlateCore/Public/Widgets/SWindow.h
+++ b/Engine/Source/Runtime/SlateCore/Public/Widgets/SWindow.h
@@ -157,6 +157,7 @@ public:
 		, _UserResizeBorder(FMargin(5, 5, 5, 5))
 
 	{
+		_AccessibleParams = FAccessibleWidgetData(EAccessibleBehavior::Auto);
 	}
 
 		/** Type of this window */
@@ -784,6 +785,11 @@ public:
 	bool IsDrawingEnabled() const { return bIsDrawingEnabled; }
 
 	virtual bool Advanced_IsWindow() const { return true; }
+
+#if WITH_ACCESSIBILITY
+	virtual TSharedPtr<FSlateAccessibleWidget> CreateAccessibleWidget() override;
+	virtual void SetDefaultAccessibleText(EAccessibleType AccessibleType = EAccessibleType::Main) override;
+#endif
 private:
 	virtual FReply OnFocusReceived( const FGeometry& MyGeometry, const FFocusEvent& InFocusEvent ) override;
 	virtual FReply OnMouseButtonDown( const FGeometry& MyGeometry, const FPointerEvent& MouseEvent ) override;
diff --git a/Engine/Source/Runtime/UMG/Private/Components/Button.cpp b/Engine/Source/Runtime/UMG/Private/Components/Button.cpp
--- a/Engine/Source/Runtime/UMG/Private/Components/Button.cpp
+++ b/Engine/Source/Runtime/UMG/Private/Components/Button.cpp
@@ -25,6 +25,9 @@ UButton::UButton(const FObjectInitializer& ObjectInitializer)
 	TouchMethod = EButtonTouchMethod::DownAndUp;
 
 	IsFocusable = true;
+
+	AccessibleBehavior = ESlateAccessibleBehavior::Summary;
+	bCanChildrenBeAccessible = false;
 }
 
 void UButton::ReleaseSlateResources(bool bReleaseChildren)
diff --git a/Engine/Source/Runtime/UMG/Private/Components/CheckBox.cpp b/Engine/Source/Runtime/UMG/Private/Components/CheckBox.cpp
--- a/Engine/Source/Runtime/UMG/Private/Components/CheckBox.cpp
+++ b/Engine/Source/Runtime/UMG/Private/Components/CheckBox.cpp
@@ -26,6 +26,8 @@ UCheckBox::UCheckBox(const FObjectInitializer& ObjectInitializer)
 	BorderBackgroundColor_DEPRECATED = FLinearColor::White;
 
 	IsFocusable = true;
+	AccessibleBehavior = ESlateAccessibleBehavior::Summary;
+	bCanChildrenBeAccessible = false;
 }
 
 void UCheckBox::ReleaseSlateResources(bool bReleaseChildren)
diff --git a/Engine/Source/Runtime/UMG/Private/Components/Slider.cpp b/Engine/Source/Runtime/UMG/Private/Components/Slider.cpp
--- a/Engine/Source/Runtime/UMG/Private/Components/Slider.cpp
+++ b/Engine/Source/Runtime/UMG/Private/Components/Slider.cpp
@@ -12,6 +12,8 @@
 USlider::USlider(const FObjectInitializer& ObjectInitializer)
 	: Super(ObjectInitializer)
 {
+	MinValue = 0.0f;
+	MaxValue = 1.0f;
 	Orientation = EOrientation::Orient_Horizontal;
 	SliderBarColor = FLinearColor::White;
 	SliderHandleColor = FLinearColor::White;
@@ -23,6 +25,9 @@ USlider::USlider(const FObjectInitializer& ObjectInitializer)
 	// HACK: THIS SHOULD NOT COME FROM CORESTYLE AND SHOULD INSTEAD BY DEFINED BY ENGINE TEXTURES/PROJECT SETTINGS
 	static const FSliderStyle StaticSlider = FCoreStyle::Get().GetWidgetStyle<FSliderStyle>("Slider");
 	WidgetStyle = StaticSlider;
+
+	AccessibleBehavior = ESlateAccessibleBehavior::Summary;
+	bCanChildrenBeAccessible = false;
 }
 
 TSharedRef<SWidget> USlider::RebuildWidget()
@@ -51,6 +56,7 @@ void USlider::SynchronizeProperties()
 	MySlider->SetSliderBarColor(SliderBarColor);
 	MySlider->SetSliderHandleColor(SliderHandleColor);
 	MySlider->SetValue(ValueBinding);
+	MySlider->SetMinAndMaxValues(MinValue, MaxValue);
 	MySlider->SetLocked(Locked);
 	MySlider->SetIndentHandle(IndentHandle);
 	MySlider->SetStepSize(StepSize);
@@ -98,6 +104,23 @@ float USlider::GetValue() const
 	return Value;
 }
 
+float USlider::GetNormalizedValue() const
+{
+	if (MySlider.IsValid())
+	{
+		return MySlider->GetNormalizedValue();
+	}
+
+	if (MinValue == MaxValue)
+	{
+		return 1.0f;
+	}
+	else
+	{
+		return (Value - MinValue) / (MaxValue - MinValue);
+	}
+}
+
 void USlider::SetValue(float InValue)
 {
 	Value = InValue;
@@ -107,6 +130,27 @@ void USlider::SetValue(float InValue)
 	}
 }
 
+void USlider::SetMinValue(float InValue)
+{
+	MinValue = InValue;
+	if (MySlider.IsValid())
+	{
+		// Because SSlider clamps min/max values upon setting them,
+		// we have to send both values together to ensure that they
+		// don't get out of sync.
+		MySlider->SetMinAndMaxValues(MinValue, MaxValue);
+	}
+}
+
+void USlider::SetMaxValue(float InValue)
+{
+	MaxValue = InValue;
+	if (MySlider.IsValid())
+	{
+		MySlider->SetMinAndMaxValues(MinValue, MaxValue);
+	}
+}
+
 void USlider::SetIndentHandle(bool InIndentHandle)
 {
 	IndentHandle = InIndentHandle;
diff --git a/Engine/Source/Runtime/UMG/Private/Components/TextBlock.cpp b/Engine/Source/Runtime/UMG/Private/Components/TextBlock.cpp
--- a/Engine/Source/Runtime/UMG/Private/Components/TextBlock.cpp
+++ b/Engine/Source/Runtime/UMG/Private/Components/TextBlock.cpp
@@ -30,6 +30,9 @@ UTextBlock::UTextBlock(const FObjectInitializer& ObjectInitializer)
 		static ConstructorHelpers::FObjectFinder<UFont> RobotoFontObj(*UWidget::GetDefaultFontName());
 		Font = FSlateFontInfo(RobotoFontObj.Object, 24, FName("Bold"));
 	}
+
+	AccessibleBehavior = ESlateAccessibleBehavior::Auto;
+	bCanChildrenBeAccessible = false;
 }
 
 void UTextBlock::PostLoad()
diff --git a/Engine/Source/Runtime/UMG/Private/Components/Widget.cpp b/Engine/Source/Runtime/UMG/Private/Components/Widget.cpp
--- a/Engine/Source/Runtime/UMG/Private/Components/Widget.cpp
+++ b/Engine/Source/Runtime/UMG/Private/Components/Widget.cpp
@@ -162,6 +162,10 @@ UWidget::UWidget(const FObjectInitializer& ObjectInitializer)
 	RenderTransformPivot = FVector2D(0.5f, 0.5f);
 	Cursor = EMouseCursor::Default;
 
+	AccessibleBehavior = ESlateAccessibleBehavior::NotAccessible;
+	AccessibleSummaryBehavior = ESlateAccessibleBehavior::Auto;
+	bCanChildrenBeAccessible = true;
+
 #if WITH_EDITORONLY_DATA
 	{ static const FAutoRegisterLocalizationDataGatheringCallback AutomaticRegistrationOfLocalizationGatherer(UWidget::StaticClass(), &GatherWidgetForLocalization); }
 #endif
@@ -1162,6 +1166,12 @@ void UWidget::SynchronizeProperties()
 		SafeWidget->SetToolTipText(PROPERTY_BINDING(FText, ToolTipText));
 	}
 
+#if WITH_ACCESSIBILITY
+	SafeWidget->SetAccessibleBehavior((EAccessibleBehavior)AccessibleBehavior, PROPERTY_BINDING(FText, AccessibleText), EAccessibleType::Main);
+	SafeWidget->SetAccessibleBehavior((EAccessibleBehavior)AccessibleSummaryBehavior, PROPERTY_BINDING(FText, AccessibleSummaryText), EAccessibleType::Summary);
+	SafeWidget->SetCanChildrenBeAccessible(bCanChildrenBeAccessible);
+#endif
+
 #if !UE_BUILD_SHIPPING
 	SafeWidget->AddMetadata<FReflectionMetaData>(MakeShared<FReflectionMetaData>(GetFName(), GetClass(), this, GetSourceAssetOrClass()));
 #endif
diff --git a/Engine/Source/Runtime/UMG/Public/Components/SlateWrapperTypes.h b/Engine/Source/Runtime/UMG/Public/Components/SlateWrapperTypes.h
--- a/Engine/Source/Runtime/UMG/Public/Components/SlateWrapperTypes.h
+++ b/Engine/Source/Runtime/UMG/Public/Components/SlateWrapperTypes.h
@@ -32,6 +32,27 @@ enum class ESlateVisibility : uint8
 	SelfHitTestInvisible
 };
 
+/** Whether a widget should be included in accessibility, and if so, how its text should be retrieved. */
+UENUM(BlueprintType)
+enum class ESlateAccessibleBehavior : uint8
+{
+	/** Not accessible. */
+	NotAccessible,
+	/**
+	 * Accessible, first checking to see if there's any custom default text assigned for widgets of this type.
+	 * If not, then it will attempt to use the alternate behavior (ie AccessibleSummaryBehavior instead of AccessibleBehavior)
+	 * and return that value instead. This acts as a reference so that you only need to set one value for both of them
+	 * to return the same thing.
+	 */
+	Auto,
+	/** Accessible, and traverse all child widgets and concat their AccessibleSummaryText together. */
+	Summary,
+	/** Accessible, and retrieve manually-assigned text from a TAttribute. */
+	Custom,
+	/** Accessible, and use the tooltip's accessible text. */
+	ToolTip
+};
+
 /** The sizing options of UWidgets */
 UENUM(BlueprintType)
 namespace ESlateSizeRule
diff --git a/Engine/Source/Runtime/UMG/Public/Components/Slider.h b/Engine/Source/Runtime/UMG/Public/Components/Slider.h
--- a/Engine/Source/Runtime/UMG/Public/Components/Slider.h
+++ b/Engine/Source/Runtime/UMG/Public/Components/Slider.h
@@ -29,13 +29,21 @@ class UMG_API USlider : public UWidget
 
 public:
 	/** The volume value to display. */
-	UPROPERTY(EditAnywhere, Category=Appearance, meta=( ClampMin="0", ClampMax="1", UIMin="0", UIMax="1"))
+	UPROPERTY(EditAnywhere, Category=Appearance, meta=(UIMin="0", UIMax="1"))
 	float Value;
 
 	/** A bindable delegate to allow logic to drive the value of the widget */
 	UPROPERTY()
 	FGetFloat ValueDelegate;
 
+	/** The minimum value the slider can be set to. */
+	UPROPERTY(EditAnywhere, BlueprintReadOnly, Category = Appearance)
+	float MinValue;
+
+	/** The maximum value the slider can be set to. */
+	UPROPERTY(EditAnywhere, BlueprintReadOnly, Category = Appearance)
+	float MaxValue;
+
 public:
 	
 	/** The progress bar style */
@@ -71,7 +79,7 @@ public:
 	bool RequiresControllerLock;
 
 	/** The amount to adjust the value by, when using a controller or keyboard */
-	UPROPERTY(EditAnywhere, BlueprintReadOnly, Category=Appearance, meta=( ClampMin="0", ClampMax="1", UIMin="0", UIMax="1"))
+	UPROPERTY(EditAnywhere, BlueprintReadOnly, Category=Appearance, meta=(UIMin="0", UIMax="1"))
 	float StepSize;
 
 	/** Should the slider be focusable? */
@@ -104,10 +112,22 @@ public:
 	UFUNCTION(BlueprintCallable, Category="Behavior")
 	float GetValue() const;
 
+	/** Get the current value scaled from 0 to 1 */
+	UFUNCTION(BlueprintCallable, Category = "Behavior")
+	float GetNormalizedValue() const;
+
 	/** Sets the current value of the slider. */
 	UFUNCTION(BlueprintCallable, Category="Behavior")
 	void SetValue(float InValue);
 
+	/** Sets the minimum value of the slider. */
+	UFUNCTION(BlueprintCallable, Category = "Behavior")
+	void SetMinValue(float InValue);
+
+	/** Sets the maximum value of the slider. */
+	UFUNCTION(BlueprintCallable, Category = "Behavior")
+	void SetMaxValue(float InValue);
+
 	/** Sets if the slidable area should be indented to fit the handle */
 	UFUNCTION(BlueprintCallable, Category="Behavior")
 	void SetIndentHandle(bool InValue);
diff --git a/Engine/Source/Runtime/UMG/Public/Components/Widget.h b/Engine/Source/Runtime/UMG/Public/Components/Widget.h
--- a/Engine/Source/Runtime/UMG/Public/Components/Widget.h
+++ b/Engine/Source/Runtime/UMG/Public/Components/Widget.h
@@ -313,6 +313,32 @@ public:
 	UPROPERTY(EditAnywhere, Category="Behavior", meta=(InlineEditConditionToggle))
 	uint8 bOverride_Cursor : 1;
 
+	/** Whether or not children of this widget can appear as distinct accessible widgets. */
+	UPROPERTY(EditAnywhere, Category = "Accessibility")
+	uint8 bCanChildrenBeAccessible : 1;
+
+	/** Whether or not the widget is accessible, and how to describe it. If set to custom, additional customization options will appear. */
+	UPROPERTY(EditAnywhere, Category = "Accessibility")
+	ESlateAccessibleBehavior AccessibleBehavior;
+
+	/** When AccessibleBehavior is set to Custom, this is the text that will be used to describe the widget. */
+	UPROPERTY(EditAnywhere, BlueprintReadOnly, Category = "Accessibility", meta = (MultiLine = true))
+	FText AccessibleText;
+
+	UPROPERTY()
+	FGetText AccessibleTextDelegate;
+
+	/** How to describe this widget when it's being presented through a summary of a parent widget. If set to custom, additional customization options will appear. */
+	UPROPERTY(EditAnywhere, Category = "Accessibility", AdvancedDisplay)
+	ESlateAccessibleBehavior AccessibleSummaryBehavior;
+
+	/** When AccessibleSummaryBehavior is set to Custom, this is the text that will be used to describe the widget. */
+	UPROPERTY(EditAnywhere, BlueprintReadOnly, Category = "Accessibility", meta = (MultiLine = true), AdvancedDisplay)
+	FText AccessibleSummaryText;
+
+	UPROPERTY()
+	FGetText AccessibleSummaryTextDelegate;
+
 protected:
 
 	/**
@@ -949,4 +975,6 @@ private:
 private:
 	PROPERTY_BINDING_IMPLEMENTATION(FText, ToolTipText);
 	PROPERTY_BINDING_IMPLEMENTATION(bool, bIsEnabled);
+	PROPERTY_BINDING_IMPLEMENTATION(FText, AccessibleText);
+	PROPERTY_BINDING_IMPLEMENTATION(FText, AccessibleSummaryText);
 };
