Fix Win7 accessibility problems

diff --git a/Engine/Source/Programs/UnrealBuildTool/Configuration/UEBuildTarget.cs b/Engine/Source/Programs/UnrealBuildTool/Configuration/UEBuildTarget.cs
--- a/Engine/Source/Programs/UnrealBuildTool/Configuration/UEBuildTarget.cs
+++ b/Engine/Source/Programs/UnrealBuildTool/Configuration/UEBuildTarget.cs
@@ -3076,7 +3076,7 @@ namespace UnrealBuildTool
 				GlobalCompileEnvironment.Definitions.Add("WITH_PLUGIN_SUPPORT=0");
 			}
 
-			if (Rules.bCompileWithAccessibilitySupport)
+			if (Rules.bCompileWithAccessibilitySupport && !Rules.bIsBuildingConsoleApplication)
 			{
 				GlobalCompileEnvironment.Definitions.Add("WITH_ACCESSIBILITY=1");
 			}
diff --git a/Engine/Source/Runtime/ApplicationCore/ApplicationCore.Build.cs b/Engine/Source/Runtime/ApplicationCore/ApplicationCore.Build.cs
--- a/Engine/Source/Runtime/ApplicationCore/ApplicationCore.Build.cs
+++ b/Engine/Source/Runtime/ApplicationCore/ApplicationCore.Build.cs
@@ -32,7 +32,7 @@ public class ApplicationCore : ModuleRules
 			AddEngineThirdPartyPrivateStaticDependencies(Target, 
 				"XInput"
 				);
-            if (Target.bCompileWithAccessibilitySupport)
+            if (Target.bCompileWithAccessibilitySupport && !Target.bIsBuildingConsoleApplication)
             {
                 PublicAdditionalLibraries.Add("uiautomationcore.lib");
             }
diff --git a/Engine/Source/Runtime/ApplicationCore/Private/Windows/Accessibility/WindowsUIAManager.cpp b/Engine/Source/Runtime/ApplicationCore/Private/Windows/Accessibility/WindowsUIAManager.cpp
--- a/Engine/Source/Runtime/ApplicationCore/Private/Windows/Accessibility/WindowsUIAManager.cpp
+++ b/Engine/Source/Runtime/ApplicationCore/Private/Windows/Accessibility/WindowsUIAManager.cpp
@@ -149,10 +149,6 @@ void FWindowsUIAManager::OnWidgetProviderRemoved(TSharedRef<IAccessibleWidget> I
 	}
 }
 
-#if WINVER > 0x0A00 // Win10
-typedef HRESULT(WINAPI* UiaRaiseNotificationEventFunc)(IRawElementProviderSimple*, NotificationKind, NotificationProcessing, BSTR, BSTR);
-#endif
-
 void FWindowsUIAManager::OnEventRaised(TSharedRef<IAccessibleWidget> Widget, EAccessibleEvent Event, FVariant OldValue, FVariant NewValue)
 {
 	if (UiaClientsAreListening())
@@ -185,16 +181,13 @@ void FWindowsUIAManager::OnEventRaised(TSharedRef<IAccessibleWidget> Widget, EAc
 			break;
 		case EAccessibleEvent::Notification:
 		{
-#if WINVER > 0x0A00 // Win10
-			// GetProcAddress returns a function type with no args and the compiler doesn't like the cast to one with args, but we know better.
-			// todo: are these checks necessary if the compiler is targeting Win10 already?
-			#pragma warning(suppress: 4191)
-			UiaRaiseNotificationEventFunc NotificationFunc = (UiaRaiseNotificationEventFunc)GetProcAddress(GetModuleHandle(TEXT("Uiautomationcore.dll")), "UiaRaiseNotificationEvent");
+			typedef HRESULT(WINAPI* UiaRaiseNotificationEventFunc)(IRawElementProviderSimple*, NotificationKind, NotificationProcessing, BSTR, BSTR);
+			// Cast to intermediary void* to avoid compiler warning 4191, since GetProcAddress doesn't know function arguments
+			UiaRaiseNotificationEventFunc NotificationFunc = (UiaRaiseNotificationEventFunc)(void*)GetProcAddress(GetModuleHandle(TEXT("Uiautomationcore.dll")), "UiaRaiseNotificationEvent");
 			if (NotificationFunc)
 			{
 				NotificationFunc(&ScopedProvider.Provider, NotificationKind_ActionCompleted, NotificationProcessing_All, SysAllocString(*NewValue.GetValue<FString>()), SysAllocString(TEXT("")));
 			}
-#endif
 			break;
 		}
 		case EAccessibleEvent::BeforeRemoveFromParent:
@@ -212,7 +205,13 @@ void FWindowsUIAManager::OnEventRaised(TSharedRef<IAccessibleWidget> Widget, EAc
 			break;
 		}
 		case EAccessibleEvent::WidgetRemoved:
-			UiaDisconnectProvider(&ScopedProvider.Provider);
+			typedef HRESULT(WINAPI* UiaDisconnectProviderFunc)(IRawElementProviderSimple*);
+			// Cast to intermediary void* to avoid compiler warning 4191, since GetProcAddress doesn't know function arguments
+			UiaDisconnectProviderFunc DisconnectFunc = (UiaDisconnectProviderFunc)(void*)GetProcAddress(GetModuleHandle(TEXT("Uiautomationcore.dll")), "UiaDisconnectProvider");
+			if (DisconnectFunc)
+			{
+				DisconnectFunc(&ScopedProvider.Provider);
+			}
 			break;
 		}
 	}
diff --git a/Engine/Source/Runtime/ApplicationCore/Private/Windows/Accessibility/WindowsUIAWidgetProvider.cpp b/Engine/Source/Runtime/ApplicationCore/Private/Windows/Accessibility/WindowsUIAWidgetProvider.cpp
--- a/Engine/Source/Runtime/ApplicationCore/Private/Windows/Accessibility/WindowsUIAWidgetProvider.cpp
+++ b/Engine/Source/Runtime/ApplicationCore/Private/Windows/Accessibility/WindowsUIAWidgetProvider.cpp
@@ -324,11 +324,6 @@ HRESULT STDCALL FWindowsUIAWidgetProvider::GetPropertyValue(PROPERTYID propertyI
 		// todo: figure out what goes here
 		pRetVal->bstrVal = SysAllocString(*LOCTEXT("Slate", "Slate").ToString());
 		break;
-	//case UIA_FullDescriptionPropertyId:
-	//	pRetVal->vt = VT_BSTR;
-	//	// todo: how is this different from help text?
-	//	pRetVal->bstrVal = SysAllocString(*Widget->GetHelpText());
-	//	break;
 	case UIA_HasKeyboardFocusPropertyId:
 		pRetVal->vt = VT_BOOL;
 		pRetVal->boolVal = Widget->HasFocus() ? VARIANT_TRUE : VARIANT_FALSE;
@@ -403,24 +398,6 @@ HRESULT STDCALL FWindowsUIAWidgetProvider::GetPropertyValue(PROPERTYID propertyI
 		pRetVal->vt = VT_I4;
 		pRetVal->lVal = ::GetCurrentProcessId();
 		break;
-	case UIA_SizePropertyId:
-		pRetVal->vt = VT_R8 | VT_ARRAY;
-		pRetVal->parray = SafeArrayCreateVector(VT_R8, 0, 2);
-		if (pRetVal->parray)
-		{
-			FBox2D Box = Widget->GetBounds();
-			LONG i = 0;
-			float Width = Box.Max.X - Box.Min.X;
-			float Height = Box.Max.Y - Box.Min.Y;
-			bValid &= (SafeArrayPutElement(pRetVal->parray, &i, &Width) != S_OK);
-			i = 1;
-			bValid &= (SafeArrayPutElement(pRetVal->parray, &i, &Height) != S_OK);
-		}
-		else
-		{
-			bValid = false;
-		}
-		break;
 	default:
 		pRetVal->vt = VT_EMPTY;
 		break;
