Hidden windows are also consider minimized state wise

diff --git a/Engine/Source/Runtime/ApplicationCore/Private/Linux/LinuxWindow.cpp b/Engine/Source/Runtime/ApplicationCore/Private/Linux/LinuxWindow.cpp
--- a/Engine/Source/Runtime/ApplicationCore/Private/Linux/LinuxWindow.cpp
+++ b/Engine/Source/Runtime/ApplicationCore/Private/Linux/LinuxWindow.cpp
@@ -718,7 +718,7 @@ bool FLinuxWindow::IsMaximized() const
 /** @return true if the native window is minimized, false otherwise */
 bool FLinuxWindow::IsMinimized() const
 {
-	return SDL_GetWindowFlags(HWnd) & SDL_WINDOW_MINIMIZED;
+	return SDL_GetWindowFlags(HWnd) & SDL_WINDOW_MINIMIZED || SDL_GetWindowFlags(HWnd) & SDL_WINDOW_HIDDEN;
 }
 
 /** @return true if the native window is visible, false otherwise */
