If we reshape the window to either 0 width or 0 height Hide that window over setting it to 0, 0

diff --git a/Engine/Source/Runtime/ApplicationCore/Private/Linux/LinuxWindow.cpp b/Engine/Source/Runtime/ApplicationCore/Private/Linux/LinuxWindow.cpp
--- a/Engine/Source/Runtime/ApplicationCore/Private/Linux/LinuxWindow.cpp
+++ b/Engine/Source/Runtime/ApplicationCore/Private/Linux/LinuxWindow.cpp
@@ -454,7 +454,7 @@ void FLinuxWindow::Restore()
 /** Native window should make itself visible */
 void FLinuxWindow::Show()
 {
-	if ( !bIsVisible )
+	if ( !bIsVisible && !bZeroSize )
 	{
 		bIsVisible = true;
 		SDL_ShowWindow( HWnd );
@@ -580,7 +580,19 @@ void FLinuxWindow::ReshapeWindow( int32 NewX, int32 NewY, int32 NewWidth, int32
 	VirtualWidth  = NewWidth;
 	VirtualHeight = NewHeight;
 
-	if ( LinuxWindow )
+	if (bZeroSize && NewWidth != 0 && NewHeight != 0)
+	{
+		Show();
+		bZeroSize = false;
+	}
+	else if (NewWidth == 0 || NewHeight == 0)
+	{
+		Hide();
+		bZeroSize = true;
+	}
+
+	// Avoid broadcasting we have set a zero size as it will attempt to resize the backbuffer which on some RHI is invalid per the spec (ie. Vulkan)
+	if (LinuxWindow && !bZeroSize)
 	{
 		OwningApplication->GetMessageHandler()->OnSizeChanged(
 			LinuxWindow.ToSharedRef(),
diff --git a/Engine/Source/Runtime/ApplicationCore/Public/Linux/LinuxWindow.h b/Engine/Source/Runtime/ApplicationCore/Public/Linux/LinuxWindow.h
--- a/Engine/Source/Runtime/ApplicationCore/Public/Linux/LinuxWindow.h
+++ b/Engine/Source/Runtime/ApplicationCore/Public/Linux/LinuxWindow.h
@@ -269,6 +269,11 @@ private:
 	/** Whether native properties cache is valid */
 	bool bValidNativePropertiesCache;
 
+	/** Whether the window size was set to zero at some point, if its the case we need to hide our window and will need to
+	 *  restore the window when we set our selfs to a new size that is not zero
+	 */
+	bool bZeroSize;
+
 	/**
 	 * Ratio of pixels to SlateUnits in this window.
 	 * E.g. DPIScale of 2.0 means there is a 2x2 pixel square for every 1x1 SlateUnit.
