IOS Accessibility implementation for VoiceOver support

diff --git a/Engine/Source/Runtime/ApplicationCore/Private/IOS/Accessibility/IOSAccessibilityCache.cpp b/Engine/Source/Runtime/ApplicationCore/Private/IOS/Accessibility/IOSAccessibilityCache.cpp
new file mode 100644
--- /dev/null
+++ b/Engine/Source/Runtime/ApplicationCore/Private/IOS/Accessibility/IOSAccessibilityCache.cpp
@@ -0,0 +1,162 @@
+// Copyright 1998-2019 Epic Games, Inc. All Rights Reserved.
+
+#if WITH_ACCESSIBILITY
+
+#include "IOS/Accessibility/IOSAccessibilityCache.h"
+#include "IOS/Accessibility/IOSAccessibilityElement.h"
+#include "IOS/IOSApplication.h"
+#include "IOS/IOSAppDelegate.h"
+#include "IOS/IOSAsyncTask.h"
+#include "IOS/IOSView.h"
+
+@implementation FIOSAccessibilityCache
+
+- (id)init
+{
+	Cache = [[NSMutableDictionary alloc] init];
+	return self;
+}
+
+-(void)dealloc
+{
+	[Cache release];
+	[super dealloc];
+}
+
+-(FIOSAccessibilityContainer*)GetAccessibilityElement:(AccessibleWidgetId)Id
+{
+	if (Id == IAccessibleWidget::InvalidAccessibleWidgetId)
+	{
+		return nil;
+	}
+
+	FIOSAccessibilityContainer* ExistingElement = [Cache objectForKey:[NSNumber numberWithInt:Id]];
+	if (ExistingElement == nil)
+	{
+		AccessibleWidgetId ParentId = IAccessibleWidget::InvalidAccessibleWidgetId;
+		// All IAccessibleWidget functions must be run on Game Thread
+		[IOSAppDelegate WaitAndRunOnGameThread : [Id, &ParentId]()
+		{
+			TSharedPtr<IAccessibleWidget> Widget = [IOSAppDelegate GetDelegate].IOSApplication->GetAccessibleMessageHandler()->GetAccessibleWidgetFromId(Id);
+			if (Widget.IsValid())
+			{
+				TSharedPtr<IAccessibleWidget> ParentWidget = Widget->GetParent();
+				if (ParentWidget.IsValid())
+				{
+					ParentId = ParentWidget->GetId();
+				}
+			}
+		}];
+
+		ExistingElement = [[[FIOSAccessibilityContainer alloc] initWithId:Id AndParentId:ParentId] autorelease];
+		[Cache setObject:ExistingElement forKey:[NSNumber numberWithInt:Id]];
+	}
+	return ExistingElement;
+}
+
+-(bool)AccessibilityElementExists:(AccessibleWidgetId)Id
+{
+	return [Cache objectForKey:[NSNumber numberWithInt:Id]] != nil;
+}
+
+-(void)RemoveAccessibilityElement:(AccessibleWidgetId)Id
+{
+	[Cache removeObjectForKey:[NSNumber numberWithInt:Id]];
+}
+
+-(void)Clear
+{
+	[Cache removeAllObjects];
+}
+
++(id)AccessibilityElementCache
+{
+	static FIOSAccessibilityCache* Cache = nil;
+	if (Cache == nil)
+	{
+		Cache = [[self alloc] init];
+	}
+	return Cache;
+}
+
+-(void)UpdateAllCachedProperties
+{
+	TArray<AccessibleWidgetId> Ids;
+	for (NSString* Key in Cache)
+	{
+		Ids.Add(Key.intValue);
+	}
+
+	if (Ids.Num() > 0)
+	{
+		// All IAccessibleWidget functions must be run on Game Thread
+		FFunctionGraphTask::CreateAndDispatchWhenReady([Ids]()
+		{
+			for (const AccessibleWidgetId Id : Ids)
+			{
+				TSharedPtr<IAccessibleWidget> Widget = [IOSAppDelegate GetDelegate].IOSApplication->GetAccessibleMessageHandler()->GetAccessibleWidgetFromId(Id);
+				if (Widget.IsValid())
+				{
+					// Children
+					TArray<int32> ChildIds;
+					for (int32 i = 0; i < Widget->GetNumberOfChildren(); ++i)
+					{
+						TSharedPtr<IAccessibleWidget> Child = Widget->GetChildAt(i);
+						if (Child.IsValid())
+						{
+							ChildIds.Add(Child->GetId());
+						}
+					}
+
+					// Bounding rect
+					FBox2D Bounds = Widget->GetBounds();
+					const float Scale = [IOSAppDelegate GetDelegate].IOSView.contentScaleFactor;
+					Bounds.Min /= Scale;
+					Bounds.Max /= Scale;
+
+					// Visibility
+					const bool bIsEnabled = Widget->IsEnabled();
+					const bool bIsVisible = !Widget->IsHidden();
+
+					// All UIKit functions must be run on Main Thread
+					dispatch_async(dispatch_get_main_queue(), ^
+					{
+						for (AccessibleWidgetId ChildId : ChildIds)
+						{
+							// Ensure child widgets are created. This is pretty wasteful although at the moment
+							// we have no good way to detect either when a widget is created outside of the hierarchy
+							// or when am SWidgetSwitcher changes to a child that hasn't been initialized yet.
+							[[FIOSAccessibilityCache AccessibilityElementCache] GetAccessibilityElement:ChildId];
+						}
+							
+						FIOSAccessibilityContainer* Element = [[FIOSAccessibilityCache AccessibilityElementCache] GetAccessibilityElement:Id];
+						Element.ChildIds = ChildIds;
+						if (bIsVisible)
+						{
+							Element.Bounds = CGRectMake(Bounds.Min.X, Bounds.Min.Y, Bounds.Max.X - Bounds.Min.X, Bounds.Max.Y - Bounds.Min.Y);
+						}
+						else
+						{
+							// Rather than removing hidden items from the tree, we set their bounding rect to 0 so they're never selected.
+							Element.Bounds = CGRectMake(0, 0, 0, 0);
+						}
+
+						if (bIsEnabled)
+						{
+							[Element GetLeaf].Traits &= ~UIAccessibilityTraitNotEnabled;
+						}
+						else
+						{
+							[Element GetLeaf].Traits |= UIAccessibilityTraitNotEnabled;
+						}
+						Element.bIsVisible = bIsVisible;
+					});
+				}
+			}
+		}, TStatId(), NULL, ENamedThreads::GameThread);
+	}
+}
+
+@end
+
+#endif
diff --git a/Engine/Source/Runtime/ApplicationCore/Private/IOS/Accessibility/IOSAccessibilityElement.cpp b/Engine/Source/Runtime/ApplicationCore/Private/IOS/Accessibility/IOSAccessibilityElement.cpp
new file mode 100644
--- /dev/null
+++ b/Engine/Source/Runtime/ApplicationCore/Private/IOS/Accessibility/IOSAccessibilityElement.cpp
@@ -0,0 +1,290 @@
+// Copyright 1998-2019 Epic Games, Inc. All Rights Reserved.
+
+#if WITH_ACCESSIBILITY
+
+#include "IOS/Accessibility/IOSAccessibilityElement.h"
+
+#include "Async/TaskGraphInterfaces.h"
+#include "GenericPlatform/GenericAccessibleInterfaces.h"
+#include "IOS/Accessibility/IOSAccessibilityCache.h"
+#include "IOS/IOSApplication.h"
+#include "IOS/IOSAppDelegate.h"
+#include "IOS/IOSView.h"
+
+@implementation FIOSAccessibilityContainer
+
+@synthesize Id;
+@synthesize ChildIds;
+@synthesize Bounds;
+@synthesize bIsVisible;
+
+-(id)initWithId:(AccessibleWidgetId)InId AndParentId:(AccessibleWidgetId)InParentId
+{
+	id UIParent = nil;
+	if (InParentId != IAccessibleWidget::InvalidAccessibleWidgetId)
+	{
+		UIParent = [[FIOSAccessibilityCache AccessibilityElementCache] GetAccessibilityElement:InParentId];
+	}
+	else
+	{
+		UIParent = [IOSAppDelegate GetDelegate].IOSView;
+	}
+
+	if (self = [super initWithAccessibilityContainer:UIParent])
+	{
+		self.Id = InId;
+		self.bIsVisible = true;
+		Leaf = [[FIOSAccessibilityLeaf alloc] initWithParent:self];
+	}
+
+	return self;
+}
+
+-(void)dealloc
+{
+	[Leaf release];
+	[super dealloc];
+}
+
+-(void)SetParent:(AccessibleWidgetId)InParentId
+{
+	if (InParentId != IAccessibleWidget::InvalidAccessibleWidgetId)
+	{
+		self.accessibilityContainer = [[FIOSAccessibilityCache AccessibilityElementCache] GetAccessibilityElement:InParentId];
+	}
+	else
+	{
+		self.accessibilityContainer = [IOSAppDelegate GetDelegate].IOSView;
+	}
+}
+
+-(FIOSAccessibilityLeaf*)GetLeaf
+{
+	return Leaf;
+}
+
+-(BOOL)isAccessibilityElement
+{
+	// Containers are never accessible
+	return NO;
+}
+
+-(CGRect)accessibilityFrame
+{
+	return Bounds;
+}
+
+-(NSInteger)accessibilityElementCount
+{
+	// The extra +1 is from the Leaf element
+	return self.ChildIds.Num() + 1;
+}
+
+-(id)accessibilityElementAtIndex:(NSInteger)index
+{
+	if (index == self.ChildIds.Num())
+	{
+		return Leaf;
+	}
+	else
+	{
+		return [[FIOSAccessibilityCache AccessibilityElementCache] GetAccessibilityElement:self.ChildIds[index]];
+	}
+}
+
+-(NSInteger)indexOfAccessibilityElement:(id)element
+{
+	if ([element isKindOfClass : [FIOSAccessibilityLeaf class]])
+	{
+		// If it's a Leaf, it is the last child of the parent
+		return [[element accessibilityContainer] accessibilityElementCount] - 1;
+	}
+
+	const AccessibleWidgetId OtherId = ((FIOSAccessibilityContainer*)element).Id;
+	for (int32 i = 0; i < self.ChildIds.Num(); ++i)
+	{
+		if (self.Id == OtherId)
+		{
+			return i;
+		}
+	}
+
+	return NSNotFound;
+}
+
+@end
+
+@implementation FIOSAccessibilityLeaf
+
+@synthesize Label;
+@synthesize Hint;
+@synthesize Value;
+@synthesize Traits;
+
+-(id)initWithParent:(FIOSAccessibilityContainer*)Parent
+{
+	if (self = [super initWithAccessibilityContainer : Parent])
+	{
+		const AccessibleWidgetId ParentId = Parent.Id;
+		// All IAccessibleWidget functions must be run on Game Thread
+		FFunctionGraphTask::CreateAndDispatchWhenReady([ParentId]()
+		{
+			TSharedPtr<IAccessibleWidget> Widget = [IOSAppDelegate GetDelegate].IOSApplication->GetAccessibleMessageHandler()->GetAccessibleWidgetFromId(ParentId);
+			if (Widget.IsValid())
+			{
+				// Most accessibility traits cannot be changed after setting, so initialize them here
+				UIAccessibilityTraits InitialTraits = UIAccessibilityTraitNone;
+				if (Widget->AsProperty() && !FMath::IsNearlyZero(Widget->AsProperty()->GetStepSize()))
+				{
+					InitialTraits |= UIAccessibilityTraitAdjustable;
+				}
+				if (Widget->AsActivatable())
+				{
+					InitialTraits |= UIAccessibilityTraitButton;
+				}
+				if (Widget->GetWidgetType() == EAccessibleWidgetType::Image)
+				{
+					InitialTraits |= UIAccessibilityTraitImage;
+				}
+				if (Widget->GetWidgetType() == EAccessibleWidgetType::Hyperlink)
+				{
+					InitialTraits |= UIAccessibilityTraitLink;
+				}
+				if (!Widget->IsEnabled())
+				{
+					InitialTraits |= UIAccessibilityTraitNotEnabled;
+				}
+
+				// All UIKit functions must be run on Main Thread
+				dispatch_async(dispatch_get_main_queue(), ^
+				{
+					FIOSAccessibilityContainer* Element = [[FIOSAccessibilityCache AccessibilityElementCache] GetAccessibilityElement:ParentId];
+					[Element GetLeaf].Traits = InitialTraits;
+				});
+			}
+		}, TStatId(), NULL, ENamedThreads::GameThread);
+	}
+	return self;
+}
+
+-(BOOL)isAccessibilityElement
+{
+	return YES;
+}
+
+-(CGRect)accessibilityFrame
+{
+	return [self.accessibilityContainer accessibilityFrame];
+}
+
+/**
+ * Since Strings are expensive to copy over and over, we don't update them with polling.
+ * When any one of the three are requested, we cache all three in order to restrict it to
+ * a single cross-thread call. This will add a slight delay when the user taps on a widget,
+ * but at the moment it isn't noticeable. The slower the framerate of the game, the more
+ * noticeable this delay may become.
+ */
+-(void)CacheStrings
+{
+	const double StartTime = FPlatformTime::Seconds();
+	if (StartTime - LastCachedStringTime > 1.0)
+	{
+		const AccessibleWidgetId TempId = ((FIOSAccessibilityContainer*)self.accessibilityContainer).Id;
+		FString TempLabel, TempHint, TempValue;
+		// All IAccessibleWidget functions must be run on Game Thread
+		[IOSAppDelegate WaitAndRunOnGameThread : [TempId, &TempLabel, &TempHint, &TempValue]()
+		{
+			TSharedPtr<IAccessibleWidget> Widget = [IOSAppDelegate GetDelegate].IOSApplication->GetAccessibleMessageHandler()->GetAccessibleWidgetFromId(TempId);
+			if (Widget.IsValid())
+			{
+				TempLabel = Widget->GetWidgetName();
+				TempHint = Widget->GetHelpText();
+				if (Widget->AsProperty())
+				{
+					TempValue = Widget->AsProperty()->GetValue();
+				}
+			}
+		}];
+		self.Label = MoveTemp(TempLabel);
+		self.Hint = MoveTemp(TempHint);
+		self.Value = MoveTemp(TempValue);
+		LastCachedStringTime = FPlatformTime::Seconds();
+	}
+}
+
+-(NSString*)accessibilityLabel
+{
+	[self CacheStrings];
+	if (!self.Label.IsEmpty())
+	{
+		return [NSString stringWithFString:self.Label];
+	}
+	return nil;
+}
+
+-(NSString*)accessibilityHint
+{
+	[self CacheStrings];
+	if (!self.Hint.IsEmpty())
+	{
+		return [NSString stringWithFString:self.Hint];
+	}
+	return nil;
+}
+
+-(NSString*)accessibilityValue
+{
+	[self CacheStrings];
+	if (!self.Value.IsEmpty())
+	{
+		return [NSString stringWithFString:self.Value];
+	}
+	return nil;
+}
+
+-(UIAccessibilityTraits)accessibilityTraits
+{
+	return self.Traits;
+}
+
+-(void)accessibilityIncrement
+{
+	const AccessibleWidgetId TempId = ((FIOSAccessibilityContainer*)self.accessibilityContainer).Id;
+	// All IAccessibleWidget functions must be run on Game Thread
+	FFunctionGraphTask::CreateAndDispatchWhenReady([TempId]()
+	{
+		TSharedPtr<IAccessibleWidget> Widget = [IOSAppDelegate GetDelegate].IOSApplication->GetAccessibleMessageHandler()->GetAccessibleWidgetFromId(TempId);
+		if (Widget.IsValid())
+		{
+			IAccessibleProperty* Property = Widget->AsProperty();
+			if (Property && !FMath::IsNearlyZero(Property->GetStepSize()))
+			{
+				const float CurrentValue = FCString::Atof(*Property->GetValue());
+				Property->SetValue(FString::SanitizeFloat(CurrentValue + Property->GetStepSize()));
+			}
+		}
+	}, TStatId(), NULL, ENamedThreads::GameThread);
+}
+
+-(void)accessibilityDecrement
+{
+	const AccessibleWidgetId TempId = ((FIOSAccessibilityContainer*)self.accessibilityContainer).Id;
+	// All IAccessibleWidget functions must be run on Game Thread
+	FFunctionGraphTask::CreateAndDispatchWhenReady([TempId]()
+	{
+		TSharedPtr<IAccessibleWidget> Widget = [IOSAppDelegate GetDelegate].IOSApplication->GetAccessibleMessageHandler()->GetAccessibleWidgetFromId(TempId);
+		if (Widget.IsValid())
+		{
+			IAccessibleProperty* Property = Widget->AsProperty();
+			if (Property && !FMath::IsNearlyZero(Property->GetStepSize()))
+			{
+				const float CurrentValue = FCString::Atof(*Property->GetValue());
+				Property->SetValue(FString::SanitizeFloat(CurrentValue - Property->GetStepSize()));
+			}
+		}
+	}, TStatId(), NULL, ENamedThreads::GameThread);
+}
+
+@end
+
+#endif
diff --git a/Engine/Source/Runtime/ApplicationCore/Private/IOS/IOSAppDelegate.cpp b/Engine/Source/Runtime/ApplicationCore/Private/IOS/IOSAppDelegate.cpp
--- a/Engine/Source/Runtime/ApplicationCore/Private/IOS/IOSAppDelegate.cpp
+++ b/Engine/Source/Runtime/ApplicationCore/Private/IOS/IOSAppDelegate.cpp
@@ -6,6 +6,7 @@
 #include "Modules/Boilerplate/ModuleBoilerplate.h"
 #include "Misc/CallbackDevice.h"
 #include "IOS/IOSView.h"
+#include "IOS/IOSWindow.h"
 #include "Async/TaskGraphInterfaces.h"
 #include "GenericPlatform/GenericPlatformChunkInstall.h"
 #include "IOS/IOSPlatformMisc.h"
@@ -15,6 +16,7 @@
 #include "Misc/OutputDeviceError.h"
 #include "Misc/CommandLine.h"
 #include "IOS/IOSPlatformFramePacer.h"
+#include "IOS/IOSApplication.h"
 #include "IOS/IOSAsyncTask.h"
 #include "Misc/ConfigCacheIni.h"
 #include "IOS/IOSPlatformCrashContext.h"
@@ -33,6 +35,10 @@
 #include <AVFoundation/AVAudioSession.h>
 #include "HAL/IConsoleManager.h"
 
+#if WITH_ACCESSIBILITY
+#include "IOS/Accessibility/IOSAccessibilityCache.h"
+#endif
+
 // this is the size of the game thread stack, it must be a multiple of 4k
 #if (UE_BUILD_SHIPPING || UE_BUILD_TEST)
 #define GAME_THREAD_STACK_SIZE 2 * 1024 * 1024
@@ -151,7 +157,9 @@ bool FIOSCoreDelegates::PassesPushNotificationFilters(NSDictionary* Payload)
 @synthesize timer;
 @synthesize IdleTimerEnableTimer;
 @synthesize IdleTimerEnablePeriod;
-
+#if WITH_ACCESSIBILITY
+@synthesize AccessibilityCacheTimer;
+#endif
 @synthesize savedOpenUrlParameters;
 @synthesize BackgroundSessionEventCompleteDelegate;
 
@@ -212,6 +220,13 @@ static IOSAppDelegate* CachedDelegate = nil;
 	[IOSView release];
 	[SlateController release];
 	[timer release];
+#if WITH_ACCESSIBILITY
+	if (AccessibilityCacheTimer != nil)
+	{
+		[AccessibilityCacheTimer invalidate];
+		[AccessibilityCacheTimer release];
+	}
+#endif
 	[super dealloc];
 }
 
@@ -259,6 +274,16 @@ static IOSAppDelegate* CachedDelegate = nil;
 	}
 	self.savedOpenUrlParameters = nil; // clear after saved openurl delegate running
 
+#if WITH_ACCESSIBILITY
+	// Initialize accessibility code if VoiceOver is enabled. This must happen after Slate has been initialized.
+	dispatch_async(dispatch_get_main_queue(), ^{
+		if (UIAccessibilityIsVoiceOverRunning())
+		{
+			[[IOSAppDelegate GetDelegate] OnVoiceOverStatusChanged];
+		}
+	});
+#endif
+
     while( !GIsRequestingExit )
 	{
         if (self.bIsSuspended)
@@ -1209,10 +1234,55 @@ static FAutoConsoleVariableRef CVarGEnableThermalsReport(
 #endif
     
 	[self InitializeAudioSession];
-	
+
+#if WITH_ACCESSIBILITY
+	[[NSNotificationCenter defaultCenter] addObserver:self selector:@selector(OnVoiceOverStatusChanged) name:UIAccessibilityVoiceOverStatusChanged object:nil];
+#endif
+    
 	return YES;
 }
 
+#if WITH_ACCESSIBILITY
+-(void)OnVoiceOverStatusChanged
+{
+	if (UIAccessibilityIsVoiceOverRunning())
+	{
+		// This must happen asynchronously because when the app activates from a suspended state,
+		// the IOS notification will emit before the game thread wakes up. This does mean that the
+		// accessibility element tree will probably not be 100% completed when the application
+		// opens for the first time. If this is a problem we can add separate branches for startup
+		// vs waking up.
+		FFunctionGraphTask::CreateAndDispatchWhenReady([]()
+		{
+			FIOSApplication* Application = [IOSAppDelegate GetDelegate].IOSApplication;
+			Application->GetAccessibleMessageHandler()->SetActive(true);
+			AccessibleWidgetId WindowId = Application->GetAccessibleMessageHandler()->GetAccessibleWindowId(Application->FindWindowByAppDelegateView());
+			dispatch_async(dispatch_get_main_queue(), ^{
+                IOSAppDelegate* Delegate = [IOSAppDelegate GetDelegate];
+				[Delegate.IOSView SetAccessibilityWindow:WindowId];
+				if (Delegate.AccessibilityCacheTimer == nil)
+				{
+					// Start caching accessibility data so that it can be returned instantly to IOS. If not cached, the data takes too long
+					// to retrieve due to cross-thread waiting and IOS will timeout.
+					Delegate.AccessibilityCacheTimer = [NSTimer scheduledTimerWithTimeInterval:0.1f target:[FIOSAccessibilityCache AccessibilityElementCache] selector:@selector(UpdateAllCachedProperties) userInfo:nil repeats:YES];
+				}
+			});
+		}, TStatId(), NULL, ENamedThreads::GameThread);
+	}
+	else
+	{
+		[AccessibilityCacheTimer invalidate];
+		AccessibilityCacheTimer = nil;
+		[[IOSAppDelegate GetDelegate].IOSView SetAccessibilityWindow : IAccessibleWidget::InvalidAccessibleWidgetId];
+		[[FIOSAccessibilityCache AccessibilityElementCache] Clear];
+		FFunctionGraphTask::CreateAndDispatchWhenReady([]()
+		{
+			[IOSAppDelegate GetDelegate].IOSApplication->GetAccessibleMessageHandler()->SetActive(false);
+		}, TStatId(), NULL, ENamedThreads::GameThread);
+	}
+}
+#endif
+
 - (void) StartGameThread
 {
 	// create a new thread (the pointer will be retained forever)
@@ -1229,6 +1299,23 @@ static FAutoConsoleVariableRef CVarGEnableThermalsReport(
 	}
 }
 
++(bool)WaitAndRunOnGameThread:(TUniqueFunction<void()>)Function
+{
+	FGraphEventRef Task = FFunctionGraphTask::CreateAndDispatchWhenReady(MoveTemp(Function), TStatId(), NULL, ENamedThreads::GameThread);
+
+	const double MaxThreadWaitTime = 2.0;
+	const double StartTime = FPlatformTime::Seconds();
+	while ((FPlatformTime::Seconds() - StartTime) < MaxThreadWaitTime)
+	{
+		FPlatformProcess::Sleep(0.05f);
+		if (Task->IsComplete())
+		{
+			return true;
+		}
+	}
+	return false;
+}
+
 #if !PLATFORM_TVOS
 extern EDeviceScreenOrientation ConvertFromUIInterfaceOrientation(UIInterfaceOrientation Orientation);
 #endif
diff --git a/Engine/Source/Runtime/ApplicationCore/Private/IOS/IOSApplication.cpp b/Engine/Source/Runtime/ApplicationCore/Private/IOS/IOSApplication.cpp
--- a/Engine/Source/Runtime/ApplicationCore/Private/IOS/IOSApplication.cpp
+++ b/Engine/Source/Runtime/ApplicationCore/Private/IOS/IOSApplication.cpp
@@ -13,6 +13,10 @@
 #include "HAL/IConsoleManager.h"
 #include "IOS/IOSAsyncTask.h"
 #include "Stats/Stats.h"
+#if WITH_ACCESSIBILITY
+#include "IOS/Accessibility/IOSAccessibilityCache.h"
+#include "IOS/Accessibility/IOSAccessibilityElement.h"
+#endif
 
 FIOSApplication* FIOSApplication::CreateIOSApplication()
 {
@@ -49,6 +53,15 @@ void FIOSApplication::SetMessageHandler( const TSharedRef< FGenericApplicationMe
 	}
 }
 
+#if WITH_ACCESSIBILITY
+void FIOSApplication::SetAccessibleMessageHandler(const TSharedRef<FGenericAccessibleMessageHandler>& InAccessibleMessageHandler)
+{
+	GenericApplication::SetAccessibleMessageHandler(InAccessibleMessageHandler);
+	InAccessibleMessageHandler->SetAccessibleEventDelegate(FGenericAccessibleMessageHandler::FAccessibleEvent::CreateRaw(this, &FIOSApplication::OnAccessibleEventRaised));
+	InAccessibleMessageHandler->SetActive(UIAccessibilityIsVoiceOverRunning());
+}
+#endif
+
 void FIOSApplication::AddExternalInputDevice(TSharedPtr<IInputDevice> InputDevice)
 {
 	if (InputDevice.IsValid())
@@ -224,3 +237,53 @@ bool FIOSApplication::IsGamepadAttached() const
 
 	return false;
 }
+
+TSharedRef<FIOSWindow> FIOSApplication::FindWindowByAppDelegateView()
+{
+	for (const TSharedRef<FIOSWindow>& Window : Windows)
+	{
+		if ([IOSAppDelegate GetDelegate].Window == static_cast<UIWindow*>(Window->GetOSWindowHandle()))
+		{
+			return Window;
+		}
+	}
+    check(false);
+    return Windows[0];
+}
+
+#if WITH_ACCESSIBILITY
+void FIOSApplication::OnAccessibleEventRaised(TSharedRef<IAccessibleWidget> Widget, EAccessibleEvent Event, FVariant OldValue, FVariant NewValue)
+{
+	// This should only be triggered by the accessible message handler which initiates from the Slate thread.
+	check(IsInGameThread());
+
+	const AccessibleWidgetId Id = Widget->GetId();
+	switch (Event)
+	{
+	case EAccessibleEvent::BeforeRemoveFromParent:
+		dispatch_async(dispatch_get_main_queue(), ^{
+			[[[FIOSAccessibilityCache AccessibilityElementCache] GetAccessibilityElement:Id] SetParent:IAccessibleWidget::InvalidAccessibleWidgetId];
+		});
+		break;
+	case EAccessibleEvent::AfterAddToParent:
+	{
+		AccessibleWidgetId ParentId = IAccessibleWidget::InvalidAccessibleWidgetId;
+		TSharedPtr<IAccessibleWidget> ParentWidget = Widget->GetParent();
+		if (ParentWidget.IsValid())
+		{
+			ParentId = ParentWidget->GetId();
+		}
+		dispatch_async(dispatch_get_main_queue(), ^{
+			[[[FIOSAccessibilityCache AccessibilityElementCache] GetAccessibilityElement:Id] SetParent:ParentId];
+		});
+		UIAccessibilityPostNotification(UIAccessibilityLayoutChangedNotification, nil);
+		break;
+	}
+	case EAccessibleEvent::WidgetRemoved:
+		dispatch_async(dispatch_get_main_queue(), ^{
+			[[FIOSAccessibilityCache AccessibilityElementCache] RemoveAccessibilityElement:Id];
+		});
+		break;
+	}
+}
+#endif
diff --git a/Engine/Source/Runtime/ApplicationCore/Private/IOS/IOSView.cpp b/Engine/Source/Runtime/ApplicationCore/Private/IOS/IOSView.cpp
--- a/Engine/Source/Runtime/ApplicationCore/Private/IOS/IOSView.cpp
+++ b/Engine/Source/Runtime/ApplicationCore/Private/IOS/IOSView.cpp
@@ -3,6 +3,7 @@
 #include "IOS/IOSView.h"
 #include "IOS/IOSAppDelegate.h"
 #include "IOS/IOSApplication.h"
+#include "IOSWindow.h"
 #include "Misc/ConfigCacheIni.h"
 #include "HAL/IConsoleManager.h"
 #include "Misc/CommandLine.h"
@@ -16,6 +17,11 @@
 
 #include "IOS/IOSCommandLineHelper.h"
 
+#if WITH_ACCESSIBILITY
+#include "IOS/Accessibility/IOSAccessibilityCache.h"
+#include "IOS/Accessibility/IOSAccessibilityElement.h"
+#endif
+
 #if HAS_METAL
 id<MTLDevice> GMetalDevice = nil;
 #endif
@@ -286,6 +292,9 @@ id<MTLDevice> GMetalDevice = nil;
 -(void)dealloc
 {
 	[markedTextStyle release];
+#if WITH_ACCESSIBILITY
+	[_accessibilityElements release];
+#endif
 	[super dealloc];
 }
 
@@ -538,6 +547,37 @@ id<MTLDevice> GMetalDevice = nil;
 	SwapCount++;
 }
 
+#if WITH_ACCESSIBILITY
+
+-(void)SetAccessibilityWindow:(AccessibleWidgetId)WindowId
+{
+	if (_accessibilityElements == nil)
+	{
+		_accessibilityElements = [[NSMutableArray alloc] init];
+	}
+	else
+	{
+		[_accessibilityElements removeAllObjects];
+	}
+
+	if (WindowId != IAccessibleWidget::InvalidAccessibleWidgetId)
+	{
+		[_accessibilityElements addObject : [[FIOSAccessibilityCache AccessibilityElementCache] GetAccessibilityElement:WindowId]];
+	}
+}
+
+-(NSArray*)accessibilityElements
+{
+	return _accessibilityElements;
+}
+
+-(BOOL)isAccessibilityElement
+{
+	return NO;
+}
+
+#endif
+
 /**
  * Returns the unique ID for the given touch
  */
diff --git a/Engine/Source/Runtime/ApplicationCore/Private/Windows/Accessibility/WindowsUIAManager.cpp b/Engine/Source/Runtime/ApplicationCore/Private/Windows/Accessibility/WindowsUIAManager.cpp
--- a/Engine/Source/Runtime/ApplicationCore/Private/Windows/Accessibility/WindowsUIAManager.cpp
+++ b/Engine/Source/Runtime/ApplicationCore/Private/Windows/Accessibility/WindowsUIAManager.cpp
@@ -110,31 +110,9 @@ FWindowsUIAWindowProvider& FWindowsUIAManager::GetWindowProvider(TSharedRef<FWin
 	return static_cast<FWindowsUIAWindowProvider&>(GetWidgetProvider(WindowsApplication.GetAccessibleMessageHandler()->GetAccessibleWindow(InWindow).ToSharedRef()));
 }
 
-int FWindowsUIAManager::GetRuntimeId(FWindowsUIAWidgetProvider* InProvider)
-{
-	const int* Value = RuntimeIds.Find(InProvider);
-	if (Value)
-	{
-		return *Value;
-	}
-	else
-	{
-		static int RuntimeIdCounter = 0;
-		if (RuntimeIdCounter == INT_MAX)
-		{
-			RuntimeIdCounter = INT_MIN;
-		}
-		return RuntimeIds.Add(InProvider, RuntimeIdCounter++);
-	}
-}
-
 void FWindowsUIAManager::OnWidgetProviderRemoved(TSharedRef<IAccessibleWidget> InWidget)
 {
-	FWindowsUIAWidgetProvider* Provider = nullptr;
-	if (CachedWidgetProviders.RemoveAndCopyValue(InWidget, Provider))
-	{
-		RuntimeIds.Remove(Provider);
-	}
+	CachedWidgetProviders.Remove(InWidget);
 
 	if (CachedWidgetProviders.Num() == 0)
 	{
@@ -182,18 +160,22 @@ void FWindowsUIAManager::OnEventRaised(TSharedRef<IAccessibleWidget> Widget, EAc
 		case EAccessibleEvent::Notification:
 		{
 			typedef HRESULT(WINAPI* UiaRaiseNotificationEventFunc)(IRawElementProviderSimple*, NotificationKind, NotificationProcessing, BSTR, BSTR);
-			// Cast to intermediary void* to avoid compiler warning 4191, since GetProcAddress doesn't know function arguments
-			UiaRaiseNotificationEventFunc NotificationFunc = (UiaRaiseNotificationEventFunc)(void*)GetProcAddress(GetModuleHandle(TEXT("Uiautomationcore.dll")), "UiaRaiseNotificationEvent");
-			if (NotificationFunc)
+			HMODULE Handle = GetModuleHandle(TEXT("Uiautomationcore.dll"));
+			if (Handle)
 			{
-				NotificationFunc(&ScopedProvider.Provider, NotificationKind_ActionCompleted, NotificationProcessing_All, SysAllocString(*NewValue.GetValue<FString>()), SysAllocString(TEXT("")));
+				// Cast to intermediary void* to avoid compiler warning 4191, since GetProcAddress doesn't know function arguments
+				UiaRaiseNotificationEventFunc NotificationFunc = (UiaRaiseNotificationEventFunc)(void*)GetProcAddress(Handle, "UiaRaiseNotificationEvent");
+				if (NotificationFunc)
+				{
+					NotificationFunc(&ScopedProvider.Provider, NotificationKind_ActionCompleted, NotificationProcessing_All, SysAllocString(*NewValue.GetValue<FString>()), SysAllocString(TEXT("")));
+				}
 			}
 			break;
 		}
 		case EAccessibleEvent::BeforeRemoveFromParent:
 		{
 			// ChildRemoved events require an ID passed along with them to identify which child (which no longer exists) was removed
-			int id[2] = { UiaAppendRuntimeId, GetRuntimeId(&ScopedProvider.Provider) };
+			int id[2] = { UiaAppendRuntimeId, Widget->GetId() };
 			FScopedWidgetProvider ParentProvider(GetWidgetProvider(Widget->GetParent().ToSharedRef()));
 			UiaRaiseStructureChangedEvent(&ParentProvider.Provider, StructureChangeType_ChildRemoved, id, ARRAYSIZE(id));
 			break;
@@ -206,11 +188,15 @@ void FWindowsUIAManager::OnEventRaised(TSharedRef<IAccessibleWidget> Widget, EAc
 		}
 		case EAccessibleEvent::WidgetRemoved:
 			typedef HRESULT(WINAPI* UiaDisconnectProviderFunc)(IRawElementProviderSimple*);
-			// Cast to intermediary void* to avoid compiler warning 4191, since GetProcAddress doesn't know function arguments
-			UiaDisconnectProviderFunc DisconnectFunc = (UiaDisconnectProviderFunc)(void*)GetProcAddress(GetModuleHandle(TEXT("Uiautomationcore.dll")), "UiaDisconnectProvider");
-			if (DisconnectFunc)
+			HMODULE Handle = GetModuleHandle(TEXT("Uiautomationcore.dll"));
+			if (Handle)
 			{
-				DisconnectFunc(&ScopedProvider.Provider);
+				// Cast to intermediary void* to avoid compiler warning 4191, since GetProcAddress doesn't know function arguments
+				UiaDisconnectProviderFunc DisconnectFunc = (UiaDisconnectProviderFunc)(void*)GetProcAddress(Handle, "UiaDisconnectProvider");
+				if (DisconnectFunc)
+				{
+					DisconnectFunc(&ScopedProvider.Provider);
+				}
 			}
 			break;
 		}
diff --git a/Engine/Source/Runtime/ApplicationCore/Private/Windows/Accessibility/WindowsUIAPropertyGetters.cpp b/Engine/Source/Runtime/ApplicationCore/Private/Windows/Accessibility/WindowsUIAPropertyGetters.cpp
--- a/Engine/Source/Runtime/ApplicationCore/Private/Windows/Accessibility/WindowsUIAPropertyGetters.cpp
+++ b/Engine/Source/Runtime/ApplicationCore/Private/Windows/Accessibility/WindowsUIAPropertyGetters.cpp
@@ -21,7 +21,7 @@ FVariant GetPropertyValue(TSharedRef<IAccessibleWidget> AccessibleWidget, PROPER
 	{
 	case UIA_ValueIsReadOnlyPropertyId:
 	case UIA_RangeValueIsReadOnlyPropertyId:
-		return !AccessibleWidget->AsProperty()->IsReadOnly();
+		return AccessibleWidget->AsProperty()->IsReadOnly();
 	case UIA_RangeValueLargeChangePropertyId:
 	case UIA_RangeValueSmallChangePropertyId:
 		return static_cast<double>(AccessibleWidget->AsProperty()->GetStepSize());
diff --git a/Engine/Source/Runtime/ApplicationCore/Private/Windows/Accessibility/WindowsUIAWidgetProvider.cpp b/Engine/Source/Runtime/ApplicationCore/Private/Windows/Accessibility/WindowsUIAWidgetProvider.cpp
--- a/Engine/Source/Runtime/ApplicationCore/Private/Windows/Accessibility/WindowsUIAWidgetProvider.cpp
+++ b/Engine/Source/Runtime/ApplicationCore/Private/Windows/Accessibility/WindowsUIAWidgetProvider.cpp
@@ -472,7 +472,7 @@ HRESULT STDCALL FWindowsUIAWidgetProvider::GetRuntimeId(SAFEARRAY** pRetVal)
 {
 	if (IsValid())
 	{
-		int rtId[] = { UiaAppendRuntimeId, UIAManager->GetRuntimeId(this) };
+		int rtId[] = { UiaAppendRuntimeId, Widget->GetId() };
 		*pRetVal = SafeArrayCreateVector(VT_I4, 0, 2);
 		if (*pRetVal)
 		{
diff --git a/Engine/Source/Runtime/ApplicationCore/Public/GenericPlatform/GenericAccessibleInterfaces.h b/Engine/Source/Runtime/ApplicationCore/Public/GenericPlatform/GenericAccessibleInterfaces.h
--- a/Engine/Source/Runtime/ApplicationCore/Public/GenericPlatform/GenericAccessibleInterfaces.h
+++ b/Engine/Source/Runtime/ApplicationCore/Public/GenericPlatform/GenericAccessibleInterfaces.h
@@ -229,6 +229,8 @@ public:
 	virtual const FString& GetText() const = 0;
 };
 
+typedef int32 AccessibleWidgetId; 
+
 /**
  * Provides the core set of accessible data that is necessary in order for widget traversal and TTS to be implemented.
  * In order to support functionality beyond this, subclasses must implement the other accessible interfaces and
@@ -240,12 +242,22 @@ public:
 	IAccessibleWidget() {}
 	virtual ~IAccessibleWidget() {}
 
+	static const AccessibleWidgetId InvalidAccessibleWidgetId = -1;
+
+	/**
+	 * Get an application-unique identifier for this widget. If the widget is destroyed,
+	 * a different widget is allowed to re-use that ID.
+	 *
+	 * @return A unique ID that specifically refers to this widget.
+	 */
+	virtual AccessibleWidgetId GetId() const = 0;
 	/**
 	 * Whether or not the underlying widget backing this interface still exists
 	 *
 	 * @return true if functions can be called on this interface and should return valid results
 	 */
 	virtual bool IsValid() const = 0;
+
 	/**
 	 * Returns the window at the top of this widget's hierarchy. This function may return itself for accessible windows,
 	 * and could return nullptr in cases where the widget is not currently part of a hierarchy.
@@ -456,6 +468,23 @@ public:
 	 */
 	virtual TSharedPtr<IAccessibleWidget> GetAccessibleWindow(const TSharedRef<FGenericWindow>& InWindow) const { return nullptr; }
 
+	/**
+	 * Creates or retrieves the identifier for an accessible object for a native OS window.
+	 * todo: Behavior for non-native windows (virtual or others) is currently undefined.
+	 *
+	 * @param InWindow The native window to find the accessible window for
+	 * @return The identifier for the accessible window created
+	 */
+	virtual AccessibleWidgetId GetAccessibleWindowId(const TSharedRef<FGenericWindow>& InWindow) const { return IAccessibleWidget::InvalidAccessibleWidgetId; }
+
+	/**
+	 * Retrieves an accessible widget that matches the given identifier.
+	 *
+	 * @param Id The identifier for the widget to get.
+	 * @return The widget that matches this identifier, or nullptr if the widget does not exist.
+	 */
+	virtual TSharedPtr<IAccessibleWidget> GetAccessibleWidgetFromId(AccessibleWidgetId Id) const { return nullptr; }
+
 	/**
 	 * Push an event from an accessible widget back to the platform layer.
 	 *
diff --git a/Engine/Source/Runtime/ApplicationCore/Public/IOS/Accessibility/IOSAccessibilityCache.h b/Engine/Source/Runtime/ApplicationCore/Public/IOS/Accessibility/IOSAccessibilityCache.h
new file mode 100644
--- /dev/null
+++ b/Engine/Source/Runtime/ApplicationCore/Public/IOS/Accessibility/IOSAccessibilityCache.h
@@ -0,0 +1,44 @@
+// Copyright 1998-2019 Epic Games, Inc. All Rights Reserved.
+
+#pragma once
+
+#if WITH_ACCESSIBILITY
+
+#include "GenericPlatform/GenericAccessibleInterfaces.h"
+
+@class FIOSAccessibilityContainer;
+
+/**
+ * This class is a singleton and should be accessed through [FIOSAccessibilityCache AccessibilityElementCache].
+ * Stores a list of accessible containers that map to AccessibleWidgetIds for lookup.
+ * The cache is also responsible for polling attributes from the underlying
+ * IAccessibleWidgets that are too expensive to be done when requested by IOS due to
+ * needing to be accessed from a different thread.
+ *
+ * Leaf elements can be accessed by getting their container from the cache and calling
+ * [container GetLeaf] on it.
+ */
+@interface FIOSAccessibilityCache : NSObject
+{
+@private
+	/** AccessibleWidgetId(String)->FIOSAccessibilityContainer map for all created containers. */
+	NSMutableDictionary* Cache;
+}
+
+/** Retrieve a cached container, or create one if it doesn't exist yet. */
+-(FIOSAccessibilityContainer*)GetAccessibilityElement:(AccessibleWidgetId)Id;
+/** Returns true if the Cache contains the Id. Does not create one if it doesn't exist. */
+-(bool)AccessibilityElementExists:(AccessibleWidgetId)Id;
+/** Removes an entry from the Cache. */
+-(void)RemoveAccessibilityElement:(AccessibleWidgetId)Id;
+/** Completely empties the cache. */
+-(void)Clear;
+/** Loop over all cached elements and update any properties necessary on the Game thread. */
+-(void)UpdateAllCachedProperties;
+
+/** Singleton accessor */
++(id)AccessibilityElementCache;
+
+@end
+
+#endif
diff --git a/Engine/Source/Runtime/ApplicationCore/Public/IOS/Accessibility/IOSAccessibilityElement.h b/Engine/Source/Runtime/ApplicationCore/Public/IOS/Accessibility/IOSAccessibilityElement.h
new file mode 100644
--- /dev/null
+++ b/Engine/Source/Runtime/ApplicationCore/Public/IOS/Accessibility/IOSAccessibilityElement.h
@@ -0,0 +1,75 @@
+// Copyright 1998-2019 Epic Games, Inc. All Rights Reserved.
+
+#pragma once
+
+#if WITH_ACCESSIBILITY
+
+#include "CoreMinimal.h"
+#include "GenericPlatform/GenericAccessibleInterfaces.h"
+#import <UIKit/UIKit.h>
+
+@class FIOSAccessibilityLeaf;
+
+/**
+ * UIAccessibilityElements cannot be both accessible and have children. While
+ * the same class can be used in both cases, the value they return for
+ * isAccessibilityElement determines what type of widget they are. If false,
+ * the widget is a container and only functions regarding children will be called.
+ * If true, only functions regarding the value of the wigdet will be called.
+ *
+ * Because of this, all IAccessibleWidgets have both a corresponding container
+ * and leaf. The leaf is always reported as the last child of the container. This
+ * is our workaround for a widget being accessible and having children at the same time.
+ */
+@interface FIOSAccessibilityContainer : UIAccessibilityElement
+{
+@private
+	/** A matching Leaf element that shares the same AccessibleWidgetId at this container. */
+	FIOSAccessibilityLeaf* Leaf;
+}
+
+/** This must be used instead of initWithAccessibilityContainer in order to work properly. */
+-(id)initWithId:(AccessibleWidgetId)InId AndParentId:(AccessibleWidgetId)InParentId;
+/** Updates the accessibilityContainer property from UIAccessibilityElement. */
+-(void)SetParent:(AccessibleWidgetId)InParentId;
+/** Get the accessible version of this widget. */
+-(FIOSAccessibilityLeaf*)GetLeaf;
+
+/** The identifier used to access this widget through the accessible API. */
+@property (nonatomic) AccessibleWidgetId Id;
+/** A list of identifiers for all current children of this container. */
+@property (nonatomic) TArray<AccessibleWidgetId> ChildIds;
+/** The bounding rect of the container. */
+@property (nonatomic) CGRect Bounds;
+/** Whether or not the widget is currently visible. */
+@property (nonatomic) bool bIsVisible;
+
+@end
+
+/**
+ * The accessible version of a widget for a given AccessibleWidgetId. A Leaf is
+ * guaranteed to have an FIOSAccessibilityContainer as its container, and can be
+ * accessed with [self accessibilityContainer] (in order to get things like bounds).
+ */
+@interface FIOSAccessibilityLeaf : UIAccessibilityElement
+{
+@private
+	/** Timestamp for when Label, Hint, and Value were last cached. */
+	double LastCachedStringTime;
+}
+
+/** This must be used instead of initWithAccessibilityContainer in order to work properly. */
+-(id)initWithParent:(FIOSAccessibilityContainer*)Parent;
+
+/** A cached version of the name of the widget. */
+@property (nonatomic) FString Label;
+/** A cached version of the help text of the widget. */
+@property (nonatomic) FString Hint;
+/** A cached version of the value of property widgets. */
+@property (nonatomic) FString Value;
+/** Bitflag of traits that describe the widget. Most are set once on initialization. */
+@property (nonatomic) UIAccessibilityTraits Traits;
+
+@end
+
+#endif
diff --git a/Engine/Source/Runtime/ApplicationCore/Public/IOS/IOSAppDelegate.h b/Engine/Source/Runtime/ApplicationCore/Public/IOS/IOSAppDelegate.h
--- a/Engine/Source/Runtime/ApplicationCore/Public/IOS/IOSAppDelegate.h
+++ b/Engine/Source/Runtime/ApplicationCore/Public/IOS/IOSAppDelegate.h
@@ -129,6 +129,13 @@ APPLICATIONCORE_API
 /** The time delay (in seconds) between idle timer enable requests and actually enabling the idle timer */
 @property (readonly) float IdleTimerEnablePeriod;
 
+#if WITH_ACCESSIBILITY
+/** Timer used for updating cached data from game thread for all accessible widgets. */
+@property (nonatomic, retain) NSTimer* AccessibilityCacheTimer;
+/** Callback for IOS notification of VoiceOver being enabled or disabled. */
+-(void)OnVoiceOverStatusChanged;
+#endif
+
 // parameters passed from openURL
 @property (nonatomic, retain) NSMutableArray* savedOpenUrlParameters;
 
@@ -166,6 +173,8 @@ APPLICATIONCORE_API
 -(bool)IsIdleTimerEnabled;
 -(void)EnableIdleTimer:(bool)bEnable;
 -(void)StartGameThread;
+/** Uses the TaskGraph to execute a function on the game thread, and then blocks until the function is executed. */
++(bool)WaitAndRunOnGameThread:(TUniqueFunction<void()>)Function;
 -(void)NoUrlCommandLine;
 
 -(int)GetAudioVolume;
diff --git a/Engine/Source/Runtime/ApplicationCore/Public/IOS/IOSApplication.h b/Engine/Source/Runtime/ApplicationCore/Public/IOS/IOSApplication.h
--- a/Engine/Source/Runtime/ApplicationCore/Public/IOS/IOSApplication.h
+++ b/Engine/Source/Runtime/ApplicationCore/Public/IOS/IOSApplication.h
@@ -18,6 +18,9 @@ public:
 	virtual ~FIOSApplication() {}
 
 	void SetMessageHandler( const TSharedRef< FGenericApplicationMessageHandler >& InMessageHandler );
+#if WITH_ACCESSIBILITY
+	virtual void SetAccessibleMessageHandler(const TSharedRef<FGenericAccessibleMessageHandler>& InAccessibleMessageHandler) override;
+#endif
 
 	virtual void PollGameDeviceState( const float TimeDelta ) override;
 
@@ -35,13 +38,17 @@ public:
 
 	virtual bool IsGamepadAttached() const override;
 
+	TSharedRef<FIOSWindow> FindWindowByAppDelegateView();
+
 protected:
 	virtual void InitializeWindow( const TSharedRef< FGenericWindow >& Window, const TSharedRef< FGenericWindowDefinition >& InDefinition, const TSharedPtr< FGenericWindow >& InParent, const bool bShowImmediately ) override;
 
 private:
 
 	FIOSApplication();
-
+#if WITH_ACCESSIBILITY
+	void OnAccessibleEventRaised(TSharedRef<IAccessibleWidget> Widget, EAccessibleEvent Event, FVariant OldValue, FVariant NewValue);
+#endif
 
 private:
 
diff --git a/Engine/Source/Runtime/ApplicationCore/Public/IOS/IOSView.h b/Engine/Source/Runtime/ApplicationCore/Public/IOS/IOSView.h
--- a/Engine/Source/Runtime/ApplicationCore/Public/IOS/IOSView.h
+++ b/Engine/Source/Runtime/ApplicationCore/Public/IOS/IOSView.h
@@ -14,6 +14,10 @@
 #import <QuartzCore/CAMetalLayer.h>
 #endif
 
+#if WITH_ACCESSIBILITY
+#include "GenericPlatform/GenericAccessibleInterfaces.h"
+#endif
+
 struct FKeyboardConfig
 {
 	UIKeyboardType KeyboardType;
@@ -82,9 +86,17 @@ APPLICATIONCORE_API
 	
 	volatile int32 KeyboardShowCount;
 
-
+#if WITH_ACCESSIBILITY
+@private
+	// Single-element array used for the accessibilityElements property. It holds the IAccessibleWindow for the app.
+	NSMutableArray* _accessibilityElements;
+#endif
 }
 
+#if WITH_ACCESSIBILITY
+/** Repopulate _accessibilityElements when the accessible window's ID has changed. */
+-(void)SetAccessibilityWindow:(AccessibleWidgetId)WindowId;
+#endif
 
 //// SHARED FUNCTIONALITY
 @property (nonatomic) GLuint SwapCount;
diff --git a/Engine/Source/Runtime/ApplicationCore/Public/Windows/Accessibility/WindowsUIAManager.h b/Engine/Source/Runtime/ApplicationCore/Public/Windows/Accessibility/WindowsUIAManager.h
--- a/Engine/Source/Runtime/ApplicationCore/Public/Windows/Accessibility/WindowsUIAManager.h
+++ b/Engine/Source/Runtime/ApplicationCore/Public/Windows/Accessibility/WindowsUIAManager.h
@@ -53,16 +53,6 @@ public:
 	 */
 	FWindowsUIAWindowProvider& GetWindowProvider(TSharedRef<FWindowsWindow> InWindow);
 
-	/**
-	 * Generate an application-unique ID for this Provider. A widget must be able to be uniquely
-	 * identified in the case where it gets deleted, but only within the context of a specific hwnd.
-	 * It is OK for widgets to overlap IDs in separate running instances of the same application.
-	 * This returns an int because it's passed directly to Windows as one.
-	 *
-	 * @return An application-unique identifier for this Provider
-	 */
-	int GetRuntimeId(FWindowsUIAWidgetProvider* InProvider);
-
 	/**
 	 * Notify the Manager that the RefCount for a Provider has reached 0, which will cause it to be removed from the cache.
 	 *
@@ -87,7 +77,6 @@ private:
 
 	/** Cache of all Providers with a RefCount of at least 1 that map to an accessible widget. */
 	TMap<TSharedRef<IAccessibleWidget>, FWindowsUIAWidgetProvider*> CachedWidgetProviders;
-	TMap<FWindowsUIAWidgetProvider*, int> RuntimeIds;
 	/**
 	 * A set of all Providers with a RefCount of at least 1. This also includes non-widget Providers such as text ranges.
 	 * When the application is closed, this is used to notify Providers held by external applications that they are invalid.
diff --git a/Engine/Source/Runtime/SlateCore/Private/Widgets/Accessibility/SlateAccessibleMessageHandler.cpp b/Engine/Source/Runtime/SlateCore/Private/Widgets/Accessibility/SlateAccessibleMessageHandler.cpp
--- a/Engine/Source/Runtime/SlateCore/Private/Widgets/Accessibility/SlateAccessibleMessageHandler.cpp
+++ b/Engine/Source/Runtime/SlateCore/Private/Widgets/Accessibility/SlateAccessibleMessageHandler.cpp
@@ -11,12 +11,13 @@
 #include "Input/HittestGrid.h"
 
 DECLARE_CYCLE_STAT(TEXT("Slate Accessibility: Parent Updated"), STAT_AccessibilitySlateParentUpdated, STATGROUP_Accessibility);
+DECLARE_CYCLE_STAT(TEXT("Slate Accessibility: Children Updated"), STAT_AccessibilitySlateChildrenUpdated, STATGROUP_Accessibility);
 DECLARE_CYCLE_STAT(TEXT("Slate Accessibility: Behavior Changed"), STAT_AccessibilitySlateBehaviorChanged, STATGROUP_Accessibility);
 DECLARE_CYCLE_STAT(TEXT("Slate Accessibility: Event Raised"), STAT_AccessibilitySlateEventRaised, STATGROUP_Accessibility);
 
 void FSlateAccessibleMessageHandler::OnActivate()
 {
-	// widgets are initialized when their accessibility window is created
+	// widgets are initialized when their accessible window is created
 }
 
 void FSlateAccessibleMessageHandler::OnDeactivate()
@@ -37,6 +38,21 @@ TSharedPtr<IAccessibleWidget> FSlateAccessibleMessageHandler::GetAccessibleWindo
 	return nullptr;
 }
 
+AccessibleWidgetId FSlateAccessibleMessageHandler::GetAccessibleWindowId(const TSharedRef<FGenericWindow>& InWindow) const
+{
+	TSharedPtr<IAccessibleWidget> AccessibleWindow = GetAccessibleWindow(InWindow);
+	if (AccessibleWindow.IsValid())
+	{
+		return AccessibleWindow->GetId();
+	}
+	return IAccessibleWidget::InvalidAccessibleWidgetId;
+}
+
+TSharedPtr<IAccessibleWidget> FSlateAccessibleMessageHandler::GetAccessibleWidgetFromId(AccessibleWidgetId Id) const
+{
+	return FSlateAccessibleWidgetCache::Get().GetAccessibleWidgetFromId(Id);
+}
+
 void FSlateAccessibleMessageHandler::OnWidgetRemoved(SWidget* Widget)
 {
 	if (IsActive())
@@ -84,6 +100,25 @@ void FSlateAccessibleMessageHandler::OnWidgetParentChanged(TSharedRef<SWidget> W
 	}
 }
 
+void FSlateAccessibleMessageHandler::OnWidgetChildrenChanged(TSharedRef<SWidget> Widget)
+{
+	if (IsActive())
+	{
+		SCOPE_CYCLE_COUNTER(STAT_AccessibilitySlateChildrenUpdated);
+
+		TSharedPtr<SWidget> Parent = Widget;
+		while (Parent.IsValid() && !Parent->IsAccessible())
+		{
+			Parent = Parent->GetParentWidget();
+		}
+		if (Parent.IsValid())
+		{
+			TSharedPtr<IAccessibleWidget> AccessibleParent = FSlateAccessibleWidgetCache::Get().GetAccessibleWidget(Parent);
+			StaticCastSharedPtr<FSlateAccessibleWidget>(AccessibleParent)->MarkChildrenDirty();
+		}
+	}
+}
+
 void FSlateAccessibleMessageHandler::OnWidgetAccessibleBehaviorChanged(TSharedRef<SWidget> Widget)
 {
 	if (IsActive())
@@ -138,16 +173,6 @@ void FSlateAccessibleMessageHandler::OnWidgetEventRaised(TSharedRef<SWidget> Wid
 		// todo: not sure what to do for a case like focus changed to not-accessible widget. maybe pass through a nullptr?
 		if (Widget->IsAccessible())
 		{
-			TSharedPtr<SWidget> Parent = Widget->GetParentWidget();
-			while (Parent.IsValid())
-			{
-				if (!Parent->CanChildrenBeAccessible())
-				{
-					return;
-				}
-				Parent = Parent->GetParentWidget();
-			}
-
 			FSlateAccessibleMessageHandler::RaiseEvent(FSlateAccessibleWidgetCache::Get().GetAccessibleWidget(Widget).ToSharedRef(), Event, OldValue, NewValue);
 		}
 	}
diff --git a/Engine/Source/Runtime/SlateCore/Private/Widgets/Accessibility/SlateCoreAccessibleWidgets.cpp b/Engine/Source/Runtime/SlateCore/Private/Widgets/Accessibility/SlateCoreAccessibleWidgets.cpp
--- a/Engine/Source/Runtime/SlateCore/Private/Widgets/Accessibility/SlateCoreAccessibleWidgets.cpp
+++ b/Engine/Source/Runtime/SlateCore/Private/Widgets/Accessibility/SlateCoreAccessibleWidgets.cpp
@@ -12,6 +12,7 @@
 #include "Input/HittestGrid.h"
 #include "Application/SlateApplicationBase.h"
 #include "Application/SlateWindowHelper.h"
+#include "Math/NumericLimits.h"
 
 DECLARE_CYCLE_STAT(TEXT("Slate Accessibility: Get Widget At Point"), STAT_AccessibilitySlateGetChildAtPosition, STATGROUP_Accessibility);
 
@@ -21,12 +22,27 @@ FSlateAccessibleWidget::FSlateAccessibleWidget(TWeakPtr<SWidget> InWidget, EAcce
 	, SiblingIndex(INDEX_NONE)
 	, bChildrenDirty(true)
 {
+	static AccessibleWidgetId RuntimeIdCounter = 0;
+	if (RuntimeIdCounter == TNumericLimits<AccessibleWidgetId>::Max())
+	{
+		RuntimeIdCounter = TNumericLimits<AccessibleWidgetId>::Min();
+	}
+	if (RuntimeIdCounter == InvalidAccessibleWidgetId)
+	{
+		++RuntimeIdCounter;
+	}
+	Id = RuntimeIdCounter++;
 }
 
 FSlateAccessibleWidget::~FSlateAccessibleWidget()
 {
 }
 
+AccessibleWidgetId FSlateAccessibleWidget::GetId() const
+{
+	return Id;
+}
+
 bool FSlateAccessibleWidget::IsValid() const
 {
 	return Widget.IsValid();
diff --git a/Engine/Source/Runtime/SlateCore/Public/Widgets/Accessibility/SlateAccessibleMessageHandler.h b/Engine/Source/Runtime/SlateCore/Public/Widgets/Accessibility/SlateAccessibleMessageHandler.h
--- a/Engine/Source/Runtime/SlateCore/Public/Widgets/Accessibility/SlateAccessibleMessageHandler.h
+++ b/Engine/Source/Runtime/SlateCore/Public/Widgets/Accessibility/SlateAccessibleMessageHandler.h
@@ -1,3 +1,5 @@
+// Copyright 1998-2019 Epic Games, Inc. All Rights Reserved.
+
 #pragma once
 
 #if WITH_ACCESSIBILITY
@@ -17,6 +19,8 @@ public:
 	virtual void OnActivate() override;
 	virtual void OnDeactivate() override;
 	virtual TSharedPtr<IAccessibleWidget> GetAccessibleWindow(const TSharedRef<FGenericWindow>& InWindow) const override;
+	virtual AccessibleWidgetId GetAccessibleWindowId(const TSharedRef<FGenericWindow>& InWindow) const override;
+	virtual TSharedPtr<IAccessibleWidget> GetAccessibleWidgetFromId(AccessibleWidgetId Id) const override;
 	virtual bool ApplicationIsAccessible() const override { return true; }
 	//~
 
@@ -32,6 +36,15 @@ public:
 	 * @param Widget The widget whose parent changed.
 	 */
 	void OnWidgetParentChanged(TSharedRef<SWidget> Widget);
+	/**
+	 * Callback for a Slate widget's children changing. Although somewhat rare, widgets are not required
+	 * to return a child even if the child was parented to it (for instance, SWidgetSwitcher).
+	 *
+	 * Note: This should not be called if the children are already calling OnWidgetParentChanged().
+	 *
+	 * @param Widget The widget whose children changed.
+	 */
+	void OnWidgetChildrenChanged(TSharedRef<SWidget> Widget);
 	/**
 	 * Callback for a Slate widget's main accessible behavior changing. 
 	 *
diff --git a/Engine/Source/Runtime/SlateCore/Public/Widgets/Accessibility/SlateAccessibleWidgetCache.h b/Engine/Source/Runtime/SlateCore/Public/Widgets/Accessibility/SlateAccessibleWidgetCache.h
--- a/Engine/Source/Runtime/SlateCore/Public/Widgets/Accessibility/SlateAccessibleWidgetCache.h
+++ b/Engine/Source/Runtime/SlateCore/Public/Widgets/Accessibility/SlateAccessibleWidgetCache.h
@@ -35,13 +35,17 @@ public:
 	 */
 	TSharedPtr<FSlateAccessibleWidget> RemoveWidget(SWidget* Widget)
 	{
-		TSharedPtr<FSlateAccessibleWidget> RemovedWidget = nullptr;
+		TSharedPtr<FSlateAccessibleWidget> RemovedWidget;
 		TSharedPtr<FSlateAccessibleWidget>* AccessibleWidget = AccessibleWidgetMap.Find(Widget);
 		if (AccessibleWidget)
 		{
 			// Make sure to disconnect this widget from the tree
 			(*AccessibleWidget)->UpdateParent(nullptr);
 			AccessibleWidgetMap.RemoveAndCopyValue(Widget, RemovedWidget);
+			if (RemovedWidget.IsValid())
+			{
+				AccessibleIdMap.Remove(RemovedWidget->GetId());
+			}
 		}
 		return RemovedWidget;
 	}
@@ -56,7 +60,7 @@ public:
 	 * @param bCreateIfMissing If the widget is not already in the cache, whether or not it should be created
 	 * @return The accessible widget for the given Slate widget, or nullptr if the widget is not accessible.
 	 */
-	TSharedPtr<FSlateAccessibleWidget> GetAccessibleWidget(TSharedPtr<SWidget> Widget, bool bCreateIfMissing = true)
+	TSharedPtr<FSlateAccessibleWidget> GetAccessibleWidget(const TSharedPtr<SWidget> Widget, bool bCreateIfMissing = true)
 	{
 		if (!Widget.IsValid() || !Widget->IsAccessible())
 		{
@@ -71,6 +75,7 @@ public:
 		else if (bCreateIfMissing)
 		{
 			TSharedPtr<FSlateAccessibleWidget> NewWidget = AccessibleWidgetMap.Add(Widget.Get(), Widget->CreateAccessibleWidget());
+			AccessibleIdMap.Add(NewWidget->GetId(), NewWidget);
 			// todo: fix this for nested windows
 			if (NewWidget->AsWindow())
 			{
@@ -81,12 +86,32 @@ public:
 		return nullptr;
 	}
 
+	/**
+	 * Get a cached accessible widget for an identifier that matches the accessible widget's GetId() return value.
+	 * This will NOT create the widget if it does not exist yet, since Ids can only be generated by the widget itself.
+	 * Note that InvalidAccessibleWidgetId will always return null.
+	 *
+	 * @param Id The Id of a widget that corresponds to that widget's GetId() function
+	 * @return The accessible widget for the given Id, or nullptr if there is no widget with that Id.
+	 */
+	TSharedPtr<FSlateAccessibleWidget> GetAccessibleWidgetFromId(AccessibleWidgetId Id)
+	{
+		TSharedPtr<FSlateAccessibleWidget>* AccessibleWidget = AccessibleIdMap.Find(Id);
+		if (AccessibleWidget)
+		{
+			return *AccessibleWidget;
+		}
+		return nullptr;
+	}
+
 private:
 	/**
 	 * Hold on to a shared pointer for each accessible widget until the SWidget is deleted. A raw pointer is used
 	 * because widgets are removed from the cache via the SWidget destructor.
 	 */
 	TMap<SWidget*, TSharedPtr<FSlateAccessibleWidget>> AccessibleWidgetMap;
+	/** Map each accessible widget to its ID to provide a second way to do lookups. */
+	TMap<AccessibleWidgetId, TSharedPtr<FSlateAccessibleWidget>> AccessibleIdMap;
 
 	FSlateAccessibleWidgetCache() {}
 	~FSlateAccessibleWidgetCache() {}
diff --git a/Engine/Source/Runtime/SlateCore/Public/Widgets/Accessibility/SlateCoreAccessibleWidgets.h b/Engine/Source/Runtime/SlateCore/Public/Widgets/Accessibility/SlateCoreAccessibleWidgets.h
--- a/Engine/Source/Runtime/SlateCore/Public/Widgets/Accessibility/SlateCoreAccessibleWidgets.h
+++ b/Engine/Source/Runtime/SlateCore/Public/Widgets/Accessibility/SlateCoreAccessibleWidgets.h
@@ -22,6 +22,7 @@ public:
 	virtual ~FSlateAccessibleWidget();
 
 	// IAccessibleWidget
+	virtual AccessibleWidgetId GetId() const override final;
 	virtual bool IsValid() const override final;
 	virtual TSharedPtr<IAccessibleWidget> GetTopLevelWindow() const override final;
 	virtual FBox2D GetBounds() const override final;
@@ -92,6 +93,8 @@ protected:
 	TArray<TWeakPtr<FSlateAccessibleWidget>> Children;
 	/** The index of this widget in its parent's list of children. */
 	int32 SiblingIndex;
+	/** An application-unique identifier for GetId(). */
+	AccessibleWidgetId Id;
 	/** Whether the contents of the Children array has changed and UpdateAllChildren() needs to be called. */
 	bool bChildrenDirty;
 
