Fix IOS VoiceOver deprecation warning

diff --git a/Engine/Source/Runtime/ApplicationCore/Private/IOS/IOSAppDelegate.cpp b/Engine/Source/Runtime/ApplicationCore/Private/IOS/IOSAppDelegate.cpp
--- a/Engine/Source/Runtime/ApplicationCore/Private/IOS/IOSAppDelegate.cpp
+++ b/Engine/Source/Runtime/ApplicationCore/Private/IOS/IOSAppDelegate.cpp
@@ -1236,7 +1236,7 @@ static FAutoConsoleVariableRef CVarGEnableThermalsReport(
 	[self InitializeAudioSession];
 
 #if WITH_ACCESSIBILITY
-	[[NSNotificationCenter defaultCenter] addObserver:self selector:@selector(OnVoiceOverStatusChanged) name:UIAccessibilityVoiceOverStatusChanged object:nil];
+	[[NSNotificationCenter defaultCenter] addObserver:self selector:@selector(OnVoiceOverStatusChanged) name:UIAccessibilityVoiceOverStatusDidChangeNotification object:nil];
 #endif
     
 	return YES;
