Accessibility support for overlapping visible widgets, and change IOS to use custom hit test code

diff --git a/Engine/Source/Runtime/ApplicationCore/Private/IOS/Accessibility/IOSAccessibilityCache.cpp b/Engine/Source/Runtime/ApplicationCore/Private/IOS/Accessibility/IOSAccessibilityCache.cpp
--- a/Engine/Source/Runtime/ApplicationCore/Private/IOS/Accessibility/IOSAccessibilityCache.cpp
+++ b/Engine/Source/Runtime/ApplicationCore/Private/IOS/Accessibility/IOSAccessibilityCache.cpp
@@ -121,34 +121,11 @@
 					// All UIKit functions must be run on Main Thread
 					dispatch_async(dispatch_get_main_queue(), ^
 					{
-						for (AccessibleWidgetId ChildId : ChildIds)
-						{
-							// Ensure child widgets are created. This is pretty wasteful although at the moment
-							// we have no good way to detect either when a widget is created outside of the hierarchy
-							// or when am SWidgetSwitcher changes to a child that hasn't been initialized yet.
-							[[FIOSAccessibilityCache AccessibilityElementCache] GetAccessibilityElement:ChildId];
-						}
-							
 						FIOSAccessibilityContainer* Element = [[FIOSAccessibilityCache AccessibilityElementCache] GetAccessibilityElement:Id];
 						Element.ChildIds = ChildIds;
-						if (bIsVisible)
-						{
-							Element.Bounds = CGRectMake(Bounds.Min.X, Bounds.Min.Y, Bounds.Max.X - Bounds.Min.X, Bounds.Max.Y - Bounds.Min.Y);
-						}
-						else
-						{
-							// Rather than removing hidden items from the tree, we set their bounding rect to 0 so they're never selected.
-							Element.Bounds = CGRectMake(0, 0, 0, 0);
-						}
+						Element.Bounds = Bounds;
 
-						if (bIsEnabled)
-						{
-							[Element GetLeaf].Traits &= ~UIAccessibilityTraitNotEnabled;
-						}
-						else
-						{
-							[Element GetLeaf].Traits |= UIAccessibilityTraitNotEnabled;
-						}
+						[[Element GetLeaf] SetAccessibilityTrait:UIAccessibilityTraitNotEnabled Set:!bIsEnabled];
 						Element.bIsVisible = bIsVisible;
 					});
 				}
diff --git a/Engine/Source/Runtime/ApplicationCore/Private/IOS/Accessibility/IOSAccessibilityElement.cpp b/Engine/Source/Runtime/ApplicationCore/Private/IOS/Accessibility/IOSAccessibilityElement.cpp
--- a/Engine/Source/Runtime/ApplicationCore/Private/IOS/Accessibility/IOSAccessibilityElement.cpp
+++ b/Engine/Source/Runtime/ApplicationCore/Private/IOS/Accessibility/IOSAccessibilityElement.cpp
@@ -71,7 +71,10 @@
 
 -(CGRect)accessibilityFrame
 {
-	return Bounds;
+	// This function will be called less than the function to cache the bounds, so make
+	// the IOS rect here. If we refactor the code to not using a polling-based cache,
+	// it may make more sense to change the Bounds property itself to a CGRect.
+	return CGRectMake(Bounds.Min.X, Bounds.Min.Y, Bounds.Max.X - Bounds.Min.X, Bounds.Max.Y - Bounds.Min.Y);
 }
 
 -(NSInteger)accessibilityElementCount
@@ -112,6 +115,54 @@
 	return NSNotFound;
 }
 
+-(id)accessibilityHitTest:(CGPoint)point
+{
+	const AccessibleWidgetId TempId = self.Id;
+	AccessibleWidgetId FoundId = IAccessibleWidget::InvalidAccessibleWidgetId;
+
+	const float Scale = [IOSAppDelegate GetDelegate].IOSView.contentScaleFactor;
+	const int32 X = point.x * Scale;
+	const int32 Y = point.y * Scale;
+	// Update the labels while we're on the game thread, since IOS is going to request them immediately.
+	FString TempLabel, TempHint, TempValue;
+
+	[IOSAppDelegate WaitAndRunOnGameThread:[TempId, X, Y, &FoundId, &TempLabel, &TempHint, &TempValue]()
+	{
+		const TSharedPtr<IAccessibleWidget> Widget = [IOSAppDelegate GetDelegate].IOSApplication->GetAccessibleMessageHandler()->GetAccessibleWidgetFromId(TempId);
+		if (Widget.IsValid())
+		{
+			const TSharedPtr<IAccessibleWidget> HitWidget = Widget->GetTopLevelWindow()->AsWindow()->GetChildAtPosition(X, Y);
+			if (HitWidget.IsValid())
+			{
+				FoundId = HitWidget->GetId();
+				TempLabel = HitWidget->GetWidgetName();
+				TempHint = HitWidget->GetHelpText();
+				if (HitWidget->AsProperty())
+				{
+					TempValue = HitWidget->AsProperty()->GetValue();
+				}
+			}
+		}
+	}];
+
+	if (FoundId != IAccessibleWidget::InvalidAccessibleWidgetId)
+	{
+		FIOSAccessibilityLeaf* FoundLeaf = [[[FIOSAccessibilityCache AccessibilityElementCache] GetAccessibilityElement:FoundId] GetLeaf];
+		if ([FoundLeaf ShouldCacheStrings])
+		{
+			FoundLeaf.Label = MoveTemp(TempLabel);
+			FoundLeaf.Hint = MoveTemp(TempHint);
+			FoundLeaf.Value = MoveTemp(TempValue);
+			FoundLeaf.LastCachedStringTime = FPlatformTime::Seconds();
+		}
+		return FoundLeaf;
+	}
+	else
+	{
+		return Leaf;
+	}
+}
+
 @end
 
 @implementation FIOSAccessibilityLeaf
@@ -120,6 +171,7 @@
 @synthesize Hint;
 @synthesize Value;
 @synthesize Traits;
+@synthesize LastCachedStringTime;
 
 -(id)initWithParent:(FIOSAccessibilityContainer*)Parent
 {
@@ -155,11 +207,23 @@
 					InitialTraits |= UIAccessibilityTraitNotEnabled;
 				}
 
+				FString InitialLabel = Widget->GetWidgetName();
+				FString InitialHint = Widget->GetHelpText();
+				FString InitialValue;
+				if (Widget->AsProperty())
+				{
+					InitialValue = Widget->AsProperty()->GetValue();
+				}
+
 				// All UIKit functions must be run on Main Thread
 				dispatch_async(dispatch_get_main_queue(), ^
 				{
-					FIOSAccessibilityContainer* Element = [[FIOSAccessibilityCache AccessibilityElementCache] GetAccessibilityElement:ParentId];
-					[Element GetLeaf].Traits = InitialTraits;
+					FIOSAccessibilityLeaf* Self = [[[FIOSAccessibilityCache AccessibilityElementCache] GetAccessibilityElement:ParentId] GetLeaf];
+					Self.Traits = InitialTraits;
+					Self.Label = InitialLabel;
+					Self.Hint = InitialHint;
+					Self.Value = InitialValue;
+					Self.LastCachedStringTime = FPlatformTime::Seconds();
 				});
 			}
 		}, TStatId(), NULL, ENamedThreads::GameThread);
@@ -177,44 +241,13 @@
 	return [self.accessibilityContainer accessibilityFrame];
 }
 
-/**
- * Since Strings are expensive to copy over and over, we don't update them with polling.
- * When any one of the three are requested, we cache all three in order to restrict it to
- * a single cross-thread call. This will add a slight delay when the user taps on a widget,
- * but at the moment it isn't noticeable. The slower the framerate of the game, the more
- * noticeable this delay may become.
- */
--(void)CacheStrings
+-(bool)ShouldCacheStrings
 {
-	const double StartTime = FPlatformTime::Seconds();
-	if (StartTime - LastCachedStringTime > 1.0)
-	{
-		const AccessibleWidgetId TempId = ((FIOSAccessibilityContainer*)self.accessibilityContainer).Id;
-		FString TempLabel, TempHint, TempValue;
-		// All IAccessibleWidget functions must be run on Game Thread
-		[IOSAppDelegate WaitAndRunOnGameThread : [TempId, &TempLabel, &TempHint, &TempValue]()
-		{
-			TSharedPtr<IAccessibleWidget> Widget = [IOSAppDelegate GetDelegate].IOSApplication->GetAccessibleMessageHandler()->GetAccessibleWidgetFromId(TempId);
-			if (Widget.IsValid())
-			{
-				TempLabel = Widget->GetWidgetName();
-				TempHint = Widget->GetHelpText();
-				if (Widget->AsProperty())
-				{
-					TempValue = Widget->AsProperty()->GetValue();
-				}
-			}
-		}];
-		self.Label = MoveTemp(TempLabel);
-		self.Hint = MoveTemp(TempHint);
-		self.Value = MoveTemp(TempValue);
-		LastCachedStringTime = FPlatformTime::Seconds();
-	}
+	return FPlatformTime::Seconds() - LastCachedStringTime > 1.0;
 }
 
 -(NSString*)accessibilityLabel
 {
-	[self CacheStrings];
 	if (!self.Label.IsEmpty())
 	{
 		return [NSString stringWithFString:self.Label];
@@ -224,7 +257,6 @@
 
 -(NSString*)accessibilityHint
 {
-	[self CacheStrings];
 	if (!self.Hint.IsEmpty())
 	{
 		return [NSString stringWithFString:self.Hint];
@@ -234,7 +266,6 @@
 
 -(NSString*)accessibilityValue
 {
-	[self CacheStrings];
 	if (!self.Value.IsEmpty())
 	{
 		return [NSString stringWithFString:self.Value];
@@ -242,6 +273,18 @@
 	return nil;
 }
 
+-(void)SetAccessibilityTrait:(UIAccessibilityTraits)Trait Set:(bool)IsEnabled
+{
+	if (IsEnabled)
+	{
+		self.Traits |= Trait;
+	}
+	else
+	{
+		self.Traits &= ~Trait;
+	}
+}
+
 -(UIAccessibilityTraits)accessibilityTraits
 {
 	return self.Traits;
diff --git a/Engine/Source/Runtime/ApplicationCore/Public/IOS/Accessibility/IOSAccessibilityElement.h b/Engine/Source/Runtime/ApplicationCore/Public/IOS/Accessibility/IOSAccessibilityElement.h
--- a/Engine/Source/Runtime/ApplicationCore/Public/IOS/Accessibility/IOSAccessibilityElement.h
+++ b/Engine/Source/Runtime/ApplicationCore/Public/IOS/Accessibility/IOSAccessibilityElement.h
@@ -40,7 +40,7 @@
 /** A list of identifiers for all current children of this container. */
 @property (nonatomic) TArray<AccessibleWidgetId> ChildIds;
 /** The bounding rect of the container. */
-@property (nonatomic) CGRect Bounds;
+@property (nonatomic) FBox2D Bounds;
 /** Whether or not the widget is currently visible. */
 @property (nonatomic) bool bIsVisible;
 
@@ -53,13 +53,14 @@
  */
 @interface FIOSAccessibilityLeaf : UIAccessibilityElement
 {
-@private
-	/** Timestamp for when Label, Hint, and Value were last cached. */
-	double LastCachedStringTime;
 }
 
 /** This must be used instead of initWithAccessibilityContainer in order to work properly. */
 -(id)initWithParent:(FIOSAccessibilityContainer*)Parent;
+/** Check if LastCachedStringTime was updated recently. */
+-(bool)ShouldCacheStrings;
+/**  Toggle an individual trait on or off */
+-(void)SetAccessibilityTrait:(UIAccessibilityTraits)Trait Set:(bool)IsEnabled;
 
 /** A cached version of the name of the widget. */
 @property (nonatomic) FString Label;
@@ -69,6 +70,8 @@
 @property (nonatomic) FString Value;
 /** Bitflag of traits that describe the widget. Most are set once on initialization. */
 @property (nonatomic) UIAccessibilityTraits Traits;
+/** Timestamp for when Label, Hint, and Value were last cached. */
+@property (nonatomic) double LastCachedStringTime;
 
 @end
 
diff --git a/Engine/Source/Runtime/SlateCore/Private/Widgets/Accessibility/SlateCoreAccessibleWidgets.cpp b/Engine/Source/Runtime/SlateCore/Private/Widgets/Accessibility/SlateCoreAccessibleWidgets.cpp
--- a/Engine/Source/Runtime/SlateCore/Private/Widgets/Accessibility/SlateCoreAccessibleWidgets.cpp
+++ b/Engine/Source/Runtime/SlateCore/Private/Widgets/Accessibility/SlateCoreAccessibleWidgets.cpp
@@ -354,7 +354,8 @@ TSharedPtr<IAccessibleWidget> FSlateAccessibleWidget::GetChildAtUsingGeometry(in
 	if (!IsHidden() && GetBounds().IsInside(FVector2D(X, Y)))
 	{
 		UpdateAllChildren();
-		for (int32 i = 0; i < Children.Num(); ++i)
+		// Traverse the hierarchy backwards in order to handle the case where widgets are overlaid on top of each other.
+		for (int32 i = Children.Num() - 1; i >= 0; --i)
 		{
 			if (Children[i].IsValid())
 			{
@@ -382,6 +383,7 @@ TSharedPtr<FGenericWindow> FSlateAccessibleWindow::GetNativeWindow() const
 
 TSharedPtr<IAccessibleWidget> FSlateAccessibleWindow::GetChildAtPosition(int32 X, int32 Y)
 {
+	TSharedPtr<IAccessibleWidget> HitWidget;
 	if (Widget.IsValid())
 	{
 		static const bool UseHitTestGrid = false;
@@ -403,15 +405,15 @@ TSharedPtr<IAccessibleWidget> FSlateAccessibleWindow::GetChildAtPosition(int32 X
 					break;
 				}
 			}
-			return FSlateAccessibleWidgetCache::Get().GetAccessibleWidget(LastAccessibleWidget);
+			HitWidget = FSlateAccessibleWidgetCache::Get().GetAccessibleWidget(LastAccessibleWidget);
 		}
 		else
 		{
-			return GetChildAtUsingGeometry(X, Y);
+			HitWidget = GetChildAtUsingGeometry(X, Y);
 		}
 	}
 
-	return nullptr;
+	return HitWidget;
 }
 
 TSharedPtr<IAccessibleWidget> FSlateAccessibleWindow::GetFocusedWidget() const
