Do not set our selfs hidden if we did not change the size. For example when we are first created and have 0,0 width/height

diff --git a/Engine/Source/Runtime/ApplicationCore/Private/Linux/LinuxWindow.cpp b/Engine/Source/Runtime/ApplicationCore/Private/Linux/LinuxWindow.cpp
--- a/Engine/Source/Runtime/ApplicationCore/Private/Linux/LinuxWindow.cpp
+++ b/Engine/Source/Runtime/ApplicationCore/Private/Linux/LinuxWindow.cpp
@@ -575,22 +575,27 @@ void FLinuxWindow::ReshapeWindow( int32 NewX, int32 NewY, int32 NewWidth, int32
 		break;
 	}
 
+	// Only consider this if we have changed our selfs from an existing size to zero or vis versa
+	if (NewWidth != VirtualWidth || NewHeight != VirtualHeight)
+	{
+		if (bZeroSize && NewWidth != 0 && NewHeight != 0)
+		{
+			Show();
+			bZeroSize = false;
+		}
+		else if (NewWidth == 0 || NewHeight == 0)
+		{
+			UE_DEBUG_BREAK();
+			Hide();
+			bZeroSize = true;
+		}
+	}
+
 	RegionWidth   = NewWidth;
 	RegionHeight  = NewHeight;
 	VirtualWidth  = NewWidth;
 	VirtualHeight = NewHeight;
 
-	if (bZeroSize && NewWidth != 0 && NewHeight != 0)
-	{
-		Show();
-		bZeroSize = false;
-	}
-	else if (NewWidth == 0 || NewHeight == 0)
-	{
-		Hide();
-		bZeroSize = true;
-	}
-
 	// Avoid broadcasting we have set a zero size as it will attempt to resize the backbuffer which on some RHI is invalid per the spec (ie. Vulkan)
 	if (LinuxWindow && !bZeroSize)
 	{
