Revert the bZeroSize fix and just set the size to 1 if we are setting to 0 instead.

The bZeroSize fix was causing notification windows to hide forever while in a Slate Throttle mode causing GameViewportClient draw issue

diff --git a/Engine/Source/Runtime/ApplicationCore/Private/Linux/LinuxWindow.cpp b/Engine/Source/Runtime/ApplicationCore/Private/Linux/LinuxWindow.cpp
--- a/Engine/Source/Runtime/ApplicationCore/Private/Linux/LinuxWindow.cpp
+++ b/Engine/Source/Runtime/ApplicationCore/Private/Linux/LinuxWindow.cpp
@@ -454,7 +454,7 @@ void FLinuxWindow::Restore()
 /** Native window should make itself visible */
 void FLinuxWindow::Show()
 {
-	if ( !bIsVisible && !bZeroSize )
+	if ( !bIsVisible )
 	{
 		bIsVisible = true;
 		SDL_ShowWindow( HWnd );
@@ -529,6 +529,11 @@ void FLinuxWindow::ReshapeWindow( int32 NewX, int32 NewY, int32 NewWidth, int32
 		}
 	}
 
+	// If we have set our self to 0,0 Width/Height it will not be allowed we will still show the window
+	// this is a work around to at least reduce the visibile impact of a window that is lingering
+	NewWidth  = FMath::Max(NewWidth, 1);
+	NewHeight = FMath::Max(NewHeight, 1);
+
 	// X11 will take until the next frame to send a SizeChanged event. This means the X11 window
 	// will most likely have resized already by the time we render but the slate renderer will
 	// not have been updated leading to an incorrect frame.
@@ -575,29 +580,13 @@ void FLinuxWindow::ReshapeWindow( int32 NewX, int32 NewY, int32 NewWidth, int32
 		break;
 	}
 
-	// Only consider this if we have changed our selfs from an existing size to zero or vis versa
-	if (NewWidth != VirtualWidth || NewHeight != VirtualHeight)
-	{
-		if (bZeroSize && NewWidth != 0 && NewHeight != 0)
-		{
-			Show();
-			bZeroSize = false;
-		}
-		else if (NewWidth == 0 || NewHeight == 0)
-		{
-			UE_DEBUG_BREAK();
-			Hide();
-			bZeroSize = true;
-		}
-	}
-
 	RegionWidth   = NewWidth;
 	RegionHeight  = NewHeight;
 	VirtualWidth  = NewWidth;
 	VirtualHeight = NewHeight;
 
 	// Avoid broadcasting we have set a zero size as it will attempt to resize the backbuffer which on some RHI is invalid per the spec (ie. Vulkan)
-	if (LinuxWindow && !bZeroSize)
+	if (LinuxWindow )
 	{
 		OwningApplication->GetMessageHandler()->OnSizeChanged(
 			LinuxWindow.ToSharedRef(),
diff --git a/Engine/Source/Runtime/ApplicationCore/Public/Linux/LinuxWindow.h b/Engine/Source/Runtime/ApplicationCore/Public/Linux/LinuxWindow.h
--- a/Engine/Source/Runtime/ApplicationCore/Public/Linux/LinuxWindow.h
+++ b/Engine/Source/Runtime/ApplicationCore/Public/Linux/LinuxWindow.h
@@ -269,11 +269,6 @@ private:
 	/** Whether native properties cache is valid */
 	bool bValidNativePropertiesCache;
 
-	/** Whether the window size was set to zero at some point, if its the case we need to hide our window and will need to
-	 *  restore the window when we set our selfs to a new size that is not zero
-	 */
-	bool bZeroSize;
-
 	/**
 	 * Ratio of pixels to SlateUnits in this window.
 	 * E.g. DPIScale of 2.0 means there is a 2x2 pixel square for every 1x1 SlateUnit.
