Add Linux game splash screen support (without text)

Daimler was asking about time to display splash loading screen.
  This should help them.

Limitations:
  Only loads .bmp files
  Doesn't support any dynamic text (SetSplashText method NYI right now)
  Uses straight SDL2 (Not GL, Freetype, or ImageWrapper)

diff --git a/Engine/Source/Runtime/ApplicationCore/Private/Linux/LinuxPlatformSplash.cpp b/Engine/Source/Runtime/ApplicationCore/Private/Linux/LinuxPlatformSplash.cpp
--- a/Engine/Source/Runtime/ApplicationCore/Private/Linux/LinuxPlatformSplash.cpp
+++ b/Engine/Source/Runtime/ApplicationCore/Private/Linux/LinuxPlatformSplash.cpp
@@ -23,15 +23,17 @@
 #include "Modules/ModuleManager.h"
 #include "HAL/PlatformApplicationMisc.h"
 
+#include "SDL.h"
+
 #if WITH_EDITOR
 
-#include "SDL.h"
 #include "IImageWrapper.h"
 #include "IImageWrapperModule.h"
 
 #include "ft2build.h"
 #include FT_FREETYPE_H
 
+#endif // WITH_EDITOR
 
 /**
  * Splash screen functions and static globals
@@ -57,24 +59,29 @@ public:
 
 private:
 	static SDL_Surface* LoadImage(const FString &InImagePath);
+	void Redraw();
+
+#if WITH_EDITOR
 	void DrawCharacter(int32 penx, int32 peny, FT_GlyphSlot Glyph, int32 CurTypeIndex, float Red, float Green, float Blue);
 	void RenderStrings();
 	bool OpenFonts();
-	void Redraw();
 
 	FT_Library FontLibrary = nullptr;
 	FT_Face FontSmall = nullptr;
 	FT_Face FontNormal = nullptr;
 	FT_Face FontLarge = nullptr;
 
-	SDL_Surface *SplashSurface = nullptr;
-	SDL_Window *SplashWindow = nullptr;
 	SDL_Renderer *SplashRenderer = nullptr;
 	SDL_Texture *SplashTexture = nullptr;
 	FText SplashText[ SplashTextType::NumTextTypes ];
 	Rect SplashTextRects[ SplashTextType::NumTextTypes ];
 
 	unsigned char *ScratchSpace = nullptr;
+#endif // WITH_EDITOR
+
+	SDL_Surface *SplashSurface = nullptr;
+	SDL_Window *SplashWindow = nullptr;
+
 	bool bNeedsRedraw = false;
 	bool bStringsChanged = false;
 };
@@ -82,17 +89,14 @@ private:
 static FLinuxSplashState *GSplashState = nullptr;
 
 
+//---------------------------------------------------------
 FLinuxSplashState::~FLinuxSplashState()
 {
+#if WITH_EDITOR
 	// Just in case SDL's renderer steps on GL state...
 	SDL_Window *CurrentWindow = SDL_GL_GetCurrentWindow();
 	SDL_GLContext CurrentContext = SDL_GL_GetCurrentContext();
 
-	if (SplashSurface)
-	{
-		SDL_FreeSurface(SplashSurface);
-	}
-
 	if (SplashTexture)
 	{
 		SDL_DestroyTexture(SplashTexture);
@@ -102,12 +106,19 @@ FLinuxSplashState::~FLinuxSplashState()
 	{
 		SDL_DestroyRenderer(SplashRenderer);
 	}
+#endif // WITH_EDITOR
+
+	if (SplashSurface)
+	{
+		SDL_FreeSurface(SplashSurface);
+	}
 
 	if (SplashWindow)
 	{
 		SDL_DestroyWindow(SplashWindow);
 	}
 
+#if WITH_EDITOR
 	if (ScratchSpace)
 	{
 		FMemory::Free(ScratchSpace);
@@ -137,10 +148,13 @@ FLinuxSplashState::~FLinuxSplashState()
 	{
 		SDL_GL_MakeCurrent(CurrentWindow, CurrentContext);
 	}
+#endif // WITH_EDITOR
 
 	// do not deinit SDL here
 }
 
+#if WITH_EDITOR
+
 //---------------------------------------------------------
 bool FLinuxSplashState::OpenFonts()
 {
@@ -152,7 +166,7 @@ bool FLinuxSplashState::OpenFonts()
 
 	// small font face
 	FString FontPath = FPaths::ConvertRelativePathToFull(FPaths::EngineContentDir() / TEXT("Slate/Fonts/Roboto-Light.ttf"));
-	
+
 	if (FT_New_Face(FontLibrary, TCHAR_TO_UTF8(*FontPath), 0, &FontSmall))
 	{
 		UE_LOG(LogHAL, Error, TEXT("*** Unable to open small font face for splash screen."));
@@ -164,7 +178,7 @@ bool FLinuxSplashState::OpenFonts()
 
 	// normal font face
 	FontPath = FPaths::ConvertRelativePathToFull(FPaths::EngineContentDir() / TEXT("Slate/Fonts/Roboto-Regular.ttf"));
-	
+
 	if (FT_New_Face(FontLibrary, TCHAR_TO_UTF8(*FontPath), 0, &FontNormal))
 	{
 		UE_LOG(LogHAL, Error, TEXT("*** Unable to open normal font face for splash screen."));
@@ -173,13 +187,13 @@ bool FLinuxSplashState::OpenFonts()
 	{
 		FT_Set_Pixel_Sizes(FontNormal, 0, 12);
 	}
-	
+
 	// large font face
 	FontPath = FPaths::ConvertRelativePathToFull(FPaths::EngineContentDir() / TEXT("Slate/Fonts/Roboto-Bold.ttf"));
-	
+
 	if (FT_New_Face(FontLibrary, TCHAR_TO_UTF8(*FontPath), 0, &FontLarge))
 	{
-		UE_LOG(LogHAL, Error, TEXT("*** Unable to open large font face for splash screen."));		
+		UE_LOG(LogHAL, Error, TEXT("*** Unable to open large font face for splash screen."));
 	}
 	else
 	{
@@ -217,12 +231,12 @@ void FLinuxSplashState::DrawCharacter(int32 PenX, int32 PenY, FT_GlyphSlot Glyph
 			// find pixel position in splash image
 			const int PosX = PenX + GlyphX + (Glyph->metrics.horiBearingX >> 6);
 			const int PosY = PenY + GlyphY - (Glyph->metrics.horiBearingY >> 6);
-			
+
 			// make sure pixel is in drawing rectangle
 			if (PosX < MinX || PosX >= MaxX || PosY < MinY || PosY >= MaxY)
 				continue;
 
-			// get index of pixel in glyph bitmap			
+			// get index of pixel in glyph bitmap
 			const int32 SourceIndex = (GlyphY * GlyphPitch) + GlyphX;
 			int32 DestIndex = (PosY * SplashWidth + PosX) * SplashBPP;
 
@@ -266,12 +280,12 @@ void FLinuxSplashState::RenderStrings()
 	for (int CurTypeIndex=0; CurTypeIndex<SplashTextType::NumTextTypes; CurTypeIndex++)
 	{
 		FT_Face Font = nullptr;
-		
+
 		// initial pen position
 		uint32 PenX = SplashTextRects[ CurTypeIndex].Left;
 		uint32 PenY = SplashTextRects[ CurTypeIndex].Bottom;
 		Kerning.x = 0;
-		
+
 		// Set font color and text position
 		if (CurTypeIndex == SplashTextType::StartupProgress)
 		{
@@ -295,7 +309,7 @@ void FLinuxSplashState::RenderStrings()
 			Red = Green = Blue = 160.0f;
 			Font = FontSmall;
 		}
-		
+
 		// sanity check: make sure we have a font loaded.
 		if (!Font)
 		{
@@ -312,7 +326,7 @@ void FLinuxSplashState::RenderStrings()
 			for (int i=0; i<Text.Len(); i++)
 			{
 				FT_ULong CharacterCode;
-				
+
 				// fetch next glyph
 				if (bRightJustify)
 				{
@@ -320,13 +334,12 @@ void FLinuxSplashState::RenderStrings()
 				}
 				else
 				{
-					CharacterCode = (Uint32)(Text[i]);				
+					CharacterCode = (Uint32)(Text[i]);
 				}
-							
+
 				FT_UInt GlyphIndex = FT_Get_Char_Index(Font, CharacterCode);
 				FT_Load_Glyph(Font, GlyphIndex, FT_LOAD_DEFAULT);
 
-				
 				if (Font->glyph->format != FT_GLYPH_FORMAT_BITMAP)
 				{
 					FT_Render_Glyph(Font->glyph, FT_RENDER_MODE_NORMAL);
@@ -339,22 +352,22 @@ void FLinuxSplashState::RenderStrings()
 					{
 						FT_Get_Kerning(Font, GlyphIndex, LastGlyph, FT_KERNING_DEFAULT, &Kerning);
 					}
-					
-					PenX -= (Font->glyph->metrics.horiAdvance - Kerning.x) >> 6;	
+
+					PenX -= (Font->glyph->metrics.horiAdvance - Kerning.x) >> 6;
 				}
 				else
 				{
 					if (LastGlyph != 0)
 					{
 						FT_Get_Kerning(Font, LastGlyph, GlyphIndex, FT_KERNING_DEFAULT, &Kerning);
-					}	
+					}
 				}
 
 				LastGlyph = GlyphIndex;
 
 				// draw character
 				DrawCharacter(PenX, PenY, Font->glyph, CurTypeIndex, Red, Green, Blue);
-							
+
 				if (!bRightJustify)
 				{
 					PenX += (Font->glyph->metrics.horiAdvance - Kerning.x) >> 6;
@@ -367,6 +380,9 @@ void FLinuxSplashState::RenderStrings()
 	}
 }
 
+#endif // WITH_EDITOR
+
+
 /**
  * @brief Helper function to load an image in any format.
  *
@@ -376,6 +392,7 @@ void FLinuxSplashState::RenderStrings()
  */
 SDL_Surface* FLinuxSplashState::LoadImage(const FString &ImagePath)
 {
+#if WITH_EDITOR
 	TArray<uint8> RawFileData;
 
 	// Load the image buffer first (unless it's BMP)
@@ -402,23 +419,13 @@ SDL_Surface* FLinuxSplashState::LoadImage(const FString &ImagePath)
 			}
 		}
 	}
+#endif // WITH_EDITOR
 
 	// If for some reason the image cannot be loaded, use the default BMP function
 	return SDL_LoadBMP(TCHAR_TO_UTF8(*ImagePath));
 }
 
 
-// This class helps cleans up SDL_Surfaces when they go out of scope.
-class SdlSurfacePtr
-{
-public:
-	SdlSurfacePtr(SDL_Surface *InSurface) : Surface(InSurface) {}
-	~SdlSurfacePtr() { if (Surface) { SDL_FreeSurface(Surface); } }
-private:
-	SDL_Surface *Surface;
-};
-
-
 /** Helper function to init resources used by the splash window */
 bool FLinuxSplashState::InitSplashResources(const FText &AppName, const FString &SplashPath, const FString &IconPath)
 {
@@ -448,16 +455,11 @@ bool FLinuxSplashState::InitSplashResources(const FText &AppName, const FString
 		return -1;
 	}
 
-	SDL_Surface *SplashIconImage = LoadImage(IconPath);
-	SdlSurfacePtr SplashIconPtr(SplashIconImage);
-	if (SplashIconImage == nullptr)
-	{
-		UE_LOG(LogHAL, Warning, TEXT("FLinuxSplashState::InitSplashResources() : Splash icon could not be created! SDL_Error: %s"), UTF8_TO_TCHAR(SDL_GetError()));
-	}
-
+#if WITH_EDITOR
 	// Just in case SDL's renderer steps on GL state...
 	SDL_Window *CurrentWindow = SDL_GL_GetCurrentWindow();
 	SDL_GLContext CurrentContext = SDL_GL_GetCurrentContext();
+#endif
 
 	// on modern X11, your windows might turn gray if they don't pump the event queue fast enough.
 	// But this is because they opt-in to an optional window manager protocol by default; legacy
@@ -482,17 +484,31 @@ bool FLinuxSplashState::InitSplashResources(const FText &AppName, const FString
 		return false;
 	}
 
-	if (SplashIconImage != nullptr)
+	if (!IconPath.IsEmpty())
 	{
-		SDL_SetWindowIcon(SplashWindow, SplashIconImage);  // SDL_SetWindowIcon() makes a copy of this.
+		SDL_Surface *SplashIconImage = LoadImage(IconPath);
+
+		if (SplashIconImage == nullptr)
+		{
+			UE_LOG(LogHAL, Warning, TEXT("FLinuxSplashState::InitSplashResources() : Splash icon could not be created! SDL_Error: %s"), UTF8_TO_TCHAR(SDL_GetError()));
+		}
+		else
+		{
+			SDL_SetWindowIcon(SplashWindow, SplashIconImage);  // SDL_SetWindowIcon() makes a copy of this.
+			SDL_FreeSurface(SplashIconImage);
+		}
 	}
 
+#if WITH_EDITOR
 	SplashRenderer = SDL_CreateRenderer(SplashWindow, -1, 0);
+#endif
 
 	// it's safe to set the hint back once the Renderer is created (since it might recreate the window to get a GL or whatever visual).
 	SDL_SetHint(SDL_HINT_VIDEO_X11_NET_WM_PING, OriginalHint ? OriginalHint : "1");
 	SDL_free(OriginalHint);
 
+#if WITH_EDITOR
+
 	if (SplashRenderer == nullptr)
 	{
 		UE_LOG(LogHAL, Error, TEXT("FLinuxSplashState::InitSplashResources() : Splash screen renderer could not be created! SDL_Error: %s"), UTF8_TO_TCHAR(SDL_GetError()));
@@ -554,9 +570,27 @@ bool FLinuxSplashState::InitSplashResources(const FText &AppName, const FString
 		SDL_GL_MakeCurrent(CurrentWindow, CurrentContext);
 	}
 
+#else
+
+	SDL_Surface *WindowSurface = SDL_GetWindowSurface( SplashWindow );
+	if (WindowSurface == nullptr)
+	{
+		UE_LOG(LogHAL, Error, TEXT("FLinuxSplashState::InitSplashResources() : Splash window surface could not be created! SDL_Error: %s"), UTF8_TO_TCHAR(SDL_GetError()));
+		return false;
+	}
+
+	SDL_BlitSurface(SplashSurface, NULL, WindowSurface, NULL);
+	SDL_FreeSurface(WindowSurface);
+
+	SDL_ShowWindow(SplashWindow);
+	Pump();
+
+#endif
+
 	return true;
 }
 
+
 /**
  * Sets the text displayed on the splash screen (for startup/loading progress)
  *
@@ -574,8 +608,11 @@ void FLinuxSplashState::SetSplashText( const SplashTextType::Type InType, const
 #endif
 }
 
+
+//---------------------------------------------------------
 void FLinuxSplashState::Redraw()
 {
+#if WITH_EDITOR
 	if (bNeedsRedraw || bStringsChanged)
 	{
 		// Just in case SDL's renderer steps on GL state...
@@ -592,8 +629,16 @@ void FLinuxSplashState::Redraw()
 			SDL_GL_MakeCurrent(CurrentWindow, CurrentContext);
 		}
 	}
+#else
+	if (bNeedsRedraw)
+	{
+		SDL_UpdateWindowSurface(SplashWindow);
+		bNeedsRedraw = false;
+	}
+#endif
 }
 
+
 void FLinuxSplashState::Pump()
 {
 	if (SplashWindow)
@@ -606,10 +651,10 @@ void FLinuxSplashState::Pump()
 				bNeedsRedraw = true;
 			}
 		}
+
 		Redraw();
 	}
 }
-#endif //WITH_EDITOR
 
 
 /**
@@ -618,9 +663,8 @@ void FLinuxSplashState::Pump()
  */
 void FLinuxPlatformSplash::Show( )
 {
-#if WITH_EDITOR
 	// need to do a splash screen?
-	if(GSplashState || FParse::Param(FCommandLine::Get(),TEXT("NOSPLASH")) == true)
+	if(GSplashState || !FApp::CanEverRender() || FParse::Param(FCommandLine::Get(), TEXT("NOSPLASH")) == true)
 	{
 		return;
 	}
@@ -672,7 +716,7 @@ void FLinuxPlatformSplash::Show( )
 		if (!bIconFound)
 		{
 			UE_LOG(LogHAL, Warning, TEXT("Game icon not found."));
-			return;	// early out
+			IconPath.Reset();
 		}
 	}
 
@@ -686,7 +730,7 @@ void FLinuxPlatformSplash::Show( )
 
 	// In the editor, we'll display loading info
 	FText AppName;
-	if( GIsEditor )
+	if ( GIsEditor )
 	{
 		// Set initial startup progress info
 		{
@@ -725,7 +769,6 @@ void FLinuxPlatformSplash::Show( )
 		delete GSplashState;
 		GSplashState = nullptr;
 	}
-#endif //WITH_EDITOR
 }
 
 
@@ -734,10 +777,8 @@ void FLinuxPlatformSplash::Show( )
 */
 void FLinuxPlatformSplash::Hide()
 {
-#if WITH_EDITOR
 	delete GSplashState;
 	GSplashState = nullptr;
-#endif
 }
 
 
@@ -749,7 +790,6 @@ void FLinuxPlatformSplash::Hide()
 */
 void FLinuxPlatformSplash::SetSplashText( const SplashTextType::Type InType, const TCHAR* InText )
 {
-#if WITH_EDITOR
 	if (GSplashState)
 	{
 		// We only want to bother drawing startup progress in the editor, since this information is
@@ -760,6 +800,4 @@ void FLinuxPlatformSplash::SetSplashText( const SplashTextType::Type InType, con
 		}
 		GSplashState->Pump();
 	}
-
-#endif //WITH_EDITOR
 }
