Accessibility: Added support for culture and language property in Windows.

diff --git a/Engine/Source/Runtime/ApplicationCore/Private/Windows/Accessibility/WindowsUIAManager.cpp b/Engine/Source/Runtime/ApplicationCore/Private/Windows/Accessibility/WindowsUIAManager.cpp
--- a/Engine/Source/Runtime/ApplicationCore/Private/Windows/Accessibility/WindowsUIAManager.cpp
+++ b/Engine/Source/Runtime/ApplicationCore/Private/Windows/Accessibility/WindowsUIAManager.cpp
@@ -15,6 +15,8 @@
 #include "Windows/Accessibility/WindowsUIAPropertyGetters.h"
 #include "Windows/WindowsApplication.h"
 #include "Windows/WindowsWindow.h"
+#include "Internationalization/Internationalization.h"
+#include "Internationalization/Culture.h"
 
 TMap<EAccessibleWidgetType, ULONG> FWindowsUIAManager::WidgetTypeToWindowsTypeMap;
 TMap<EAccessibleWidgetType, FText> FWindowsUIAManager::WidgetTypeToTextMap;
@@ -34,6 +36,8 @@ void EmitPropertyChangedEvent(FWindowsUIAWidgetProvider* Provider, PROPERTYID Pr
 
 FWindowsUIAManager::FWindowsUIAManager(const FWindowsApplication& InApplication)
 	: WindowsApplication(InApplication)
+	, OnCultureChangedHandle()
+	, CachedCurrentLocaleLCID(0)
 {
 	OnAccessibleMessageHandlerChanged();
 
@@ -88,6 +92,12 @@ FWindowsUIAManager::~FWindowsUIAManager()
 		// The UIA Manager may be deleted before the Providers are, and external applications may attempt to call functions on those Providers.
 		Provider->OnUIAManagerDestroyed();
 	}
+
+	if (OnCultureChangedHandle.IsValid()&& FInternationalization::Get().IsAvailable())
+	{
+		FInternationalization::Get().OnCultureChanged().Remove(OnCultureChangedHandle);
+		OnCultureChangedHandle.Reset();
+	}
 }
 
 FWindowsUIAWidgetProvider& FWindowsUIAManager::GetWidgetProvider(TSharedRef<IAccessibleWidget> InWidget)
@@ -112,9 +122,8 @@ FWindowsUIAWindowProvider& FWindowsUIAManager::GetWindowProvider(TSharedRef<FWin
 {
 	if (CachedWidgetProviders.Num() == 0)
 	{
-		// Enable application accessibility if this is the first Provider. The first Get*Provider() request
-		// MUST go through this function since IAccessibleWidgets will not exist until SetActive is called.
-		WindowsApplication.GetAccessibleMessageHandler()->SetActive(true);
+		// The first Get*Provider() request / MUST go through this function since IAccessibleWidgets will not exist untilthe accessible message handler is set active.
+		OnAccessibilityEnabled();
 	}
 
 	TSharedPtr<IAccessibleWidget> AccessibleWindow = WindowsApplication.GetAccessibleMessageHandler()->GetAccessibleWindow(InWindow);
@@ -122,6 +131,15 @@ FWindowsUIAWindowProvider& FWindowsUIAManager::GetWindowProvider(TSharedRef<FWin
 	return static_cast<FWindowsUIAWindowProvider&>(GetWidgetProvider(AccessibleWindow.ToSharedRef()));
 }
 
+void FWindowsUIAManager::OnAccessibilityEnabled()
+{
+	WindowsApplication.GetAccessibleMessageHandler()->SetActive(true);
+	// Register for language and locale changes for internationalization.
+	// Updates the LCID to be returned in FWindowsUIAWidgetProvider  as a UIA Property 
+	UpdateCachedCurrentLocaleLCID();
+	OnCultureChangedHandle = FInternationalization::Get().OnCultureChanged().AddRaw(this, &FWindowsUIAManager::UpdateCachedCurrentLocaleLCID);
+}
+
 void FWindowsUIAManager::OnWidgetProviderRemoved(TSharedRef<IAccessibleWidget> InWidget)
 {
 	CachedWidgetProviders.Remove(InWidget);
@@ -135,7 +153,18 @@ void FWindowsUIAManager::OnWidgetProviderRemoved(TSharedRef<IAccessibleWidget> I
 		// Note that there could still be control Providers with valid references. In practice I don't think
 		// this is possible, but if we run into problems then we can simply call AddRef and Release on the
 		// underlying widget Provider whenever a control Provider gets added/removed.
-		WindowsApplication.GetAccessibleMessageHandler()->SetActive(false);
+		OnAccessibilityDisabled();
+	}
+}
+
+void FWindowsUIAManager::OnAccessibilityDisabled()
+{
+	WindowsApplication.GetAccessibleMessageHandler()->SetActive(false);
+	CachedCurrentLocaleLCID = 0;
+	if (OnCultureChangedHandle.IsValid())
+	{
+		FInternationalization::Get().OnCultureChanged().Remove(OnCultureChangedHandle);
+		OnCultureChangedHandle.Reset();
 	}
 }
 
@@ -238,6 +267,16 @@ void FWindowsUIAManager::OnEventRaised(TSharedRef<IAccessibleWidget> Widget, EAc
 	}
 }
 
+void FWindowsUIAManager::UpdateCachedCurrentLocaleLCID()
+{
+	CachedCurrentLocaleLCID = FInternationalization::Get().GetCurrentLocale()->GetLCID();
+	// an LCID of 0 is invalid  and and UIA won't do anything with it.
+	// Get the default OS locale LCID to give to screen readers  in that case
+	if (CachedCurrentLocaleLCID == 0)
+	{
+		CachedCurrentLocaleLCID = FInternationalization::Get().GetDefaultLocale()->GetLCID();
+	}
+}
 #if !UE_BUILD_SHIPPING
 void FWindowsUIAManager::DumpAccessibilityStats() const
 {
diff --git a/Engine/Source/Runtime/ApplicationCore/Private/Windows/Accessibility/WindowsUIAWidgetProvider.cpp b/Engine/Source/Runtime/ApplicationCore/Private/Windows/Accessibility/WindowsUIAWidgetProvider.cpp
--- a/Engine/Source/Runtime/ApplicationCore/Private/Windows/Accessibility/WindowsUIAWidgetProvider.cpp
+++ b/Engine/Source/Runtime/ApplicationCore/Private/Windows/Accessibility/WindowsUIAWidgetProvider.cpp
@@ -314,11 +314,10 @@ HRESULT STDCALL FWindowsUIAWidgetProvider::GetPropertyValue(PROPERTYID propertyI
 		pRetVal->vt = VT_I4;
 		pRetVal->lVal = WidgetTypeToControlType(Widget);
 		break;
-	//case UIA_CulturePropertyId:
-	//	pRetVal->vt = VT_I4;
-	//	// todo: https://docs.microsoft.com/en-us/windows/desktop/Intl/language-identifier-constants-and-strings
-	//	pRetVal->lVal = 0;
-	//	break;
+	case UIA_CulturePropertyId:
+		pRetVal->vt = VT_I4;
+		pRetVal->lVal = UIAManager->GetCachedCurrentLocaleLCID();
+		break;
 	case UIA_FrameworkIdPropertyId:
 		pRetVal->vt = VT_BSTR;
 		// todo: figure out what goes here
diff --git a/Engine/Source/Runtime/ApplicationCore/Public/Windows/Accessibility/WindowsUIAManager.h b/Engine/Source/Runtime/ApplicationCore/Public/Windows/Accessibility/WindowsUIAManager.h
--- a/Engine/Source/Runtime/ApplicationCore/Public/Windows/Accessibility/WindowsUIAManager.h
+++ b/Engine/Source/Runtime/ApplicationCore/Public/Windows/Accessibility/WindowsUIAManager.h
@@ -67,6 +67,7 @@ public:
 	 */
 	void OnAccessibleMessageHandlerChanged();
 
+	uint32 GetCachedCurrentLocaleLCID() const { return CachedCurrentLocaleLCID; }
 
 	static TMap<EAccessibleWidgetType, ULONG> WidgetTypeToWindowsTypeMap;
 	static TMap<EAccessibleWidgetType, FText> WidgetTypeToTextMap;
@@ -79,6 +80,15 @@ private:
 	/** Callback function for processing events raised from the AccessibleMessageHandler */
 	void OnEventRaised(TSharedRef<IAccessibleWidget> Widget, EAccessibleEvent Event, FVariant OldValue, FVariant NewValue);
 
+	/** Called when the first window widget provider is requested. */
+	void OnAccessibilityEnabled(); 
+
+	/** Called when the last widget provider is removed. */
+	void OnAccessibilityDisabled();
+
+	/** Updates the cached current locale LCID to match the current locale LCID in FInternationalization */
+	void UpdateCachedCurrentLocaleLCID();
+
 	/** Cache of all Providers with a RefCount of at least 1 that map to an accessible widget. */
 	TMap<TSharedRef<IAccessibleWidget>, FWindowsUIAWidgetProvider*> CachedWidgetProviders;
 	/**
@@ -88,6 +98,15 @@ private:
 	TSet<FWindowsUIABaseProvider*> ProviderList;
 	/** A reference back to the owning application that can be used to access the AccessibleMessageHandler */
 	const FWindowsApplication& WindowsApplication;
+
+	/**Handle associated with registering for culture change events from FInternationalization */
+	FDelegateHandle OnCultureChangedHandle; 
+	
+	/**Current LCID of the current locale in FInternationalization. Used to ensure screen readers 
+	* pronounce localized text in a native voice. 
+	* For more info: https://docs.microsoft.com/en-us/windows/win32/intl/language-identifier-constants-and-strings
+	*/
+	uint32 CachedCurrentLocaleLCID;
 };
 
 #endif
