Platform survey now has total and free space and it's actually consistent on Mac / Windows.

diff --git a/Engine/Source/Runtime/ApplicationCore/Private/Mac/MacPlatformSurvey.cpp b/Engine/Source/Runtime/ApplicationCore/Private/Mac/MacPlatformSurvey.cpp
--- a/Engine/Source/Runtime/ApplicationCore/Private/Mac/MacPlatformSurvey.cpp
+++ b/Engine/Source/Runtime/ApplicationCore/Private/Mac/MacPlatformSurvey.cpp
@@ -113,12 +113,16 @@ bool FMacPlatformSurvey::GetSurveyResults( FHardwareSurveyResults& OutResults, b
 	// Get CPU info
 	OutResults.CPUInfo = FMacPlatformMisc::GetCPUInfo();
 
-	// Get HDD details
+	// get HDD details
 	OutResults.HardDriveGB = -1;
-	NSDictionary *HDDAttributes = [[NSFileManager defaultManager] attributesOfFileSystemForPath: @"/" error: nil];
-	if (HDDAttributes)
+	OutResults.HardDriveFreeMB = -1;
+
+	uint64 TotalDiskSpace = 0;
+	uint64 FreeDiskSpace = 0;
+	if (FPlatformMisc::GetDiskTotalAndFreeSpace(FPlatformProcess::BaseDir(), TotalDiskSpace, FreeDiskSpace))
 	{
-		OutResults.HardDriveGB = (uint32)([[HDDAttributes objectForKey: NSFileSystemFreeSize] longLongValue] / 1024 / 1024 / 1024);
+		OutResults.HardDriveGB = (uint32)(TotalDiskSpace / uint64(1024 * 1024 * 1024));
+		OutResults.HardDriveFreeMB = (uint32)(FreeDiskSpace / uint64(1024 * 1024));
 	}
 	else
 	{
diff --git a/Engine/Source/Runtime/ApplicationCore/Private/Windows/WindowsPlatformSurvey.cpp b/Engine/Source/Runtime/ApplicationCore/Private/Windows/WindowsPlatformSurvey.cpp
--- a/Engine/Source/Runtime/ApplicationCore/Private/Windows/WindowsPlatformSurvey.cpp
+++ b/Engine/Source/Runtime/ApplicationCore/Private/Windows/WindowsPlatformSurvey.cpp
@@ -157,10 +157,14 @@ bool FWindowsPlatformSurvey::GetSurveyResults( FHardwareSurveyResults& OutResult
 
 	// get HDD details
 	OutResults.HardDriveGB = -1;
-	ULARGE_INTEGER TotalBytes;
-	if (GetDiskFreeSpaceEx(FPlatformProcess::BaseDir(), NULL, &TotalBytes, NULL))
+	OutResults.HardDriveFreeMB = -1;
+	
+	uint64 TotalDiskSpace = 0;
+	uint64 FreeDiskSpace = 0;
+	if (FPlatformMisc::GetDiskTotalAndFreeSpace(FPlatformProcess::BaseDir(), TotalDiskSpace, FreeDiskSpace))
 	{
-		OutResults.HardDriveGB = (TotalBytes.HighPart << 2) | (TotalBytes.LowPart >> 30);
+		OutResults.HardDriveGB = (uint32)(TotalDiskSpace / uint64(1024 * 1024 * 1024));
+		OutResults.HardDriveFreeMB = (uint32)(FreeDiskSpace / uint64(1024 * 1024));
 	}
 	else
 	{
diff --git a/Engine/Source/Runtime/ApplicationCore/Public/GenericPlatform/GenericPlatformSurvey.h b/Engine/Source/Runtime/ApplicationCore/Public/GenericPlatform/GenericPlatformSurvey.h
--- a/Engine/Source/Runtime/ApplicationCore/Public/GenericPlatform/GenericPlatformSurvey.h
+++ b/Engine/Source/Runtime/ApplicationCore/Public/GenericPlatform/GenericPlatformSurvey.h
@@ -187,6 +187,7 @@ struct FHardwareSurveyResults
 	TCHAR MultimediaAPI_DEPRECATED[MaxStringLength];
 
 	uint32 HardDriveGB;
+	uint32 HardDriveFreeMB;
 	uint32 MemoryMB;
 
 	float CPUPerformanceIndex;
