OpenGL on Mac (still used by the Launcher or CrashReportClient) doesn't allow zero sized windows, so for such windows set the size to (1,1) instead and set alpha to zero to hide the window.

diff --git a/Engine/Source/Runtime/ApplicationCore/Private/Mac/MacWindow.cpp b/Engine/Source/Runtime/ApplicationCore/Private/Mac/MacWindow.cpp
--- a/Engine/Source/Runtime/ApplicationCore/Private/Mac/MacWindow.cpp
+++ b/Engine/Source/Runtime/ApplicationCore/Private/Mac/MacWindow.cpp
@@ -627,7 +627,7 @@ void FMacWindow::ApplySizeAndModeChanges(int32 X, int32 Y, int32 Width, int32 He
 			Height = FMath::CeilToInt(Height / DPIScaleFactor);
 
 			const FVector2D CocoaPosition = FMacApplication::ConvertSlatePositionToCocoa(X, Y);
-			NSRect Rect = NSMakeRect(CocoaPosition.X, CocoaPosition.Y - Height + 1, FMath::Max(Width, 0), FMath::Max(Height, 0));
+			NSRect Rect = NSMakeRect(CocoaPosition.X, CocoaPosition.Y - Height + 1, FMath::Max(Width, 1), FMath::Max(Height, 1));
 			if (Definition->HasOSWindowBorder)
 			{
 				Rect = [WindowHandle frameRectForContentRect:Rect];
@@ -644,6 +644,9 @@ void FMacWindow::ApplySizeAndModeChanges(int32 X, int32 Y, int32 Width, int32 He
 						[WindowHandle setFrame:Rect display:YES];
 					}
 
+					const float WindowOpacity = (Definition->TransparencySupport == EWindowTransparency::PerWindow) ? Definition->Opacity : 1.0f;
+					[WindowHandle setAlphaValue:(Width > 0 && Height > 0) ? WindowOpacity : 0.0f];
+
 					if (Definition->ShouldPreserveAspectRatio)
 					{
 						[WindowHandle setContentAspectRatio:NSMakeSize((float)Width / (float)Height, 1.0f)];
