Accessibility: Introduced accessibility support for Mac. Mac VoiceOver navigation and interaction support for buttons, checkboxes, text and sliders now supported.

diff --git a/Engine/Source/Runtime/ApplicationCore/Private/Mac/Accessibility/CocoaAccessibilityView.cpp b/Engine/Source/Runtime/ApplicationCore/Private/Mac/Accessibility/CocoaAccessibilityView.cpp
new file mode 100644
--- /dev/null
+++ b/Engine/Source/Runtime/ApplicationCore/Private/Mac/Accessibility/CocoaAccessibilityView.cpp
@@ -0,0 +1,76 @@
+// Copyright 1998-2019 Epic Games, Inc. All Rights Reserved.
+
+#if WITH_ACCESSIBILITY
+#include "Mac/Accessibility/CocoaAccessibilityView.h"
+#include "Mac/Accessibility/MacAccessibilityManager.h"
+#include "Mac/Accessibility/MacAccessibilityElement.h"
+
+@implementation FCocoaAccessibilityView
+@synthesize AccessibleWindow;
+
+- (id)initWithFrame: (NSRect) Frame
+{
+	if(self = [super initWithFrame:Frame])
+	{
+		//init code goes here
+	}
+	return self;
+}
+
+- (void)dealloc
+{
+	self.AccessibleWindow = Nil;
+	[super dealloc];
+}
+
+- (void)SetAccessibilityWindowAsAccessibilityChild:(AccessibleWidgetId)InWindowId
+{
+	checkf(!IsInGameThread(), TEXT("Updating accessibility view from Game Thread! Accessibility  can only be done on Main Thread!"));
+	checkf(InWindowId != IAccessibleWidget::InvalidAccessibleWidgetId, TEXT("Cannot update view with invalid accessible widget Id! use RemoveAccessibilityWindow if trying to remove window."));
+	FMacAccessibilityElement* InAccessibleWindow = [[FMacAccessibilityManager AccessibilityManager]GetAccessibilityElement:InWindowId];
+	check(InAccessibleWindow != Nil);
+	//This ensures  that accessibility parenting is done properly and gives VoiceOVer the hierarchy it wants
+	// app > window (the NSView) > accessible elements
+	self.AccessibleWindow.accessibilityParent = self;
+	self.AccessibleWindow = InAccessibleWindow;
+}
+
+- (void)RemoveAccessibilityWindow
+{
+	checkf(!IsInGameThread(), TEXT("Clearing accessibility view from Game Thread! Accessibility  can only be done on Main Thread!"));
+	if(self.AccessibleWindow != Nil)
+	{
+		self.AccessibleWindow.accessibilityParent = Nil;
+		self.AccessibleWindow = Nil;
+	}
+}
+
+- (BOOL)isAccessibilityElement
+{
+	//we don't want the view to be interactable
+	return NO;
+}
+
+- (NSArray*)accessibilityChildren
+{
+	if(self.AccessibleWindow != Nil)
+	{
+		return self.AccessibleWindow.accessibilityChildren;
+	}
+	return Nil;
+}
+
+- (id)accessibilityFocusedUIElement
+{
+	return self.AccessibleWindow.accessibilityFocusedUIElement;
+}
+
+//@TODO: Implement hit testing properly for the window. Not sure if this is necessary though
+//- (id)accessibilityHitTest:(NSPoint) InPoint
+//{
+//	return [self.AccessibleWindow accessibilityHitTest: InPoint];
+//}
+
+@end
+
+#endif
diff --git a/Engine/Source/Runtime/ApplicationCore/Private/Mac/Accessibility/MacAccessibilityElement.cpp b/Engine/Source/Runtime/ApplicationCore/Private/Mac/Accessibility/MacAccessibilityElement.cpp
new file mode 100644
--- /dev/null
+++ b/Engine/Source/Runtime/ApplicationCore/Private/Mac/Accessibility/MacAccessibilityElement.cpp
@@ -0,0 +1,545 @@
+// Copyright 1998-2019 Epic Games, Inc. All Rights Reserved.
+
+#if WITH_ACCESSIBILITY
+
+#include "Mac/Accessibility/MacAccessibilityElement.h"
+#include "Mac/CocoaThread.h"
+#include "GenericPlatform/GenericAccessibleInterfaces.h"
+#include "Mac/Accessibility/MacAccessibilityManager.h"
+#include "Mac/MacApplication.h"
+#include "Mac/CocoaTextView.h"
+
+#import <Foundation/Foundation.h>
+
+@implementation FMacAccessibilityElement
+
+@synthesize Id;
+@synthesize ParentId;
+@synthesize OwningWindowId;
+@synthesize ChildIds;
+@synthesize Bounds;
+@synthesize Label;
+@synthesize Help;
+@synthesize Value;
+@synthesize LastCachedStringTime;
+
+- (id)initWithId:(AccessibleWidgetId)InId
+{
+	// Creating FMacAccessibilityElements shoulld only be done in MainThread
+	check([NSThread isMainThread]);
+	self = [super init];
+	if(self == Nil)
+	{
+		return self;
+	}
+	
+	self.Id = InId;
+	GameThreadCall(^{
+		TSharedPtr<IAccessibleWidget> Widget = [FMacAccessibilityManager AccessibilityManager].MacApplication->GetAccessibleMessageHandler()->GetAccessibleWidgetFromId(InId);
+		if (Widget.IsValid())
+		{
+			FString InitialLabel = Widget->GetWidgetName();
+			FString InitialHelp = Widget->GetHelpText();
+			FVariant InitialValue;
+			//@TODO: This ensures accessibility doesn't read out a person's password by giving an empty value. Do it properly
+			if (Widget->AsProperty() && !Widget->AsProperty()->IsPassword())
+			{
+				InitialValue = Widget->AsProperty()->GetValueAsVariant();
+			}
+			AccessibleWidgetId InitialParentId = IAccessibleWidget::InvalidAccessibleWidgetId;
+			TSharedPtr<IAccessibleWidget> InitialParent = Widget->GetParent();
+			if(InitialParent.IsValid())
+			{
+				InitialParentId = InitialParent->GetId();
+			}
+			AccessibleWidgetId InitialOwningWindowId = IAccessibleWidget::InvalidAccessibleWidgetId;
+			TSharedPtr<IAccessibleWidget> OwningWindow = Widget->GetWindow();
+			if(OwningWindow.IsValid())
+			{
+				InitialOwningWindowId = OwningWindow->GetId();
+			}
+			
+			EAccessibleWidgetType WidgetType = Widget->GetWidgetType();
+			MainThreadCall(^{
+				FMacAccessibilityElement* Self = [[FMacAccessibilityManager AccessibilityManager] GetAccessibilityElement:InId];
+				Self.Label = InitialLabel;
+				Self.Help = InitialHelp;
+				Self.Value = InitialValue;
+				Self.LastCachedStringTime = FPlatformTime::Seconds();
+				Self.ParentId = InitialParentId;
+				if(InitialParentId != IAccessibleWidget::InvalidAccessibleWidgetId)
+				{
+					self.accessibilityParent = [[FMacAccessibilityManager AccessibilityManager] GetAccessibilityElement:InitialParentId];
+				}
+				self.OwningWindowId = InitialOwningWindowId;
+				if(InitialOwningWindowId != IAccessibleWidget::InvalidAccessibleWidgetId)
+				{
+					self.accessibilityWindow = [[FMacAccessibilityManager AccessibilityManager] GetAccessibilityElement:InitialOwningWindowId];
+				}
+				
+				// The role and Subrole of an NSElement basically never changes, so set it here
+				//The Role and SubRoles together change the role descriptions for Voiceover to announce
+				//MOst of Slate Widgets don't need to support subroles so set to Nil here .
+				self.accessibilitySubrole = Nil;
+				switch(WidgetType)
+				{
+					case EAccessibleWidgetType::Button:
+					{
+						_AccessibilityRole = NSAccessibilityButtonRole;
+						break;
+					}
+					case EAccessibleWidgetType::CheckBox:
+					{
+						_AccessibilityRole = NSAccessibilityCheckBoxRole;
+						break;
+					}
+					case EAccessibleWidgetType::Image:
+					{
+						_AccessibilityRole = NSAccessibilityImageRole;
+						break;
+					}
+					case EAccessibleWidgetType::Slider:
+					{
+						_AccessibilityRole = NSAccessibilitySliderRole;
+						break;
+					}
+					case EAccessibleWidgetType::Text:
+					{
+						_AccessibilityRole = NSAccessibilityStaticTextRole;
+						break;
+					}
+					case EAccessibleWidgetType::TextEdit:
+					{
+						if(Widget->AsProperty()->IsReadOnly())
+						{
+							_AccessibilityRole = NSAccessibilityStaticTextRole;
+						}
+						else
+						{
+							_AccessibilityRole = NSAccessibilityTextFieldRole;
+							if(Widget->AsProperty()->IsPassword())
+							{
+								self.accessibilitySubrole = NSAccessibilitySecureTextFieldSubrole;
+							}
+						}
+						break;
+					}
+					case EAccessibleWidgetType::Window:
+					{
+						_AccessibilityRole = NSAccessibilityWindowRole;
+						break;
+					}
+					case EAccessibleWidgetType::Hyperlink:
+					{
+						_AccessibilityRole = NSAccessibilityLinkRole;
+						self.accessibilitySubrole = NSAccessibilityTextLinkSubrole;
+						break;
+					}
+					case EAccessibleWidgetType::Layout:
+					{
+						_AccessibilityRole = NSAccessibilityLayoutAreaRole;
+						break;
+					}
+					case EAccessibleWidgetType::ComboBox:
+					{
+						_AccessibilityRole = NSAccessibilityComboBoxRole;
+						break;
+					}
+					default:
+					{
+						//we use the NSAccessibilityGroupRole to forward VOiceover focus to children
+						//NSAccessibilityUnknownRole will intercept Voiceover Focus
+						//_AccessibilityRole = NSAccessibilityGroupRole;
+						_AccessibilityRole = NSAccessibilityUnknownRole;
+						break;
+					}
+				} // switch WidgetType
+				self.accessibilityRole = _AccessibilityRole;
+			}, NSDefaultRunLoopMode, false);
+		} // if widget valid
+	}, @[NSDefaultRunLoopMode], false);
+	
+	return self;
+}
+
+- (void)dealloc
+{
+	// Destroying of elements should only happen in Main Thread
+	check([NSThread isMainThread]);
+	if(_AccessibilityChildren != Nil)
+	{
+		[_AccessibilityChildren release];
+	}
+	NSAccessibilityPostNotification(self, NSAccessibilityUIElementDestroyedNotification);
+	[super dealloc];
+}
+
+- (bool)ShouldCacheStrings
+{
+	//@TODO: Magic number could be changed dynamically to suit responsiveness and performance
+	return FPlatformTime::Seconds() - LastCachedStringTime > 0.5;
+}
+
+- (void)CacheStrings
+{
+	// Updating of elements should only be done in Main Thread
+	check([NSThread isMainThread]);
+	if(![self ShouldCacheStrings])
+	{
+		return;
+	}
+	const AccessibleWidgetId TempId = Id;
+	GameThreadCall(^{
+		TSharedPtr<IAccessibleWidget> Widget = [FMacAccessibilityManager AccessibilityManager].MacApplication->GetAccessibleMessageHandler()->GetAccessibleWidgetFromId(TempId);
+		if(Widget.IsValid())
+		{
+			FString TempLabel = Widget->GetWidgetName();
+			FString TempHelp = Widget->GetHelpText();
+			FVariant TempValue;
+			//@TODO: Same as above.
+			if (Widget->AsProperty() && !Widget->AsProperty()->IsPassword())
+			{
+				TempValue = Widget->AsProperty()->GetValueAsVariant();
+			}
+			MainThreadCall(^{
+				FMacAccessibilityElement* Self = [[FMacAccessibilityManager AccessibilityManager] GetAccessibilityElement:TempId];
+				Self.Label = TempLabel;
+				Self.Help = TempHelp;
+				Self.Value = TempValue;
+				Self.LastCachedStringTime = FPlatformTime::Seconds();
+			}, NSDefaultRunLoopMode, false);
+		} // if widget valid
+	}, @[ NSDefaultRunLoopMode ], false);
+}
+
+- (void)ExposeToVoiceOver:(bool)bShouldExpose;
+{
+	// Editing of accessibility values should only happen in Main Thread
+	check([NSThread isMainThread]);
+	//Setting accessibilityRole to @"" makes VO ignore the element. It's a way to remove invisible elements from the accessibility tree. Not documented by Apple, so this may go away...
+	// https://stackoverflow.com/questions/31523333/disable-hide-accessibility-element
+	self.accessibilityRole = bShouldExpose ? _AccessibilityRole : @"";
+}
+
+- (BOOL)accessibilityNotifiesWhenDestroyed
+{
+	//this allows posting of accessibility notifications
+	return YES;
+}
+
+- (NSRect)accessibilityFrame
+{
+	//Bounds max and min are slate screen coordinates converted into cocoa screen coordinates
+	// Slate's Coordinate System is origin bottom left,+X right, +Y up
+	//Cocoa's Coordinate System is origin top left, +X right, +Y down
+	//Thus BOunds is not a valid AABB by cocoa coordinate space. We fix that here
+	return 	NSMakeRect(Bounds.Min.X, Bounds.Max.Y, Bounds.Max.X - Bounds.Min.X, Bounds.Min.Y - Bounds.Max.Y);
+}
+
+-(BOOL) isAccessibilityElement
+{
+	//this affects the element's exposure to the trackpad commander for Voiceover
+	return YES;
+}
+
+- (BOOL)isAccessibilityFocused
+{
+	//just check if the application focused element is this element
+	id ApplicationFocusedElement = NSApp.accessibilityApplicationFocusedUIElement;
+	if ([ApplicationFocusedElement isKindOfClass:[FMacAccessibilityElement class]])
+	{
+		FMacAccessibilityElement* MacElement = (FMacAccessibilityElement*) ApplicationFocusedElement;
+		return self.Id == MacElement.Id;
+	}
+	return NO;
+}
+
+- (void)setAccessibilityFocused:(BOOL)InIsFocused
+{
+	if(InIsFocused == YES)
+	{
+		// THis somehow contributes towards the ability to type in text fiellds...
+		//note that this holds a strong reference to the element and so the element won't be automatically released 
+		NSApp.accessibilityApplicationFocusedUIElement = self;
+		const AccessibleWidgetId TempId = self.Id;
+		GameThreadCall(^(){
+			TSharedPtr<IAccessibleWidget> Widget = [FMacAccessibilityManager AccessibilityManager].MacApplication->GetAccessibleMessageHandler()->GetAccessibleWidgetFromId(TempId);
+			if (Widget.IsValid() && Widget->IsEnabled() && !Widget->IsHidden())
+			{
+				Widget->SetFocus();
+			}
+		}, @[ NSDefaultRunLoopMode, UE4IMEEventMode, UE4ShowEventMode, UE4ResizeEventMode, UE4FullscreenEventMode], false); // GameThreadCall
+	}
+	[super setAccessibilityFocused:InIsFocused];
+}
+
+- (NSString*)accessibilityLabel
+{
+	if (!self.Label.IsEmpty())
+	{
+		return [NSString stringWithFString:self.Label];
+	}
+	return nil;
+}
+
+- (NSString*)accessibilityHelp
+{
+	if (!self.Help.IsEmpty())
+	{
+		return [NSString stringWithFString:self.Help];
+	}
+	return nil;
+}
+
+- (id)accessibilityValue
+{
+	if (Value.IsEmpty())
+	{
+		return Nil;
+	}
+	EVariantTypes VariantType = Value.GetType();
+	switch(VariantType)
+	{
+			//for checkboxes 
+		case EVariantTypes::Bool:
+		{
+			return Value.GetValue<bool>() ? @YES : @NO;
+		}
+			//for sliders
+		case EVariantTypes::Float:
+		{
+			return @(Value.GetValue<float>());
+		}
+			// for textedits and textboxes etc
+		case EVariantTypes::String:
+		{
+			//Passwords should be covered up
+			if(self.accessibilitySubrole == NSAccessibilitySecureTextFieldSubrole)
+			{
+				//@TODO: Consider just having the accessibility value in IAccessible Property be hidden
+				int32 PasswordLength = Value.GetValue<FString>().Len();
+				FString HiddenString;
+				HiddenString.Reserve(PasswordLength);
+				for(uint32 Index = 0; Index < PasswordLength; ++Index)
+				{
+					//the password black dot character
+					HiddenString += TCHAR(0x2202);
+				}
+				return HiddenString.GetNSString();
+			}
+			//return [NSString stringWithFString:Value.GetValue<FString>()];
+			return Value.GetValue<FString>().GetNSString();
+		}
+	} // switch VariantType
+	return Nil;
+}
+
+- (NSString*)accessibilityRoleDescription
+{
+	return NSAccessibilityRoleDescription(self.accessibilityRole, self.accessibilitySubrole);
+}
+
+- (NSArray*)accessibilityChildren
+{
+	if(ChildIds.Num() == 0)
+	{
+		return Nil;
+	}
+	//@TODO: Rather wasteful to clear it everytime it's requested. Probably an optimiazation available.
+	if(_AccessibilityChildren == Nil)
+	{
+		_AccessibilityChildren = [[NSMutableArray alloc] init];
+	}
+	else
+	{
+		[_AccessibilityChildren removeAllObjects];
+	}
+	for( const AccessibleWidgetId ChildId : ChildIds)
+	{
+		FMacAccessibilityElement* Child = [[FMacAccessibilityManager AccessibilityManager] GetAccessibilityElement:ChildId];
+		[_AccessibilityChildren addObject: Child];
+	}
+	return _AccessibilityChildren;
+}
+
+//Text
+- (NSInteger)accessibilityNumberOfCharacters
+{
+	//@TODO
+	return 0;
+}
+
+- (NSInteger)accessibilityInsertionPointLineNumber
+{
+	//@TODO
+	return 0;
+}
+
+- (NSRange)accessibilityVisibleCharacterRange
+{
+	//@TODO
+	return {NSNotFound, 0};;
+}
+
+- (NSString*)accessibilitySelectedText
+{
+	FMacAccessibilityElement* AccessibleWindow = (FMacAccessibilityElement*)self.accessibilityWindow;
+	id ParentView = AccessibleWindow.accessibilityParent;
+	if(![ParentView isKindOfClass: [FCocoaTextView class]])
+	{
+		return Nil;
+	}
+	FCocoaTextView* TextView = (FCocoaTextView*)ParentView;
+	NSAttributedString* AttributeString = [TextView attributedSubstringForProposedRange: [TextView selectedRange] actualRange: NULL];
+	return [AttributeString string];
+}
+
+- (NSRange)accessibilitySelectedTextRange
+{
+	FMacAccessibilityElement* AccessibleWindow = (FMacAccessibilityElement*)self.accessibilityWindow;
+	id ParentView = AccessibleWindow.accessibilityParent;
+	if(![ParentView isKindOfClass: [FCocoaTextView class]])
+	{
+		return {NSNotFound, 0};;
+	}
+	FCocoaTextView* TextView = (FCocoaTextView*)ParentView;
+	return [TextView selectedRange];
+}
+
+- (NSInteger)accessibilityLineForIndex:(NSInteger)index
+{
+	//TODO
+	return 0;
+}
+//Actions
+- (BOOL)accessibilityPerformPress
+{
+	const AccessibleWidgetId TempId = self.Id;
+	
+	// All IAccessibleWidget functions must be run on Game Thread
+	GameThreadCall(^(){
+		TSharedPtr<IAccessibleWidget> Widget = [FMacAccessibilityManager AccessibilityManager].MacApplication->GetAccessibleMessageHandler()->GetAccessibleWidgetFromId(TempId);
+		if (Widget.IsValid() && Widget->IsEnabled())
+		{
+			IAccessibleActivatable* Activatable = Widget->AsActivatable();
+			if(Activatable)
+			{
+				Activatable->Activate();
+				//if the activatable is checkable, it's a checkbox, post value updated notification
+				if(Activatable->IsCheckable())
+				{
+					MainThreadCall(^{
+						FMacAccessibilityElement* Self = [[FMacAccessibilityManager AccessibilityManager] GetAccessibilityElement:TempId];
+						NSAccessibilityPostNotification(Self, NSAccessibilityValueChangedNotification);
+					}, NSDefaultRunLoopMode, false);
+				}
+			}
+		} // if widget valid and enabled
+	}, @[ NSDefaultRunLoopMode, UE4IMEEventMode, UE4ShowEventMode, UE4ResizeEventMode, UE4FullscreenEventMode], false); // GameThreadCall
+	return YES;
+}
+
+- (BOOL) accessibilityPerformIncrement
+{
+	const AccessibleWidgetId TempId = self.Id;
+	// All IAccessibleWidget functions must be run on Game Thread
+	GameThreadCall(^{
+		TSharedPtr<IAccessibleWidget> Widget = [FMacAccessibilityManager AccessibilityManager].MacApplication->GetAccessibleMessageHandler()->GetAccessibleWidgetFromId(TempId);
+		if (Widget.IsValid())
+		{
+			IAccessibleProperty* Property = Widget->AsProperty();
+			if (Property && !FMath::IsNearlyZero(Property->GetStepSize()) && !Property->IsReadOnly())
+			{
+				const float CurrentValue = FCString::Atof(*Property->GetValue());
+				float NewValue = CurrentValue + Property->GetStepSize();
+				if(NewValue > Property->GetMaximum())
+				{
+					NewValue = Property->GetMaximum();
+				}
+				Property->SetValue(FString::SanitizeFloat(NewValue));
+				MainThreadCall(^{
+					FMacAccessibilityElement* Self = [[FMacAccessibilityManager AccessibilityManager] GetAccessibilityElement:TempId];
+					NSAccessibilityPostNotification(Self, NSAccessibilityValueChangedNotification);
+				}, NSDefaultRunLoopMode, false);
+			}
+		} // widget valid
+	}, @[ NSDefaultRunLoopMode ], false);
+	return YES;
+}
+
+- (BOOL)accessibilityPerformDecrement
+{
+	const AccessibleWidgetId TempId = self.Id;
+	// All IAccessibleWidget functions must be run on Game Thread
+	GameThreadCall(^{
+		TSharedPtr<IAccessibleWidget> Widget = [FMacAccessibilityManager AccessibilityManager].MacApplication->GetAccessibleMessageHandler()->GetAccessibleWidgetFromId(TempId);
+		if (Widget.IsValid())
+		{
+			IAccessibleProperty* Property = Widget->AsProperty();
+			if (Property && !FMath::IsNearlyZero(Property->GetStepSize()) && !Property->IsReadOnly())
+			{
+				const float CurrentValue = FCString::Atof(*Property->GetValue());
+				float NewValue = CurrentValue - Property->GetStepSize();
+				if(NewValue < Property->GetMinimum())
+				{
+					NewValue = Property->GetMinimum();
+				}
+				Property->SetValue(FString::SanitizeFloat(NewValue));
+				MainThreadCall(^{
+					FMacAccessibilityElement* Self = [[FMacAccessibilityManager AccessibilityManager] GetAccessibilityElement:TempId];
+					NSAccessibilityPostNotification(Self, NSAccessibilityValueChangedNotification);
+				}, NSDefaultRunLoopMode, false);
+			}
+		} // widget valid
+	}, @[ NSDefaultRunLoopMode ], false);
+	return YES;
+}
+
+//Returns YES if the selector can be called on this accessibility element
+- (BOOL)isAccessibilitySelectorAllowed:(SEL)selector
+{
+	if(selector == @selector(accessibilityValue))
+	{
+		return _AccessibilityRole == NSAccessibilityCheckBoxRole
+		|| _AccessibilityRole == NSAccessibilityTextFieldRole
+		|| _AccessibilityRole == NSAccessibilityStaticTextRole
+		|| _AccessibilityRole == NSAccessibilitySliderRole;
+		
+	}
+	//Text
+	else if(selector == @selector(accessibilityInsertionPointLineNumber))
+	{
+		return _AccessibilityRole == NSAccessibilityTextFieldRole;
+	}
+	else if(selector == @selector(accessibilityLineForIndex:))
+	{
+		return _AccessibilityRole == NSAccessibilityTextFieldRole;
+	}
+	else if(selector == @selector(accessibilityRangeForLine:))
+	{
+		return _AccessibilityRole == NSAccessibilityTextFieldRole;
+	}
+	else if(selector == @selector(accessibilityFrameForRange))
+	{
+		return _AccessibilityRole == NSAccessibilityTextFieldRole;
+	}
+	//Actions
+	else if(selector == @selector(accessibilityPerformPress))
+	{
+		return _AccessibilityRole == NSAccessibilityButtonRole
+		|| _AccessibilityRole == NSAccessibilityCheckBoxRole;
+	}
+	else if (selector == @selector(accessibilityPerformDecrement))
+	{
+		return _AccessibilityRole == NSAccessibilitySliderRole;
+	}
+	else if (selector == @selector(accessibilityPerformIncrement))
+	{
+		return _AccessibilityRole == NSAccessibilitySliderRole;
+	}
+	
+	return YES;
+}
+@end
+
+#endif 
diff --git a/Engine/Source/Runtime/ApplicationCore/Private/Mac/Accessibility/MacAccessibilityManager.cpp b/Engine/Source/Runtime/ApplicationCore/Private/Mac/Accessibility/MacAccessibilityManager.cpp
new file mode 100644
--- /dev/null
+++ b/Engine/Source/Runtime/ApplicationCore/Private/Mac/Accessibility/MacAccessibilityManager.cpp
@@ -0,0 +1,313 @@
+// Copyright 1998-2019 Epic Games, Inc. All Rights Reserved.
+
+#if WITH_ACCESSIBILITY
+
+#include "Mac/Accessibility/MacAccessibilityManager.h"
+#include "Mac/Accessibility/MacAccessibilityElement.h"
+#include "Mac/MacApplication.h"
+#include "HAL/IConsoleManager.h"
+#include "Mac/CocoaThread.h"
+
+#import <AppKit/AppKit.h>
+
+//used for registering for NSWorkspace's KVO notifications when Voiceover is turned on
+static void* MacAccessibilityManagerContext = &MacAccessibilityManagerContext;
+
+@implementation FMacAccessibilityManager
+@synthesize MacApplication;
+
+#if !UE_BUILD_SHIPPING
+static void DumpAccessibilityStatsForwarder()
+{
+	[[FMacAccessibilityManager AccessibilityManager] DumpAccessibilityStats];
+}
+#endif
+
+- (id)init
+{
+	// This should only be created in Main Thread
+	check([NSThread isMainThread]);
+	if(self = [super init])
+	{
+		Cache = [[NSMutableDictionary alloc] init];
+		//register KVO for notifications about changes in  Voiceover status.
+		// A notification is received every time Voiceover status changes.
+		//@see: observeValueForKeyPath
+		//NOTE: Don't use the NSKeyValueObservingOptionInitial option.
+		// If VoiceOVer were already enabled before the engine launched,
+		//It would send a KVO notification before the FSlateAccessibilityMessageHandler is set in FMacApplication.
+		[[NSWorkspace sharedWorkspace] addObserver:self
+										forKeyPath: @"voiceOverEnabled"
+										   options: NSKeyValueObservingOptionNew
+										   context: MacAccessibilityManagerContext];
+	}
+	return self;
+}
+
+- (void)dealloc
+{
+	[super dealloc];
+}
+
+- (void)observeValueForKeyPath:(NSString *)keyPath
+					 ofObject:(id)object
+					   change:(NSDictionary*)change
+					  context:(void*)context
+{
+	if(context == MacAccessibilityManagerContext)
+	{
+		if ([keyPath isEqualToString: @"voiceOverEnabled"])
+		{
+			NSNumber* newValue = [change objectForKey:NSKeyValueChangeNewKey];
+			bool bVoiceoverEnabled = [newValue boolValue];
+			if(bVoiceoverEnabled)
+			{
+				//@TODO: This can cause VoiceOVer to hang sometimes. Disabling for now.
+				//MacApplication->OnVoiceoverEnabled();
+			}
+			else
+			{
+				//@TODO: This causes VoiceOVer to hang when turning it on and off sometimes. Disabling for now
+				//This means that the accessibility cache is only cleared on shutdown and will persist even if VoiceOVer is off.
+				//MacApplication->OnVoiceoverDisabled();
+			}
+		} // if keypath
+	} //if context
+	else
+	{
+		[super observeValueForKeyPath:keyPath ofObject:object change:change context:context];
+	}
+}
+
+- (FMacAccessibilityElement*)GetAccessibilityElement:(AccessibleWidgetId)Id
+{
+	// Getting and creating FMacAccessibilityElements should only be done fromm Main Thread
+	check([NSThread isMainThread]);
+	if (Id == IAccessibleWidget::InvalidAccessibleWidgetId)
+	{
+		return nil;
+	}
+	FMacAccessibilityElement* ExistingElement = [Cache objectForKey:@(Id)];
+	if (ExistingElement == nil)
+	{
+		ExistingElement = [[[FMacAccessibilityElement alloc] initWithId:Id] autorelease];
+		[Cache setObject:ExistingElement forKey:@(Id)];
+	}
+	return ExistingElement;
+}
+
+- (bool)AccessibilityElementExists:(AccessibleWidgetId)Id
+{
+	return [Cache objectForKey:@(Id)] != nil;
+}
+
+- (void)RemoveAccessibilityElement:(AccessibleWidgetId)Id
+{
+	//Modifications to the cache should only be done in ain Thread
+	check([NSThread isMainThread]);
+	[Cache removeObjectForKey:@(Id)];
+}
+
+- (void)RemoveAccessibilitySubtree:(AccessibleWidgetId)RootId
+{
+	check([NSThread isMainThread]);
+	//BFS to remove entire subtree
+	TArray<AccessibleWidgetId> Queue;
+	Queue.Push(RootId);
+	while(Queue.Num() > 0)
+	{
+		AccessibleWidgetId CurrentId = Queue.Pop(false);
+		FMacAccessibilityElement* CurrentElement = [Cache objectForKey:@(RootId)];
+		if(CurrentElement.ChildIds.Num() == 0)
+		{
+			[self RemoveAccessibilityElement:CurrentId];
+			continue;
+		}
+		for(AccessibleWidgetId ChildId : CurrentElement.ChildIds)
+		{
+			Queue.Push(ChildId);
+		}
+		[self RemoveAccessibilityElement:CurrentId];
+	}
+}
+
+- (void)Clear
+{
+	check([NSThread isMainThread]);
+	[Cache removeAllObjects];
+}
+
+- (int32)GetAccessibilityCacheSize
+{
+	return [Cache count];
+}
+
+- (bool)IsAccessibilityCacheEmpty
+{
+	return [Cache count] == 0;
+}
+
+- (void)TearDown
+{
+	//This should only be called in ~FMacApplication() in a Main Thread Call
+	check([NSThread isMainThread]);
+	checkf([self IsAccessibilityCacheEmpty], TEXT("Accessibility Manager is being torn down and still has %yu elements!"), [self GetAccessibilityCacheSize]);
+	[Cache release];
+	Cache = Nil;
+	//unregister KVO for Voiceover changes
+	[[NSWorkspace sharedWorkspace] removeObserver:self
+									   forKeyPath:@"voiceOverEnabled"];
+	//Note: do not set the appllication to null as multithreading could cause the MacApplication destructor to be called before the update timer is destroyed.
+}
+
++ (FMacAccessibilityManager*)AccessibilityManager
+{
+	static FMacAccessibilityManager* manager = nil;
+	if (manager == nil)
+	{
+		manager = [[self alloc] init];
+	}
+	return manager;
+}
+
+- (void)UpdateAllCachedProperties
+{
+	//Updating of FMacAccessibilityElements should only be done from Main Thread
+	check([NSThread isMainThread]);
+	TArray<AccessibleWidgetId> Ids;
+	for (NSString* Key in Cache)
+	{
+		Ids.Add(Key.intValue);
+	}
+	if(Ids.Num() == 0)
+	{
+		return;
+	}
+	//all requests to slate accessible widget cache must be done from game thread
+	//Do not update parent here as the Mac Application is supposed to handle a EAccessibleEvent::ParentChanged event to do that
+	GameThreadCall(^{
+		for (const AccessibleWidgetId Id : Ids)
+		{
+			TSharedPtr<IAccessibleWidget> Widget = self.MacApplication->GetAccessibleMessageHandler()->GetAccessibleWidgetFromId(Id);
+			if (Widget.IsValid())
+			{
+				// Children
+				TArray<int32> ChildIds;
+				for (int32 i = 0; i < Widget->GetNumberOfChildren(); ++i)
+				{
+					TSharedPtr<IAccessibleWidget> Child = Widget->GetChildAt(i);
+					if (Child.IsValid())
+					{
+						ChildIds.Add(Child->GetId());
+					}
+				}
+				// Bounding rect in Slate screen space  .
+				//Slate origin is bottom left, +X right, +Y up
+				// Cocoa origin is top left, +X right, +Y down
+				FBox2D Bounds = Widget->GetBounds();
+				
+				//NOTE: This is just converting Slate screen coordinates to Cocoa screen coordinates
+				//Because Slate and Cocoa have different coordinate systems, Bounds is no longer a valid AABB
+				Bounds.Min = FMacApplication::ConvertSlatePositionToCocoa(Bounds.Min.X, Bounds.Min.Y);
+				Bounds.Max = FMacApplication::ConvertSlatePositionToCocoa(Bounds.Max.X, Bounds.Max.Y);
+				// Visibility
+				const bool bIsEnabled = Widget->IsEnabled();
+				const bool bIsVisible = !Widget->IsHidden();
+				// All AppKit functions must be run on Main Thread
+				MainThreadCall(^{
+					FMacAccessibilityElement* Element = [[FMacAccessibilityManager AccessibilityManager] GetAccessibilityElement:Id];
+					Element.ChildIds = ChildIds;
+					Element.Bounds = Bounds;
+					Element.accessibilityEnabled = bIsEnabled ? YES : NO;
+					//tells Voiceover to ignore elements based on visibility
+					[Element ExposeToVoiceOver:bIsVisible];
+					//@TODO: Updating the label, help and value of ALL elements is suboptimal...but it works
+					//Tested with EngineTest in Debug Editor in the Editor with ~650 elements and still fairly responsive
+					[Element CacheStrings];
+				}, NSDefaultRunLoopMode, true);
+			} // if Widget valid
+		} // for Ids loop
+	}, @[ NSDefaultRunLoopMode, UE4IMEEventMode, UE4ShowEventMode, UE4ResizeEventMode, UE4FullscreenEventMode, UE4NilEventMode ], true);
+}
+
+#if !UE_BUILD_SHIPPING
+- (void)DumpAccessibilityStats
+{
+	const uint32 NumElements = [Cache count];
+	uint32 SizeOfElement = 0;
+	uint32 CacheSize = 0;
+	
+	NSArray* Keys = [Cache allKeys];
+	for (NSString* Key in Keys)
+	{
+		FMacAccessibilityElement* Element = [Cache objectForKey: Key];
+		
+		SizeOfElement = malloc_size(Element);
+		CacheSize += sizeof(NSString*) + sizeof(FMacAccessibilityElement*) + malloc_size(Key) + SizeOfElement
+		+ malloc_size([NSString stringWithFString : Element.Label])
+		+ malloc_size([NSString stringWithFString : Element.Help])
+		+ malloc_size([NSString stringWithFString : Element.Value]);
+		
+	}
+	
+	UE_LOG(LogAccessibility, Log, TEXT("Number of Accessibility Elements: %i"), NumElements );
+	UE_LOG(LogAccessibility, Log, TEXT("Size of FMacAccessibilityElement: %u"), SizeOfElement);
+	UE_LOG(LogAccessibility, Log, TEXT("Memory stored in cache: %u kb"), CacheSize / 1000);
+}
+
+- (void)PrintCache
+{
+	TArray<AccessibleWidgetId> Ids;
+	for (NSString* Key in Cache)
+	{
+		Ids.Add(Key.intValue);
+	}
+	if(Ids.Num() == 0)
+	{
+		UE_LOG(LogAccessibility, Log, TEXT("No Cache is empty."));
+		return;
+	}
+	Ids.Sort();
+	for(AccessibleWidgetId Id : Ids)
+	{
+		[self PrintAccessibilityElementById:Id];
+	}
+}
+
+- (void)PrintAccessibilityElement: (FMacAccessibilityElement*) Element
+{
+	UE_LOG(LogAccessibility, Log, TEXT("ID: %i"), Element.Id);
+	UE_LOG(LogAccessibility, Log, TEXT("Parent ID: %i"), Element.ParentId);
+	UE_LOG(LogAccessibility, Log, TEXT("Label: %s"), *Element.Label);
+	UE_LOG(LogAccessibility, Log, TEXT("Children cOunt: %u"), Element.ChildIds.Num());
+}
+
+- (void)PrintAccessibilityElementById:(AccessibleWidgetId)Id
+{
+	FMacAccessibilityElement* Element = [Cache objectForKey:@(Id)];
+	if(Element == Nil)
+	{
+		UE_LOG(LogAccessibility, Log, TEXT("Element with ID %i is not in the cache."), Id);
+	}
+	[self PrintAccessibilityElement:Element];
+}
+
+- (id)GetAccessibilityParent:(FMacAccessibilityElement*)Element
+{
+	return (Element != Nil) ? Element.accessibilityParent : Nil;
+}
+
+- (FMacAccessibilityElement*)GetAccessibilityRoot
+{
+	return Nil;
+}
+
+- (bool)isElementOrphaned:(FMacAccessibilityElement*) Element
+{
+	return Element.accessibilityParent == Nil;
+}
+#endif
+
+@end
+
+#endif
diff --git a/Engine/Source/Runtime/ApplicationCore/Private/Mac/CocoaWindow.cpp b/Engine/Source/Runtime/ApplicationCore/Private/Mac/CocoaWindow.cpp
--- a/Engine/Source/Runtime/ApplicationCore/Private/Mac/CocoaWindow.cpp
+++ b/Engine/Source/Runtime/ApplicationCore/Private/Mac/CocoaWindow.cpp
@@ -6,7 +6,9 @@
 #include "Mac/CocoaThread.h"
 #include "Mac/MacCursor.h"
 #include "HAL/IConsoleManager.h"
-
+#if WITH_ACCESSIBILITY
+#include "Mac/Accessibility/CocoaAccessibilityView.h"
+#endif
 NSString* NSDraggingExited = @"NSDraggingExited";
 NSString* NSDraggingUpdated = @"NSDraggingUpdated";
 NSString* NSPrepareForDragOperation = @"NSPrepareForDragOperation";
@@ -543,4 +545,31 @@ NSString* NSPerformDragOperation = @"NSPerformDragOperation";
 	return YES;
 }
 
+#if WITH_ACCESSIBILITY
+
+- (void)UpdateAccessibilityView:(AccessibleWidgetId) InAccessibilityWindowId
+{
+	checkf(!IsInGameThread(), TEXT("Updating accessibility view in FCocoaWindow from Game Thread! Accessibility  can only be done on Main Thread!"));
+	NSView* OpenGLView = [self openGLView];
+	// FCocoaAccessibilityView is the base class all custom NSViews must inherit from to support accessibility
+	if([OpenGLView isKindOfClass:[FCocoaAccessibilityView class]])
+	{
+		FCocoaAccessibilityView* CurrentAccessibilityView = (FCocoaAccessibilityView*) OpenGLView;
+		[CurrentAccessibilityView SetAccessibilityWindowAsAccessibilityChild: InAccessibilityWindowId];
+	}
+}
+
+- (void)ClearAccessibilityView
+{
+	checkf(!IsInGameThread(), TEXT("Updating accessibility view in FCocoaWindow from Game Thread! Accessibility  can only be done on Main Thread!"));
+	NSView* OpenGLView = [self openGLView];
+	// FCocoaAccessibilityView is the base class all custom NSViews must inherit from to support accessibility
+	if([OpenGLView isKindOfClass:[FCocoaAccessibilityView class]])
+	{
+		FCocoaAccessibilityView* CurrentAccessibilityView = (FCocoaAccessibilityView*) OpenGLView;
+		[CurrentAccessibilityView RemoveAccessibilityWindow];
+	}
+}
+#endif
+
 @end
diff --git a/Engine/Source/Runtime/ApplicationCore/Private/Mac/MacApplication.cpp b/Engine/Source/Runtime/ApplicationCore/Private/Mac/MacApplication.cpp
--- a/Engine/Source/Runtime/ApplicationCore/Private/Mac/MacApplication.cpp
+++ b/Engine/Source/Runtime/ApplicationCore/Private/Mac/MacApplication.cpp
@@ -19,6 +19,11 @@
 #include "HAL/ThreadHeartBeat.h"
 #include "IHapticDevice.h"
 #include "Apple/ApplePlatformCrashContext.h"
+#if WITH_ACCESSIBILITY
+#include "Mac/Accessibility/MacAccessibilityManager.h"
+#include "Mac/Accessibility/MacAccessibilityElement.h"
+#endif
+
 
 #include <IOKit/IOKitLib.h>
 #include <IOKit/graphics/IOGraphicsLib.h>
@@ -64,7 +69,10 @@ FMacApplication::FMacApplication()
 ,	bEmulatingRightClick(false)
 ,	bIgnoreMouseMoveDelta(0)
 ,	bIsWorkspaceSessionActive(true)
-, 	KeyBoardLayoutData(nil)
+,	KeyBoardLayoutData(nil)
+#if WITH_ACCESSIBILITY
+,	AccessibilityCacheTimer(Nil)
+#endif
 {
 	TextInputMethodSystem = MakeShareable(new FMacTextInputMethodSystem);
 	if (!TextInputMethodSystem->Initialize())
@@ -135,6 +143,18 @@ FMacApplication::~FMacApplication()
 	}
 
 	MainThreadCall(^{
+#if WITH_ACCESSIBILITY
+		if(GetAccessibleMessageHandler()->IsActive())
+		{
+			OnVoiceoverDisabled();
+		}
+		if(AccessibilityCacheTimer != Nil)
+		{
+			[AccessibilityCacheTimer release];
+		}
+		[[FMacAccessibilityManager AccessibilityManager] TearDown];
+#endif
+
 		if (MouseMovedEventMonitor)
 		{
 			[NSEvent removeMonitor:MouseMovedEventMonitor];
@@ -192,6 +212,24 @@ void FMacApplication::SetMessageHandler(const TSharedRef<FGenericApplicationMess
 	HIDInput->SetMessageHandler(InMessageHandler);
 }
 
+#if WITH_ACCESSIBILITY
+void FMacApplication::SetAccessibleMessageHandler(const TSharedRef<FGenericAccessibleMessageHandler>& InAccessibleMessageHandler)
+{
+	GenericApplication::SetAccessibleMessageHandler(InAccessibleMessageHandler);
+	InAccessibleMessageHandler->SetAccessibleEventDelegate(FGenericAccessibleMessageHandler::FAccessibleEvent::CreateRaw(this, &FMacApplication::OnAccessibleEventRaised));
+	
+	MainThreadCall(^{
+		// Initializing FMacAccessibilityManager singleton to enable KVO registration for Voiceover changes
+		[FMacAccessibilityManager AccessibilityManager].MacApplication = this;
+		const bool bVoiceoverEnabled = [NSWorkspace sharedWorkspace].isVoiceOverEnabled;
+		if(bVoiceoverEnabled)
+		{
+			OnVoiceoverEnabled();
+		}
+	}, NSDefaultRunLoopMode, false);
+}
+#endif
+
 void FMacApplication::PollGameDeviceState(const float TimeDelta)
 {
 	// initialize any externally-implemented input devices (we delay load initialize the array so any plugins have had time to load)
@@ -1155,7 +1193,15 @@ bool FMacApplication::OnWindowDestroyed(TSharedRef<FMacWindow> DestroyedWindow)
 	{
 		OnWindowActivationChanged(DestroyedWindow, EWindowActivation::Deactivate);
 	}
-
+	
+#if WITH_ACCESSIBILITY
+	MainThreadCall(^{
+		//This is to clear the FMacAccessibilityElement* that represents the cocoa window in the FCocoaAccessibilityView.
+		// Before the FCocoaAccessibilityView is deallocated, it could be a danling pointer and accessibility calls from VoiceOVer will cause a crash
+		[WindowHandle ClearAccessibilityView];
+	}, NSDefaultRunLoopMode, false);
+#endif
+	
 	{
 		FScopeLock Lock(&WindowsMutex);
 		Windows.Remove(DestroyedWindow);
@@ -1188,7 +1234,7 @@ bool FMacApplication::OnWindowDestroyed(TSharedRef<FMacWindow> DestroyedWindow)
 	if (WindowToActivate.IsValid())
 	{
 		WindowToActivate->SetWindowFocus();
-	}
+			}
 
 	MessageHandler->OnCursorSet();
 
@@ -1220,6 +1266,29 @@ void FMacApplication::OnWindowOrderedFront(TSharedRef<FMacWindow> Window)
 		}
 	}
 	Windows = NewWindowsArray;
+	
+#if WITH_ACCESSIBILITY
+	// The IMainFrameModule creates an SWindow and tries to show it.
+	// When the SWindow calls ShowWindow(), the SlateRHIRenderer is initialized and creates a FMetalViewport.
+	// THe FMetalViewport creates a custom FMetalView NSView and sets it as the content view of the window, overriding any accessibility data on the original content view.
+	// THis happens AFTER we call OnVoiceoverEnabled() and set all the window IDs the first time
+	//We set the accessible Window ID here again to update the FMetalView NSView with accessibility children information
+	// @see FMetalViewport::FMetalViewport()
+	//@see SWindow::ShowWindow()
+	//@Review: Is there a better place for this?
+	if(GetAccessibleMessageHandler()->IsActive())
+	{
+		const AccessibleWidgetId WindowId = GetAccessibleMessageHandler()->GetAccessibleWindowId(Window);
+		MainThreadCall(^{
+			FCocoaWindow* CocoaWindow = Window->GetWindowHandle();
+			if(CocoaWindow)
+			{
+				[CocoaWindow UpdateAccessibilityView:WindowId];
+			}
+		}, NSDefaultRunLoopMode, false);
+	}
+#endif
+
 }
 
 void FMacApplication::OnWindowActivationChanged(const TSharedRef<FMacWindow>& Window, const EWindowActivation ActivationType)
@@ -2157,3 +2226,130 @@ void FDisplayMetrics::RebuildDisplayMetrics(FDisplayMetrics& OutDisplayMetrics)
 	// Apply the debug safe zones
 	OutDisplayMetrics.ApplyDefaultSafeZones();
 }
+
+#if WITH_ACCESSIBILITY
+void FMacApplication::OnAccessibleEventRaised(TSharedRef<IAccessibleWidget> Widget, EAccessibleEvent Event, FVariant OldValue, FVariant NewValue)
+{
+	// This should only be triggered by the accessible message handler which initiates from the Slate thread.
+	check(IsInGameThread());
+	
+	const AccessibleWidgetId Id = Widget->GetId();
+	switch (Event)
+	{
+		case EAccessibleEvent::FocusChange:
+		{
+			//@TODO: Posting accessibility focus notifications don't seem to do anything, investigate further
+			break;
+		}
+		case EAccessibleEvent::ParentChanged:
+		{
+			MainThreadCall(^{
+				FMacAccessibilityElement* Element = [[FMacAccessibilityManager AccessibilityManager] GetAccessibilityElement:Id];
+				if(Element)
+				{
+					const AccessibleWidgetId NewParentId = NewValue.GetValue<AccessibleWidgetId>();
+					Element.ParentId = NewParentId;
+					FMacAccessibilityElement* NewParent = [[FMacAccessibilityManager AccessibilityManager]GetAccessibilityElement:NewParentId];
+					Element.accessibilityParent = NewParent;
+
+					//if the element is orphaned by having a nill parent, we'll remove the entire subtree for now from the accessibility cache
+					//@TODO: Widget switcher unparents its children, but this is never executed. Strange
+					if(NewParent == Nil)
+					{
+						[[FMacAccessibilityManager AccessibilityManager] RemoveAccessibilitySubtree:Id];
+					}
+					
+
+				}
+				// LayoutChanged is to indicate things like "a widget became visible or hidden" while
+				// ScreenChanged is for large-scale UI changes. It can potentially take an NSString to read
+				// to the user when this happens, if we choose to support that.
+				//@TODO: Posting a notification with user info to the NSApp doesn't seem to do anything. Find out wy.
+			}, NSDefaultRunLoopMode, false);
+			break;
+		}
+		case EAccessibleEvent::WidgetRemoved:
+		{
+			MainThreadCall(^{
+				[[FMacAccessibilityManager AccessibilityManager] RemoveAccessibilityElement:Id];
+			}, NSDefaultRunLoopMode, false);
+			break;
+		}
+		case EAccessibleEvent::Notification:
+		{
+			//@TODO: Allow strings to be announced to the player
+			break;
+		}
+	} // switch event
+}
+
+void FMacApplication::OnVoiceoverEnabled()
+{
+	//Alll accessibility functions should originate from Main Thread
+	check([NSThread isMainThread]);
+	if(GetAccessibleMessageHandler()->IsActive())
+	{
+		return;
+	}
+	GetAccessibleMessageHandler()->SetActive(true);
+	// Retrieving Slate accessibility data needs to be done on the game thread
+	GameThreadCall(^{
+		FScopeLock Lock(&WindowsMutex);
+		TArray<AccessibleWidgetId> WindowIds;
+		for(const TSharedRef<FMacWindow>& window : Windows)
+		{
+			AccessibleWidgetId Id = GetAccessibleMessageHandler()->GetAccessibleWindowId(window);
+			WindowIds.Add(Id);
+		}
+		// All AppKit functions need to be called from Main Thread
+		MainThreadCall(^{
+			for(int WindowIndex = 0; WindowIndex < Windows.Num(); ++WindowIndex)
+			{
+				const TSharedRef<FMacWindow> CurrentWindow = Windows[WindowIndex];
+				FCocoaWindow* CurrentCocoaWindow = CurrentWindow->GetWindowHandle();
+				if(CurrentCocoaWindow)
+				{
+					[CurrentCocoaWindow UpdateAccessibilityView:WindowIds[WindowIndex]];
+				}
+			}// for all windows
+			//Start caching Mac Accessibility data to be returned to Voiceover upon request
+			// When Voiceover is enabled, the accessibility tree will take a while to build and
+			//accessibility may not work properly till then.
+			if (AccessibilityCacheTimer == nil)
+			{
+				AccessibilityCacheTimer = [NSTimer scheduledTimerWithTimeInterval:0.25f target:[FMacAccessibilityManager AccessibilityManager] selector:@selector(UpdateAllCachedProperties) userInfo:nil repeats:YES];
+			}
+		}, NSDefaultRunLoopMode, false);
+	}, @[ NSDefaultRunLoopMode ], false);
+}
+
+void FMacApplication::OnVoiceoverDisabled()
+{
+	// Accessibility should originate from Main Thread
+	check([NSThread isMainThread]);
+	if(!GetAccessibleMessageHandler()->IsActive())
+	{
+		return;
+	}
+	GetAccessibleMessageHandler()->SetActive(false);
+	MainThreadCall(^{
+		{
+			FScopeLock Lock(&WindowsMutex);
+			for(const TSharedRef<FMacWindow>& CurrentWindow : Windows)
+			{
+				FCocoaWindow* CurrentCocoaWindow = CurrentWindow->GetWindowHandle();
+				if(CurrentCocoaWindow)
+				{
+					[CurrentCocoaWindow ClearAccessibilityView];
+				}
+			}
+		}
+		[AccessibilityCacheTimer invalidate];
+		AccessibilityCacheTimer = nil;
+		[[FMacAccessibilityManager AccessibilityManager] Clear];
+		//This releases the accessibility element referenced by the app, deallocating it
+		NSApp.accessibilityApplicationFocusedUIElement = Nil;
+	}, NSDefaultRunLoopMode, false);
+}
+
+#endif
diff --git a/Engine/Source/Runtime/ApplicationCore/Public/GenericPlatform/GenericAccessibleInterfaces.h b/Engine/Source/Runtime/ApplicationCore/Public/GenericPlatform/GenericAccessibleInterfaces.h
--- a/Engine/Source/Runtime/ApplicationCore/Public/GenericPlatform/GenericAccessibleInterfaces.h
+++ b/Engine/Source/Runtime/ApplicationCore/Public/GenericPlatform/GenericAccessibleInterfaces.h
@@ -222,6 +222,12 @@ public:
 	 * @return A string representing the serialized value stored in the widget.
 	 */
 	virtual FString GetValue() const = 0;
+	/**
+	 * The current value stored by the widget as a FVariant. Used on platforms where the type of the value is important for accessibility methods
+	 * such as on Mac. .
+	 * @return An FVariant that holds the widget's current value and type of value.
+	 */
+	virtual FVariant GetValueAsVariant() const { return FVariant(); }
 	/*
 	 * Set the value stored by the widget. While this function accepts a String, there is no way to know
 	 * what the underlying data is stored as. The platform layer must retain some additional information
@@ -347,7 +353,7 @@ public:
 	 * @return A more-detailed description of what the widget is or how its used.
 	 */
 	virtual FString GetHelpText() const = 0;
-	
+
 	/**
 	 * Whether the widget is enabled and can be interacted with.
 	 *
diff --git a/Engine/Source/Runtime/ApplicationCore/Public/Mac/Accessibility/CocoaAccessibilityView.h b/Engine/Source/Runtime/ApplicationCore/Public/Mac/Accessibility/CocoaAccessibilityView.h
new file mode 100644
--- /dev/null
+++ b/Engine/Source/Runtime/ApplicationCore/Public/Mac/Accessibility/CocoaAccessibilityView.h
@@ -0,0 +1,33 @@
+// Copyright 1998-2019 Epic Games, Inc. All Rights Reserved.
+
+#pragma once
+#if WITH_ACCESSIBILITY
+#include "GenericPlatform/GenericAccessibleInterfaces.h"
+#
+#import <AppKit/AppKit.h>
+
+/**
+ * A custom NSView used to support accessibility.
+ * This should be the base class for any custom views that wish to support accessibility
+ */
+@class FMacAccessibilityElement;
+
+@interface FCocoaAccessibilityView : NSView
+/**The FMacAccessibilityElement* that corresponds to this views backing window */
+@property (nonatomic, assign) FMacAccessibilityElement* AccessibleWindow;
+
+/**
+ * Sets thee AccessibleWindow property. Used to create the accessibility hierarchy.
+ * Setting up of the accessibility hierarchy allows VoiceOver to interact with thee AccessibleWindow and its contained elements.
+ */
+-(void) SetAccessibilityWindowAsAccessibilityChild:(AccessibleWidgetId) InWindowId;
+
+/**
+ * Clears the AccessibleWindow property and tears down the accessibility hierarchy for this Window.
+ * This makes VOiceOVer unable to interact with the underlying elements in AccessibleWindow.
+ */
+-(void) RemoveAccessibilityWindow;
+
+@end
+
+#endif
diff --git a/Engine/Source/Runtime/ApplicationCore/Public/Mac/Accessibility/MacAccessibilityElement.h b/Engine/Source/Runtime/ApplicationCore/Public/Mac/Accessibility/MacAccessibilityElement.h
new file mode 100644
--- /dev/null
+++ b/Engine/Source/Runtime/ApplicationCore/Public/Mac/Accessibility/MacAccessibilityElement.h
@@ -0,0 +1,52 @@
+// Copyright 1998-2019 Epic Games, Inc. All Rights Reserved.
+
+#pragma once
+
+#if WITH_ACCESSIBILITY
+
+#include "CoreMinimal.h"
+#include "GenericPlatform/GenericAccessibleInterfaces.h"
+#include "Misc/Variant.h"
+#import <AppKit/AppKit.h>
+
+
+@interface FMacAccessibilityElement : NSAccessibilityElement
+{
+@private
+	/** An array that holds all the accessibility children of this array. Is filled when the accessibilityChildren property is requested. */
+	NSMutableArray* _AccessibilityChildren;
+	/** The Accessibility Role of the element. E.g NSAccessibilityButtonRole, NSAccessibilityTextFieldRole etc */
+	NSAccessibilityRole _AccessibilityRole;
+}
+- (id)initWithId:(AccessibleWidgetId)InId;
+
+/** Check if LastCachedStringTime was updated recently. */
+- (bool)ShouldCacheStrings;
+/** Updates the label, help and value of the element */
+- (void)CacheStrings;
+
+/** Exposes this element to VoiceOver's accessibility hierarchy and allows the element to be navigated to from VoiceOver */
+- (void)ExposeToVoiceOver:(bool)bShouldExpose;
+/** The identifier used to access this widget through the accessible API. */
+@property (nonatomic, assign) AccessibleWidgetId Id;
+/** The identifier used to access this widget's parent through the accessible API. */
+@property (nonatomic, assign) AccessibleWidgetId ParentId;
+/** The identifier used to access this widget's owning window through the accessible API. */
+@property (nonatomic, assign) AccessibleWidgetId OwningWindowId;;
+
+/** A list of identifiers for all current children of this container. */
+@property (nonatomic, assign) TArray<AccessibleWidgetId> ChildIds;
+/** The bounding rect of the container. Note this is NOT a valid AABB in Cocoa screen space. See accessibilityFrame implementation */
+@property (nonatomic) FBox2D Bounds;
+/** A cached version of the name of the widget. */
+@property (nonatomic) FString Label;
+/** A cached version of the help text of the widget. */
+@property (nonatomic) FString Help;
+/** A cached version of the value of property widgets. */
+@property (nonatomic) FVariant Value;
+/** Timestamp for when Label, Help, and Value were last cached. */
+@property (nonatomic, assign) double LastCachedStringTime;
+
+@end
+
+#endif
diff --git a/Engine/Source/Runtime/ApplicationCore/Public/Mac/Accessibility/MacAccessibilityManager.h b/Engine/Source/Runtime/ApplicationCore/Public/Mac/Accessibility/MacAccessibilityManager.h
new file mode 100644
--- /dev/null
+++ b/Engine/Source/Runtime/ApplicationCore/Public/Mac/Accessibility/MacAccessibilityManager.h
@@ -0,0 +1,65 @@
+// Copyright 1998-2019 Epic Games, Inc. All Rights Reserved.
+
+#pragma once
+
+#if WITH_ACCESSIBILITY
+
+#include "GenericPlatform/GenericAccessibleInterfaces.h"
+
+@class FMacAccessibilityElement;
+
+/**
+ * This class is a singleton and should be accessed through [FMacAccessibilityManager AccessibilityManager].
+ * Contains a Cache that stores a list of accessible elements that map to AccessibleWidgetIds for lookup.
+ * The cache is also responsible for polling attributes from the underlying
+ * IAccessibleWidgets that are too expensive to be done when requested by OSX due to
+ * needing to be accessed from a different thread.
+ *
+ * The Accessibility Manager is also used to monitor changes in Voiceover status and call the appropriate functions
+ * in FMacApplication using Key Value Observation (KVO).
+  */
+@interface FMacAccessibilityManager : NSObject
+{
+@private
+	/** AccessibleWidgetId(String)->FMacAccessibilityElement map for all created accessibility elements. */
+	NSMutableDictionary* Cache;
+}
+
+/** Pointer to the Mac Application that initialized this Singleton*/
+@property (nonatomic) class FMacApplication* MacApplication;
+/** Retrieve a cached element, or create one if it doesn't exist yet. */
+- (FMacAccessibilityElement*)GetAccessibilityElement:(AccessibleWidgetId)Id;
+/** Returns true if the Cache contains the Id. Does not create one if it doesn't exist. */
+- (bool)AccessibilityElementExists:(AccessibleWidgetId)Id;
+/** Returns the number of accessibility elements currently in the cache */
+- (int32)GetAccessibilityCacheSize;
+/** Returns true if the cache is empty */
+- (bool)IsAccessibilityCacheEmpty;
+/** Removes an entry from the Cache. */
+- (void)RemoveAccessibilityElement:(AccessibleWidgetId)Id;
+- (void)RemoveAccessibilitySubtree:(AccessibleWidgetId)RootId;
+/** Completely empties the cache. */
+- (void)Clear;
+/** Loop over all cached elements and update any properties necessary on the Game thread. */
+- (void)UpdateAllCachedProperties;
+/** Used to clean up memory and unregister from KVO. This should ONLY be called from ~FMacapplication() */
+- (void)TearDown;
+
+/** Singleton accessor */
++ (FMacAccessibilityManager*)AccessibilityManager;
+
+#if !UE_BUILD_SHIPPING
+- (void)DumpAccessibilityStats;
+
+- (void)PrintCache;
+- (void)PrintAccessibilityElement:(FMacAccessibilityElement*) ELement;
+- (void)PrintAccessibilityElementById:(AccessibleWidgetId) Id;
+- (id)GetAccessibilityParent:(FMacAccessibilityElement*) Element;
+- (FMacAccessibilityElement*)GetAccessibilityRoot;
+- (bool)isElementOrphaned:(FMacAccessibilityElement*) Element;
+
+#endif
+
+@end
+
+#endif
diff --git a/Engine/Source/Runtime/ApplicationCore/Public/Mac/CocoaTextView.h b/Engine/Source/Runtime/ApplicationCore/Public/Mac/CocoaTextView.h
--- a/Engine/Source/Runtime/ApplicationCore/Public/Mac/CocoaTextView.h
+++ b/Engine/Source/Runtime/ApplicationCore/Public/Mac/CocoaTextView.h
@@ -3,10 +3,17 @@
 
 #include "CoreMinimal.h"
 #include "GenericPlatform/ITextInputMethodSystem.h"
+#if WITH_ACCESSIBILITY
+#include "Mac/Accessibility/CocoaAccessibilityView.h"
+#endif
 
 #ifdef __OBJC__
-
+//INherit from FCocoaAccessibilityView to allow descendants (FMetalView in particular) to be visible to Voiceover
+#if WITH_ACCESSIBILITY
+OBJC_EXPORT @interface FCocoaTextView : FCocoaAccessibilityView <NSTextInputClient>
+#else
 OBJC_EXPORT @interface FCocoaTextView : NSView <NSTextInputClient>
+#endif
 {
 	TSharedPtr<ITextInputMethodContext> IMMContext;
 	NSRange markedRange;
diff --git a/Engine/Source/Runtime/ApplicationCore/Public/Mac/CocoaWindow.h b/Engine/Source/Runtime/ApplicationCore/Public/Mac/CocoaWindow.h
--- a/Engine/Source/Runtime/ApplicationCore/Public/Mac/CocoaWindow.h
+++ b/Engine/Source/Runtime/ApplicationCore/Public/Mac/CocoaWindow.h
@@ -3,6 +3,10 @@
 #pragma once
 
 #include "GenericPlatform/GenericWindow.h"
+#if WITH_ACCESSIBILITY
+#include "GenericPlatform/GenericAccessibleInterfaces.h"
+#endif
+
 
 /**
  * Custom window class used for input handling
@@ -23,6 +27,13 @@
 
 @property (assign) EWindowMode::Type TargetWindowMode;
 
+#if WITH_ACCESSIBILITY
+/** Adds the accessible element representing this window to the content view and exposes the window to accessibility interactions */
+- (void)UpdateAccessibilityView:(AccessibleWidgetId) InAccessibilityWindowId;
+
+/** removes the accessible element representing this window from the content view and prevents accessibility interactions being performed on the window */
+- (void)ClearAccessibilityView;
+#endif
 /** Get the frame filled by a child OpenGL view, which may cover the window or fill the content view depending upon the window style.
  @return The NSRect for a child OpenGL view. */
 - (NSRect)openGLFrame;
diff --git a/Engine/Source/Runtime/ApplicationCore/Public/Mac/MacApplication.h b/Engine/Source/Runtime/ApplicationCore/Public/Mac/MacApplication.h
--- a/Engine/Source/Runtime/ApplicationCore/Public/Mac/MacApplication.h
+++ b/Engine/Source/Runtime/ApplicationCore/Public/Mac/MacApplication.h
@@ -142,6 +142,16 @@ public:
 public:
 
 	virtual void SetMessageHandler(const TSharedRef<class FGenericApplicationMessageHandler>& InMessageHandler) override;
+#if WITH_ACCESSIBILITY
+	virtual void SetAccessibleMessageHandler(const TSharedRef<FGenericAccessibleMessageHandler>& InAccessibleMessageHandler) override;
+	
+	/** Called when Voiceover is enabled. FMacAccessibilityManager should be the only class to call this */
+	void OnVoiceoverEnabled();
+	
+	/** Called when Voiceover is disabled. FMacAccessibilityManager should be the only class to call this */
+	void OnVoiceoverDisabled();
+#endif
+
 
 	virtual void PollGameDeviceState(const float TimeDelta) override;
 
@@ -312,6 +322,10 @@ private:
 	void RecordUsage(EGestureEvent Gesture) { }
 #endif
 
+#if WITH_ACCESSIBILITY
+	void OnAccessibleEventRaised(TSharedRef<IAccessibleWidget> Widget, EAccessibleEvent Event, FVariant OldValue, FVariant NewValue);
+	#endif
+
 private:
 
 	bool bUsingHighPrecisionMouseInput;
@@ -392,6 +406,11 @@ private:
 
 	NSData* KeyBoardLayoutData;
 
+#if WITH_ACCESSIBILITY
+	/** Timer used to refresh the accessibility cache */
+	NSTimer* AccessibilityCacheTimer;
+#endif
+
 	friend class FMacWindow;
 };
 
diff --git a/Engine/Source/Runtime/Slate/Private/Widgets/Accessibility/SlateAccessibleWidgets.cpp b/Engine/Source/Runtime/Slate/Private/Widgets/Accessibility/SlateAccessibleWidgets.cpp
--- a/Engine/Source/Runtime/Slate/Private/Widgets/Accessibility/SlateAccessibleWidgets.cpp
+++ b/Engine/Source/Runtime/Slate/Private/Widgets/Accessibility/SlateAccessibleWidgets.cpp
@@ -30,7 +30,8 @@ bool FSlateAccessibleCheckBox::IsCheckable() const {
 	return true;
 }
 
-bool FSlateAccessibleCheckBox::GetCheckedState() const {
+bool FSlateAccessibleCheckBox::GetCheckedState() const
+{
 	if (Widget.IsValid())
 	{
 		return StaticCastSharedPtr<SCheckBox>(Widget.Pin())->IsChecked();
@@ -38,6 +39,18 @@ bool FSlateAccessibleCheckBox::GetCheckedState() const {
 	return false;
 }
 
+FString FSlateAccessibleCheckBox::GetValue() const
+{
+	bool bChecked = GetCheckedState();
+	return FString::FromInt(bChecked ? 1 : 0);
+}
+
+FVariant FSlateAccessibleCheckBox::GetValueAsVariant() const
+{
+	bool bChecked = GetCheckedState();
+	return FVariant(bChecked);
+}
+
 const FString& FSlateAccessibleEditableText::GetText() const
 {
 	if (Widget.IsValid())
@@ -71,6 +84,11 @@ FString FSlateAccessibleEditableText::GetValue() const
 	return GetText();
 }
 
+FVariant FSlateAccessibleEditableText::GetValueAsVariant() const
+{
+	return FVariant(GetText());
+}
+
 void FSlateAccessibleEditableText::SetValue(const FString& Value)
 {
 	if (Widget.IsValid())
@@ -112,6 +130,11 @@ FString FSlateAccessibleEditableTextBox::GetValue() const
 	return GetText();
 }
 
+FVariant FSlateAccessibleEditableTextBox::GetValueAsVariant() const
+{
+	return FVariant(GetText());
+}
+
 void FSlateAccessibleEditableTextBox::SetValue(const FString& Value)
 {
 	if (Widget.IsValid())
@@ -165,6 +188,15 @@ FString FSlateAccessibleSlider::GetValue() const
 	return FString::SanitizeFloat(0.0f);
 }
 
+FVariant FSlateAccessibleSlider::GetValueAsVariant() const
+{
+	if (Widget.IsValid())
+	{
+		return FVariant(StaticCastSharedPtr<SSlider>(Widget.Pin())->GetValue());
+	}
+	return FVariant();
+}
+
 void FSlateAccessibleSlider::SetValue(const FString& Value)
 {
 	if (Widget.IsValid())
diff --git a/Engine/Source/Runtime/Slate/Public/Widgets/Accessibility/SlateAccessibleWidgets.h b/Engine/Source/Runtime/Slate/Public/Widgets/Accessibility/SlateAccessibleWidgets.h
--- a/Engine/Source/Runtime/Slate/Public/Widgets/Accessibility/SlateAccessibleWidgets.h
+++ b/Engine/Source/Runtime/Slate/Public/Widgets/Accessibility/SlateAccessibleWidgets.h
@@ -28,18 +28,25 @@ public:
 class SLATE_API FSlateAccessibleCheckBox
 	: public FSlateAccessibleWidget
 	, public IAccessibleActivatable
+	, public IAccessibleProperty
 {
 public:
 	FSlateAccessibleCheckBox(TWeakPtr<SWidget> InWidget) : FSlateAccessibleWidget(InWidget, EAccessibleWidgetType::CheckBox) {}
 	virtual ~FSlateAccessibleCheckBox() {}
 
 	virtual IAccessibleActivatable* AsActivatable() override { return this; }
+	virtual IAccessibleProperty* AsProperty() override { return this; }
 
 	// IAccessibleActivatable
 	virtual void Activate() override;
 	virtual bool IsCheckable() const override;
 	virtual bool GetCheckedState() const override;
 	// ~
+	
+	// IAccessibleProperty
+	virtual FString GetValue() const override;
+	virtual FVariant GetValueAsVariant() const override;
+	// ~
 };
 // ~
 
@@ -65,6 +72,7 @@ public:
 	virtual bool IsReadOnly() const override;
 	virtual bool IsPassword() const override;
 	virtual FString GetValue() const override;
+	virtual FVariant GetValueAsVariant() const override;
 	virtual void SetValue(const FString& Value) override;
 	// ~
 };
@@ -92,6 +100,7 @@ public:
 	virtual bool IsReadOnly() const override;
 	virtual bool IsPassword() const override;
 	virtual FString GetValue() const override;
+	virtual FVariant GetValueAsVariant() const override;
 	virtual void SetValue(const FString& Value) override;
 	// ~
 };
@@ -135,6 +144,7 @@ public:
 	virtual float GetMaximum() const override;
 	virtual float GetMinimum() const override;
 	virtual FString GetValue() const override;
+	virtual FVariant GetValueAsVariant() const override;
 	virtual void SetValue(const FString& Value) override;
 	// ~
 };
