Copy up from Dev-Editor [REDACTED_HANDLE]

diff --git a/Engine/Source/Editor/UnrealEd/Private/Settings/LevelEditorPlaySettingsCustomization.h b/Engine/Source/Editor/UnrealEd/Private/Settings/LevelEditorPlaySettingsCustomization.h
--- a/Engine/Source/Editor/UnrealEd/Private/Settings/LevelEditorPlaySettingsCustomization.h
+++ b/Engine/Source/Editor/UnrealEd/Private/Settings/LevelEditorPlaySettingsCustomization.h
@@ -649,28 +649,17 @@ public:
 		}
 
 		// multi-player options
-		IDetailCategoryBuilder& NetworkCategory = LayoutBuilder.EditCategory("MultiplayerOptions");
-		{
-			// Number of players
-			NetworkCategory.AddProperty("PlayNumberOfClients")
-				.DisplayName(LOCTEXT("NumberOfPlayersLabel", "Number of Players"))
-				.IsEnabled(TAttribute<bool>::Create(TAttribute<bool>::FGetter::CreateSP(this, &FLevelEditorPlaySettingsCustomization::HandlePlayNumberOfClientsIsEnabled)));
-
-			NetworkCategory.AddProperty("ServerPort")
-				.DisplayName(LOCTEXT("ServerPortLabel", "Server Port"))
-				.IsEnabled(TAttribute<bool>::Create(TAttribute<bool>::FGetter::CreateSP(this, &FLevelEditorPlaySettingsCustomization::HandleServerPortIsEnabled)));
-
-			NetworkCategory.AddProperty("AdditionalServerGameOptions")
-				.DisplayName(LOCTEXT("ServerGameOptionsLabel", "Server Game Options"))
-				.IsEnabled(TAttribute<bool>::Create(TAttribute<bool>::FGetter::CreateSP(this, &FLevelEditorPlaySettingsCustomization::HandleGameOptionsIsEnabled)));
+		IDetailCategoryBuilder& NetworkCategory = LayoutBuilder.EditCategory("Multiplayer Options");
+		TArray<TSharedRef<IPropertyHandle>> AllProperties;
+		NetworkCategory.GetDefaultProperties(AllProperties);
 
-			NetworkCategory.AddProperty("PlayNetDedicated")
-				.DisplayName(LOCTEXT("RunDedicatedServerLabel", "Run Dedicated Server"))
-				.IsEnabled(TAttribute<bool>::Create(TAttribute<bool>::FGetter::CreateSP(this, &FLevelEditorPlaySettingsCustomization::HandlePlayNetDedicatedPropertyIsEnabled)));
-
-			NetworkCategory.AddProperty("ServerMapNameOverride")
-				.DisplayName(LOCTEXT("ServerMapNameOverrideLabel", "Server Map Name Override"))
-				.Visibility(TAttribute<EVisibility>::Create(TAttribute<EVisibility>::FGetter::CreateSP(this, &FLevelEditorPlaySettingsCustomization::HandleServerMapNameOverrideVisibility)));
+		{
+			// Add all the default properties in before we add custom rows at the bottom.
+			for (TSharedRef<IPropertyHandle> DefaultProperty : AllProperties)
+			{
+				NetworkCategory.AddProperty(DefaultProperty);
+			}
+			
 
 			// client window size
 			TSharedRef<IPropertyHandle> WindowHeightHandle = LayoutBuilder.GetProperty("ClientWindowHeight");
@@ -679,54 +668,19 @@ public:
 			WindowHeightHandle->MarkHiddenByCustomization();
 			WindowWidthHandle->MarkHiddenByCustomization();
 
-			NetworkCategory.AddProperty("AutoConnectToServer")
-				.DisplayName(LOCTEXT("AutoConnectToServerLabel", "Auto Connect To Server"))
-				.IsEnabled(TAttribute<bool>::Create(TAttribute<bool>::FGetter::CreateSP(this, &FLevelEditorPlaySettingsCustomization::HandleAutoConnectToServerEnabled)))
-				.Visibility(TAttribute<EVisibility>::Create(TAttribute<EVisibility>::FGetter::CreateSP(this, &FLevelEditorPlaySettingsCustomization::HandleAutoConnectToServerVisibility)));
-
-			NetworkCategory.AddProperty("RouteGamepadToSecondWindow")
-				.DisplayName(LOCTEXT("RouteGamepadToSecondWindowLabel", "Route 1st Gamepad to 2nd Client"))
-				.IsEnabled(TAttribute<bool>::Create(TAttribute<bool>::FGetter::CreateSP(this, &FLevelEditorPlaySettingsCustomization::HandleRerouteInputToSecondWindowEnabled)))
-				.Visibility(TAttribute<EVisibility>::Create(TAttribute<EVisibility>::FGetter::CreateSP(this, &FLevelEditorPlaySettingsCustomization::HandleRerouteInputToSecondWindowVisibility)));
 			
-			// Run under one instance
-			if (GEditor && GEditor->bAllowMultiplePIEWorlds)
-			{
-				NetworkCategory.AddProperty("RunUnderOneProcess")
-					.DisplayName(LOCTEXT("RunUnderOneProcessEnabledLabel", "Use Single Process"));
-			}
-			else
-			{
-				NetworkCategory.AddProperty("RunUnderOneProcess")
-					.DisplayName( LOCTEXT("RunUnderOneProcessDisabledLabel", "Run Under One Process is disabled.") )
-					.Visibility( EVisibility::Collapsed )
-					.IsEnabled( false );
-			}
-
-			// Net Mode
-			NetworkCategory.AddProperty("PlayNetMode")
-				.Visibility(TAttribute<EVisibility>::Create(TAttribute<EVisibility>::FGetter::CreateSP(this, &FLevelEditorPlaySettingsCustomization::HandlePlayNetModeVisibility)))
-				.DisplayName(LOCTEXT("PlayNetModeLabel", "Editor Multiplayer Mode"));
 
-			NetworkCategory.AddProperty("AdditionalLaunchOptions")
-				.DisplayName(LOCTEXT("AdditionalLaunchOptionsLabel", "Command Line Arguments"))
-				.Visibility(TAttribute<EVisibility>::Create(TAttribute<EVisibility>::FGetter::CreateSP(this, &FLevelEditorPlaySettingsCustomization::HandleCmdLineVisibility)));
-
-			NetworkCategory.AddProperty("NetworkEmulationSettings")
-				.IsEnabled(TAttribute<bool>(this, &FLevelEditorPlaySettingsCustomization::HandleNetworkEmulationIsEnabled));
-
-			NetworkCategory.AddCustomRow(LOCTEXT("PlayInNetworkViewportSize", "Multiplayer Viewport Size"), false)
+			NetworkCategory.AddCustomRow(LOCTEXT("PlayInNetworkViewportSize", "Viewport Size\n(Additional Clients)"), false)
 				.NameContent()
 				[
-					WindowHeightHandle->CreatePropertyNameWidget(LOCTEXT("ClientViewportSizeName", "Multiplayer Viewport Size (in pixels)"), LOCTEXT("ClientWindowSizeTooltip", "Width and Height to use when spawning additional windows."))
+					WindowHeightHandle->CreatePropertyNameWidget(LOCTEXT("ClientViewportSizeName", "Multiplayer Viewport Size (in pixels)"), LOCTEXT("ClientWindowSizeTooltip", "Width and Height to use when spawning additional clients. Useful when you need multiple clients connected but only interact with one window."))
 				]
 				.ValueContent()
 				.MaxDesiredWidth(MaxPropertyWidth)
 				[
 					SNew(SScreenResolutionCustomization, &LayoutBuilder, WindowHeightHandle, WindowWidthHandle)
 				]
-				.IsEnabled(TAttribute<bool>(this, &FLevelEditorPlaySettingsCustomization::HandleClientWindowSizePropertyIsEnabled))
-				.Visibility(TAttribute<EVisibility>::Create(TAttribute<EVisibility>::FGetter::CreateSP(this, &FLevelEditorPlaySettingsCustomization::HandleClientWindowSizePropertyVisibility)));
+				.IsEnabled(TAttribute<bool>(this, &FLevelEditorPlaySettingsCustomization::HandleClientWindowSizePropertyIsEnabled));
 				
 			NetworkCategory.AddCustomRow(LOCTEXT("AdditionalMultiplayerDetails", "Additional Options"), true)
 				.NameContent()
@@ -767,17 +721,17 @@ private:
 	FText HandleMultiplayerOptionsDescription( ) const
 	{
 		const ULevelEditorPlaySettings* PlayInSettings = GetDefault<ULevelEditorPlaySettings>();
+		const bool bLaunchSeparateServer = PlayInSettings->bLaunchSeparateServer;
 		const bool CanRunUnderOneProcess = [&PlayInSettings]{ bool RunUnderOneProcess(false); return (PlayInSettings->GetRunUnderOneProcess(RunUnderOneProcess) && RunUnderOneProcess); }();
-		const bool CanPlayNetDedicated = [&PlayInSettings]{ bool PlayNetDedicated(false); return (PlayInSettings->GetPlayNetDedicated(PlayNetDedicated) && PlayNetDedicated); }();
-		const int32 PlayNumberOfClients = [&PlayInSettings]{ int32 NumberOfClients(0); return (PlayInSettings->GetPlayNumberOfClients(NumberOfClients) ? NumberOfClients : 0); }();
+		const int32 PlayNumberOfClients = [&PlayInSettings] { int32 NumClients(0); return (PlayInSettings->GetPlayNumberOfClients(NumClients) ? NumClients : 0); }();
 		const EPlayNetMode PlayNetMode = [&PlayInSettings]{ EPlayNetMode NetMode(PIE_Standalone); return (PlayInSettings->GetPlayNetMode(NetMode) ? NetMode : PIE_Standalone); }();
 		FString Desc;
 		if (CanRunUnderOneProcess)
 		{
 			Desc += LOCTEXT("MultiplayerDescription_OneProcess", "The following will all run under one UE4 instance:\n").ToString();
-			if (CanPlayNetDedicated)
+			if (PlayNetMode == EPlayNetMode::PIE_Client)
 			{
-				Desc += LOCTEXT("MultiplayerDescription_DedicatedServer", "A dedicated server will open in a new window. ").ToString();
+				Desc += LOCTEXT("MultiplayerDescription_DedicatedServer", "A hidden dedicated server instance will run in editor. ").ToString();
 				if (PlayNumberOfClients == 1)
 				{
 					Desc += LOCTEXT("MultiplayerDescription_EditorClient", "The editor will connect as a client. ").ToString();
@@ -787,7 +741,7 @@ private:
 					Desc += FText::Format(LOCTEXT("MultiplayerDescription_EditorAndClients", "The editor will connect as a client and {0} additional client window(s) will also connect. "), FText::AsNumber(PlayNumberOfClients-1)).ToString();
 				}
 			}
-			else
+			else if(PlayNetMode == EPlayNetMode::PIE_ListenServer)
 			{
 				if (PlayNumberOfClients == 1)
 				{
@@ -798,13 +752,41 @@ private:
 					Desc += FText::Format(LOCTEXT("MultiplayerDescription_EditorListenServerAndClients", "The editor will run as a listen server and {0} additional client window(s) will also connect to it. "), FText::AsNumber(PlayNumberOfClients-1)).ToString();
 				}
 			}
+			else
+			{
+				if (PlayNumberOfClients == 1)
+				{
+					Desc += LOCTEXT("MultiplayerDescription_EditorStandalone", "The editor will run in offline mode. ").ToString();
+				}
+				else
+				{
+					Desc += FText::Format(LOCTEXT("MultiplayerDescription_StandaloneAndClients", "The editor will run offline and {0} additional offline mode window(s) will also open. "), FText::AsNumber(PlayNumberOfClients - 1)).ToString();
+				}
+
+				if (bLaunchSeparateServer)
+				{
+					Desc += LOCTEXT("MultiplayerDescription_StandaloneSeparateServer", "\nAn additional server instance will be launched but not connected to. Use \"open 127.0.0.1:<port>\" to connect. ").ToString();
+				}
+			}
 		}
 		else
 		{
 			Desc += LOCTEXT("MultiplayerDescription_MultiProcess", "The following will run with multiple UE4 instances:\n").ToString();
 			if (PlayNetMode == PIE_Standalone)
 			{
-				Desc += LOCTEXT("MultiplayerDescription_EditorOffline", "The editor will run offline. ").ToString();
+				if (PlayNumberOfClients == 1)
+				{
+					Desc += LOCTEXT("MultiplayerDescription_EditorStandalone", "The editor will run in offline mode. ").ToString();
+				}
+				else
+				{
+					Desc += FText::Format(LOCTEXT("MultiplayerDescription_StandaloneAndClients", "The editor will run offline and {0} additional offline mode window(s) will also open. "), FText::AsNumber(PlayNumberOfClients - 1)).ToString();
+				}
+
+				if (bLaunchSeparateServer)
+				{
+					Desc += LOCTEXT("MultiplayerDescription_StandaloneSeparateServer", "\nAn additional server instance will be launched but not connected to. Use \"open 127.0.0.1:<port>\" to connect. ").ToString();
+				}
 			}
 			else if (PlayNetMode == PIE_ListenServer)
 			{
@@ -819,29 +801,17 @@ private:
 			}
 			else
 			{
-				if (CanPlayNetDedicated)
+				// Client requires additional dedicated server instance
+				Desc += LOCTEXT("MultiplayerDescription_DedicatedServer", "A dedicated server will open in a new window. ").ToString();
+				if (PlayNumberOfClients == 1)
 				{
-					Desc += LOCTEXT("MultiplayerDescription_DedicatedServer", "A dedicated server will open in a new window. ").ToString();
-					if (PlayNumberOfClients == 1)
-					{
-						Desc += LOCTEXT("MultiplayerDescription_EditorClient", "The editor will connect as a client. ").ToString();
-					}
-					else
-					{
-						Desc += FText::Format(LOCTEXT("MultiplayerDescription_EditorAndClients", "The editor will connect as a client and {0} additional client window(s) will also connect. "), FText::AsNumber(PlayNumberOfClients-1)).ToString();
-					}
+					Desc += LOCTEXT("MultiplayerDescription_EditorClient", "The editor will connect as a client. ").ToString();
 				}
 				else
 				{
-					if (PlayNumberOfClients <= 2)
-					{
-						Desc += LOCTEXT("MultiplayerDescription_EditorClientAndListenServer", "A listen server will open in a new window and the editor will connect to it. ").ToString();
-					}
-					else
-					{
-						Desc += FText::Format(LOCTEXT("MultiplayerDescription_EditorClientAndListenServerClients", "A listen server will open in a new window and the editor will connect as a client and {0} additional client window(s) will also connect to it. "), FText::AsNumber(FMath::Max(0, PlayNumberOfClients-2))).ToString(); 
-					}
+					Desc += FText::Format(LOCTEXT("MultiplayerDescription_EditorAndClients", "The editor will connect as a client and {0} additional client window(s) will also connect. "), FText::AsNumber(PlayNumberOfClients-1)).ToString();
 				}
+
 			}
 		}
 		return FText::FromString(Desc);
@@ -853,48 +823,12 @@ private:
 		return GetDefault<ULevelEditorPlaySettings>()->IsClientWindowSizeActive();
 	}
 
-	// Callback for getting the visibility of the ClientWindowHeight and ClientWindowWidth properties.
-	EVisibility HandleClientWindowSizePropertyVisibility() const
-	{
-		return GetDefault<ULevelEditorPlaySettings>()->GetClientWindowSizeVisibility();
-	}
-
-	// Callback for checking whether the PlayNetDedicated is enabled.
-	bool HandlePlayNetDedicatedPropertyIsEnabled( ) const
-	{		
-		return GetDefault<ULevelEditorPlaySettings>()->IsPlayNetDedicatedActive();
-	}
-
-	// Callback for checking whether the PlayNumberOfClients is enabled.
-	bool HandlePlayNumberOfClientsIsEnabled( ) const
-	{
-		return GetDefault<ULevelEditorPlaySettings>()->IsPlayNumberOfClientsActive();
-	}
-
-	// Callback for checking whether the ServerPort is enabled.
-	bool HandleServerPortIsEnabled() const
-	{
-		return GetDefault<ULevelEditorPlaySettings>()->IsServerPortActive();
-	}
-
 	// Callback for checking whether the AdditionalServerGameOptions is enabled.
 	bool HandleGameOptionsIsEnabled( ) const
 	{
 		return GetDefault<ULevelEditorPlaySettings>()->IsAdditionalServerGameOptionsActive();
 	}
 
-	// Callback for getting the enabled state of the AutoConnectToServer property.
-	bool HandleAutoConnectToServerEnabled() const
-	{
-		return GetDefault<ULevelEditorPlaySettings>()->IsAutoConnectToServerActive();
-	}
-
-	// Callback for getting the visibility of the RerouteInputToSecondWindow property.
-	EVisibility HandleAutoConnectToServerVisibility() const
-	{
-		return GetDefault<ULevelEditorPlaySettings>()->GetAutoConnectToServerVisibility();
-	}
-
 	// Callback for getting the enabled state of the RerouteInputToSecondWindow property.
 	bool HandleRerouteInputToSecondWindowEnabled( ) const
 	{
@@ -907,30 +841,6 @@ private:
 		return GetDefault<ULevelEditorPlaySettings>()->GetRouteGamepadToSecondWindowVisibility();
 	}
 
-	// Callback for getting the visibility of the PlayNetMode property.
-	EVisibility HandlePlayNetModeVisibility( ) const
-	{
-		return GetDefault<ULevelEditorPlaySettings>()->GetPlayNetModeVisibility();
-	}
-
-	// Callback for checking if the Network Emulation can be used
-	bool HandleNetworkEmulationIsEnabled() const
-	{
-		return GetDefault<ULevelEditorPlaySettings>()->GetNetworkEmulationVisibility() == EVisibility::Visible;
-	}
-
-	// Callback for getting the visibility of the StandaloneServerMapName property.
-	EVisibility HandleServerMapNameOverrideVisibility() const
-	{
-		return GetDefault<ULevelEditorPlaySettings>()->GetServerMapNameOverrideVisibility();
-	}
-
-	// Callback for getting the visibility of the AdditionalLaunchOptions property.
-	EVisibility HandleCmdLineVisibility( ) const
-	{
-		return GetDefault<ULevelEditorPlaySettings>()->GetAdditionalLaunchOptionsVisibility();
-	}
-
 	void HandleQualityLevelComboBoxOpening()
 	{
 		const UAudioSettings* AudioSettings = GetDefault<UAudioSettings>();
diff --git a/Engine/Source/Editor/UnrealEd/Private/Settings/SettingsClasses.cpp b/Engine/Source/Editor/UnrealEd/Private/Settings/SettingsClasses.cpp
--- a/Engine/Source/Editor/UnrealEd/Private/Settings/SettingsClasses.cpp
+++ b/Engine/Source/Editor/UnrealEd/Private/Settings/SettingsClasses.cpp
@@ -430,6 +430,7 @@ ULevelEditorMiscSettings::ULevelEditorMiscSettings( const FObjectInitializer& Ob
 	PercentageThresholdForPrompt = 20.0f;
 	MinimumBoundsForCheckingSize = FVector(500.0f, 500.0f, 50.0f);
 	bCreateNewAudioDeviceForPlayInEditor = true;
+	bEnableLegacyMeshPaintMode = true;
 }
 
 void ULevelEditorMiscSettings::PostEditChangeProperty( struct FPropertyChangedEvent& PropertyChangedEvent )
@@ -461,12 +462,16 @@ ULevelEditorPlaySettings::ULevelEditorPlaySettings( const FObjectInitializer& Ob
 {
 	ClientWindowWidth = 640;
 	ClientWindowHeight = 480;
+	PlayNetMode = EPlayNetMode::PIE_Standalone;
+	bLaunchSeparateServer = false;
 	PlayNumberOfClients = 1;
 	ServerPort = 17777;
+PRAGMA_DISABLE_DEPRECATION_WARNINGS
 	PlayNetDedicated = false;
+	AutoConnectToServer = true;
+PRAGMA_ENABLE_DEPRECATION_WARNINGS
 	RunUnderOneProcess = true;
 	RouteGamepadToSecondWindow = false;
-	AutoConnectToServer = true;
 	BuildGameBeforeLaunch = EPlayOnBuildMode::PlayOnBuild_Default;
 	LaunchConfiguration = EPlayOnLaunchConfiguration::LaunchConfig_Default;
 	bAutoCompileBlueprintsOnLaunch = true;
@@ -512,6 +517,19 @@ void ULevelEditorPlaySettings::PostInitProperties()
 #endif
 }
 
+bool ULevelEditorPlaySettings::CanEditChange(const UProperty* InProperty) const
+{
+	const bool ParentVal = Super::CanEditChange(InProperty);
+	FName PropertyName = InProperty->GetFName();
+
+	if (PropertyName == GET_MEMBER_NAME_CHECKED(ULevelEditorPlaySettings, AdditionalServerLaunchParameters))
+	{
+		return ParentVal && (!RunUnderOneProcess && (PlayNetMode == EPlayNetMode::PIE_Client || bLaunchSeparateServer));
+	}
+
+	return ParentVal;
+}
+
 #if WITH_EDITOR
 void ULevelEditorPlaySettings::SwapSafeZoneTypes()
 {
diff --git a/Engine/Source/Editor/UnrealEd/Private/SkinWeightsUtilities.cpp b/Engine/Source/Editor/UnrealEd/Private/SkinWeightsUtilities.cpp
--- a/Engine/Source/Editor/UnrealEd/Private/SkinWeightsUtilities.cpp
+++ b/Engine/Source/Editor/UnrealEd/Private/SkinWeightsUtilities.cpp
@@ -299,7 +299,7 @@ bool FSkinWeightsUtilities::RemoveSkinnedWeightProfileData(USkeletalMesh* Skelet
 	LODModelDest.SkinWeightProfiles.Remove(ProfileName);
 
 	FSkeletalMeshImportData ImportDataDest;
-	LODModelDest.RawSkeletalMeshBulkData.LoadRawMesh(ImportDataDest);
+	SkeletalMesh->LoadLODImportedData(LODIndex, ImportDataDest);
 
 	//Rechunk the skeletal mesh since we remove it, we rebuild the skeletal mesh to achieve rechunking
 	UFbxSkeletalMeshImportData* OriginalSkeletalMeshImportData = UFbxSkeletalMeshImportData::GetImportDataForSkeletalMesh(SkeletalMesh, nullptr);
@@ -318,6 +318,7 @@ bool FSkinWeightsUtilities::RemoveSkinnedWeightProfileData(USkeletalMesh* Skelet
 	BuildOptions.OverlappingThresholds.ThresholdPosition = OriginalSkeletalMeshImportData->ThresholdPosition;
 	BuildOptions.OverlappingThresholds.ThresholdTangentNormal = OriginalSkeletalMeshImportData->ThresholdTangentNormal;
 	BuildOptions.OverlappingThresholds.ThresholdUV = OriginalSkeletalMeshImportData->ThresholdUV;
+	BuildOptions.OverlappingThresholds.MorphThresholdPosition = OriginalSkeletalMeshImportData->MorphThresholdPosition;
 	BuildOptions.bComputeNormals = !bShouldImportNormals || !ImportDataDest.bHasNormals;
 	BuildOptions.bComputeTangents = !bShouldImportTangents || !ImportDataDest.bHasTangents;
 	BuildOptions.bUseMikkTSpace = (OriginalSkeletalMeshImportData->NormalGenerationMethod == EFBXNormalGenerationMethod::MikkTSpace) && (!bShouldImportNormals || !bShouldImportTangents);
diff --git a/Engine/Source/Editor/UnrealEd/Private/StaticMeshEdit.cpp b/Engine/Source/Editor/UnrealEd/Private/StaticMeshEdit.cpp
--- a/Engine/Source/Editor/UnrealEd/Private/StaticMeshEdit.cpp
+++ b/Engine/Source/Editor/UnrealEd/Private/StaticMeshEdit.cpp
@@ -37,6 +37,7 @@
 #include "IMeshReductionManagerModule.h"
 #include "IMeshReductionInterfaces.h"
 
+#include "UObject/MetaData.h"
 
 bool GBuildStaticMeshCollision = 1;
 
@@ -274,7 +275,7 @@ private:
 	}
 };
 
-void DecomposeUCXMesh( const TArray<FVector>& CollisionVertices, const TArray<int32>& CollisionFaceIdx, UBodySetup* BodySetup )
+bool DecomposeUCXMesh( const TArray<FVector>& CollisionVertices, const TArray<int32>& CollisionFaceIdx, UBodySetup* BodySetup )
 {
 	// We keep no ref to this Model, so it will be GC'd at some point after the import.
 	auto TempModel = NewObject<UModel>();
@@ -294,6 +295,7 @@ void DecomposeUCXMesh( const TArray<FVector>& CollisionVertices, const TArray<in
 	ConnectivityBuilder.CreateConnectivityGroups();
 
 	// For each valid group build BSP and extract convex hulls
+	bool bSuccess = true;
 	for ( int32 i=0; i<ConnectivityBuilder.Groups.Num(); i++ )
 	{
 		const FMeshConnectivityGroup &Group = ConnectivityBuilder.Groups[ i ];
@@ -331,8 +333,11 @@ void DecomposeUCXMesh( const TArray<FVector>& CollisionVertices, const TArray<in
 
 		// Convert collision model into a collection of convex hulls.
 		// Generated convex hulls will be added to existing ones
-		BodySetup->CreateFromModel( TempModel, false );
+		bSuccess = BodySetup->CreateFromModel( TempModel, false ) && bSuccess;
 	}
+
+	// Could all meshes be properly decomposed?
+	return bSuccess;
 }
 
 /** 
@@ -1010,6 +1015,9 @@ struct ExistingStaticMeshData
 	FVector						ExistingNegativeBoundsExtension;
 
 	UStaticMesh::FOnMeshChanged	ExistingOnMeshChanged;
+	UStaticMesh* ExistingComplexCollisionMesh = nullptr;
+
+	TMap<FName, FString>* ExistingUMetaDataTagValues;
 };
 
 bool IsUsingMaterialSlotNameWorkflow(UAssetImportData* AssetImportData)
@@ -1040,6 +1048,9 @@ ExistingStaticMeshData* SaveExistingStaticMeshData(UStaticMesh* ExistingMesh, Un
 		bool bSaveMaterials = !ImportOptions->bImportMaterials;
 		ExistingMeshDataPtr = new ExistingStaticMeshData();
 
+		//Save the package UMetaData
+		ExistingMeshDataPtr->ExistingUMetaDataTagValues = UMetaData::GetMapForObject(ExistingMesh);
+
 		ExistingMeshDataPtr->ImportVersion = ExistingMesh->ImportVersion;
 		ExistingMeshDataPtr->UseMaterialNameSlotWorkflow = IsUsingMaterialSlotNameWorkflow(ExistingMesh->AssetImportData);
 
@@ -1192,6 +1203,8 @@ ExistingStaticMeshData* SaveExistingStaticMeshData(UStaticMesh* ExistingMesh, Un
 			}
 		}
 		ExistingMeshDataPtr->ExistingOnMeshChanged = ExistingMesh->OnMeshChanged;
+
+		ExistingMeshDataPtr->ExistingComplexCollisionMesh = ExistingMesh->ComplexCollisionMesh;
 	}
 
 	return ExistingMeshDataPtr;
@@ -1408,6 +1421,14 @@ void RestoreExistingMeshData(ExistingStaticMeshData* ExistingMeshDataPtr, UStati
 		return;
 	}
 
+	//Restore the package metadata
+	if (ExistingMeshDataPtr->ExistingUMetaDataTagValues)
+	{
+		UMetaData* PackageMetaData = NewMesh->GetOutermost()->GetMetaData();
+		checkSlow(PackageMetaData);
+		PackageMetaData->SetObjectValues(NewMesh, *ExistingMeshDataPtr->ExistingUMetaDataTagValues);
+	}
+
 	//Create a remap material Index use to find the matching section later
 	TArray<int32> RemapMaterial;
 	RemapMaterial.AddZeroed(NewMesh->StaticMaterials.Num());
@@ -1681,6 +1702,8 @@ void RestoreExistingMeshData(ExistingStaticMeshData* ExistingMeshDataPtr, UStati
 	NewMesh->PositiveBoundsExtension = ExistingMeshDataPtr->ExistingPositiveBoundsExtension;
 	NewMesh->NegativeBoundsExtension = ExistingMeshDataPtr->ExistingNegativeBoundsExtension;
 
+	NewMesh->ComplexCollisionMesh = ExistingMeshDataPtr->ExistingComplexCollisionMesh;
+
 	delete ExistingMeshDataPtr;
 }
 
diff --git a/Engine/Source/Editor/UnrealEd/Private/Tests/AutomationEditorCommon.cpp b/Engine/Source/Editor/UnrealEd/Private/Tests/AutomationEditorCommon.cpp
--- a/Engine/Source/Editor/UnrealEd/Private/Tests/AutomationEditorCommon.cpp
+++ b/Engine/Source/Editor/UnrealEd/Private/Tests/AutomationEditorCommon.cpp
@@ -773,9 +773,15 @@ bool FCloseAllAssetEditorsCommand::Update()
 bool FStartPIECommand::Update()
 {
 	FLevelEditorModule& LevelEditorModule = FModuleManager::Get().GetModuleChecked<FLevelEditorModule>(TEXT("LevelEditor"));
-	TSharedPtr<class IAssetViewport> ActiveLevelViewport = LevelEditorModule.GetFirstActiveViewport();
 
-	GUnrealEd->RequestPlaySession(false, ActiveLevelViewport, bSimulateInEditor, NULL, NULL, -1, false);
+	FRequestPlaySessionParams Params;
+	Params.DestinationSlateViewport = LevelEditorModule.GetFirstActiveViewport();
+	if (bSimulateInEditor)
+	{
+		Params.WorldType = EPlaySessionWorldType::SimulateInEditor;
+	}
+
+	GUnrealEd->RequestPlaySession(Params);
 	return true;
 }
 
@@ -933,7 +939,18 @@ bool FSaveLevelCommand::Update()
 
 bool FLaunchOnCommand::Update()
 {
-	GUnrealEd->AutomationPlayUsingLauncher(InLauncherDeviceID);
+	FRequestPlaySessionParams::FLauncherDeviceInfo LaunchedDeviceInfo;
+	LaunchedDeviceInfo.DeviceId = InLauncherDeviceID;
+	LaunchedDeviceInfo.DeviceName = InLauncherDeviceID.Right(InLauncherDeviceID.Find(TEXT("@")));
+
+	FRequestPlaySessionParams Params;
+	Params.LauncherTargetDevice = LaunchedDeviceInfo;
+
+	GUnrealEd->RequestPlaySession(Params);
+
+	// Immediately start our requested play session
+	GUnrealEd->StartQueuedPlaySessionRequest();
+
 	return true;
 }
 
diff --git a/Engine/Source/Editor/UnrealEd/Private/Tests/AutomationEditorPromotionCommon.cpp b/Engine/Source/Editor/UnrealEd/Private/Tests/AutomationEditorPromotionCommon.cpp
--- a/Engine/Source/Editor/UnrealEd/Private/Tests/AutomationEditorPromotionCommon.cpp
+++ b/Engine/Source/Editor/UnrealEd/Private/Tests/AutomationEditorPromotionCommon.cpp
@@ -235,7 +235,14 @@ void FEditorPromotionTestUtilities::StartPIE(bool bSimulateInEditor)
 	FLevelEditorModule& LevelEditorModule = FModuleManager::Get().GetModuleChecked<FLevelEditorModule>(TEXT("LevelEditor"));
 	TSharedPtr<class IAssetViewport> ActiveLevelViewport = LevelEditorModule.GetFirstActiveViewport();
 
-	GUnrealEd->RequestPlaySession(false, ActiveLevelViewport, bSimulateInEditor, NULL, NULL, -1, false);
+	FRequestPlaySessionParams SessionParams;
+	SessionParams.DestinationSlateViewport = ActiveLevelViewport;
+	if (bSimulateInEditor)
+	{
+		SessionParams.WorldType = EPlaySessionWorldType::SimulateInEditor;
+	}
+
+	GUnrealEd->RequestPlaySession(SessionParams);
 }
 
 /**
diff --git a/Engine/Source/Editor/UnrealEd/Private/Toolkits/BaseToolkit.cpp b/Engine/Source/Editor/UnrealEd/Private/Toolkits/BaseToolkit.cpp
--- a/Engine/Source/Editor/UnrealEd/Private/Toolkits/BaseToolkit.cpp
+++ b/Engine/Source/Editor/UnrealEd/Private/Toolkits/BaseToolkit.cpp
@@ -11,6 +11,8 @@
 #include "Editor.h"
 #include "InteractiveTool.h"
 #include "Tools/UEdMode.h"
+#include "Framework/MultiBox/MultiBoxBuilder.h"
+#include "EdMode.h"
 
 #define LOCTEXT_NAMESPACE "BaseToolkit"
 
@@ -215,10 +217,29 @@ void FModeToolkit::Init(const TSharedPtr< class IToolkitHost >& InitToolkitHost)
 	ToolkitMode = EToolkitMode::WorldCentric;
 	ToolkitHost = InitToolkitHost;
 
-	FToolkitManager::Get().RegisterNewToolkit( SharedThis( this ) );
+	if (UEdMode* ScriptableEditorMode = GetScriptableEditorMode())
+	{
+		ScriptableEditorMode->GetToolManager()->OnToolStarted.AddSP(this, &FModeToolkit::OnToolStarted);
+		ScriptableEditorMode->GetToolManager()->OnToolEnded.AddSP(this, &FModeToolkit::OnToolEnded);
+	}
 
 	FPropertyEditorModule& PropertyEditorModule = FModuleManager::GetModuleChecked<FPropertyEditorModule>("PropertyEditor");
 
+	FDetailsViewArgs ModeDetailsViewArgs(
+		/*bUpdateFromSelection=*/ false,
+		/*bLockable=*/ false,
+		/*bAllowSearch=*/ false,
+		FDetailsViewArgs::HideNameArea,
+		/*bHideSelectionTip=*/ true,
+		/*InNotifyHook=*/ nullptr,
+		/*InSearchInitialKeyFocus=*/ false,
+		/*InViewIdentifier=*/ NAME_None);
+	ModeDetailsViewArgs.DefaultsOnlyVisibility = EEditDefaultsOnlyNodeVisibility::Automatic;
+	ModeDetailsViewArgs.bShowOptions = false;
+	ModeDetailsViewArgs.bAllowMultipleTopLevelObjects = true;
+
+	ModeDetailsView = PropertyEditorModule.CreateDetailView(ModeDetailsViewArgs);
+
 	FDetailsViewArgs DetailsViewArgs(
 		/*bUpdateFromSelection=*/ false,
 		/*bLockable=*/ false,
@@ -233,9 +254,24 @@ void FModeToolkit::Init(const TSharedPtr< class IToolkitHost >& InitToolkitHost)
 	DetailsViewArgs.bAllowMultipleTopLevelObjects = true;
 
 	DetailsView = PropertyEditorModule.CreateDetailView(DetailsViewArgs);
+
+
+	FToolkitManager::Get().RegisterNewToolkit(SharedThis(this));
 }
 
 
+FModeToolkit::~FModeToolkit()
+{
+	if (UEdMode* ScriptableEditorMode = GetScriptableEditorMode())
+	{
+		if (ScriptableEditorMode->GetToolManager())
+		{
+			ScriptableEditorMode->GetToolManager()->OnToolStarted.RemoveAll(this);
+			ScriptableEditorMode->GetToolManager()->OnToolEnded.RemoveAll(this);
+		}
+	}
+}
+
 FName FModeToolkit::GetToolkitFName() const
 {
 	return FName("EditorModeToolkit");
@@ -292,33 +328,19 @@ bool FModeToolkit::CanCompleteActiveTool()
 }
 
 
-FReply FModeToolkit::StartTool(const FString& ToolTypeIdentifier)
+void FModeToolkit::OnToolStarted(UInteractiveToolManager* Manager, UInteractiveTool* Tool)
 {
-	if (GetScriptableEditorMode()->GetToolManager()->SelectActiveToolType(EToolSide::Left, ToolTypeIdentifier) == false)
+	// Update properties panel
+	UInteractiveTool* CurTool = GetScriptableEditorMode()->GetToolManager()->GetActiveTool(EToolSide::Left);
+	if (CurTool)
 	{
-		UE_LOG(LogTemp, Warning, TEXT("ToolManager: Unknown Tool Type %s"), *ToolTypeIdentifier);
-	}
-	else
-	{
-		UE_LOG(LogTemp, Warning, TEXT("ToolManager: Starting Tool Type %s"), *ToolTypeIdentifier);
-		GetScriptableEditorMode()->GetToolManager()->ActivateTool(EToolSide::Left);
-
-		// Update properties panel
-		UInteractiveTool* CurTool = GetScriptableEditorMode()->GetToolManager()->GetActiveTool(EToolSide::Left);
 		DetailsView->SetObjects(CurTool->GetToolProperties());
 	}
-	return FReply::Handled();
 }
 
-FReply FModeToolkit::EndTool(EToolShutdownType ShutdownType)
+void FModeToolkit::OnToolEnded(UInteractiveToolManager* Manager, UInteractiveTool* Tool)
 {
-	UE_LOG(LogTemp, Warning, TEXT("ENDING TOOL"));
-
-	GetScriptableEditorMode()->GetToolManager()->DeactivateTool(EToolSide::Left, ShutdownType);
-
 	DetailsView->SetObject(nullptr);
-
-	return FReply::Handled();
 }
 
 class FEdMode* FModeToolkit::GetEditorMode() const
@@ -326,7 +348,87 @@ class FEdMode* FModeToolkit::GetEditorMode() const
 	return nullptr; 
 }
 
+FText FModeToolkit::GetEditorModeDisplayName() const
+{
+	if (FEdMode* EdMode = GetEditorMode())
+	{
+		return EdMode->GetModeInfo().Name;
+	}
+	else if (UEdMode* ScriptableMode = GetScriptableEditorMode())
+	{
+		return ScriptableMode->GetModeInfo().Name;
+	}
+
+	return FText::GetEmpty();
+}
+
+FSlateIcon FModeToolkit::GetEditorModeIcon() const
+{
+	if (FEdMode* EdMode = GetEditorMode())
+	{
+		return EdMode->GetModeInfo().IconBrush;
+	}
+	else if (UEdMode* ScriptableMode = GetScriptableEditorMode())
+	{
+		return ScriptableMode->GetModeInfo().IconBrush;
+	}
+
+	return FSlateIcon();
+}
+
 UEdMode* FModeToolkit::GetScriptableEditorMode() const
 {
 	return nullptr;
-}
\ No newline at end of file
+}
+
+TSharedPtr<SWidget> FModeToolkit::GetInlineContent() const
+{
+	return SNew(SVerticalBox)
+		+ SVerticalBox::Slot()
+		.AutoHeight()
+		[
+			ModeDetailsView.ToSharedRef()
+		]
+		+ SVerticalBox::Slot()
+		.AutoHeight()
+		[
+			DetailsView.ToSharedRef()
+		];
+}
+
+void FModeToolkit::BuildToolPalette(FName PaletteName, class FToolBarBuilder& ToolbarBuilder)
+{
+	if (UEdMode* EditorMode = GetScriptableEditorMode())
+	{
+		TMap<FName, TArray<TSharedPtr<FUICommandInfo>>> CommandLists = EditorMode->GetModeCommands();
+		TArray<TSharedPtr<FUICommandInfo>>* CurrentCommandListPtr = CommandLists.Find(PaletteName);
+		if (CurrentCommandListPtr)
+		{
+			TArray<TSharedPtr<FUICommandInfo>> CurrentCommandList = *CurrentCommandListPtr;
+			for (TSharedPtr<FUICommandInfo> Command : CurrentCommandList)
+			{
+				ToolbarBuilder.AddToolBarButton(Command);
+			}
+			
+		}
+
+	}
+}
+
+void FModeToolkit::OnToolPaletteChanged(FName PaletteName)
+{
+	if (UEdMode* EditorMode = GetScriptableEditorMode())
+	{
+		EditorMode->SetCurrentPaletteName(PaletteName);
+		if (UInteractiveTool* CurrentTool = EditorMode->GetToolManager()->GetActiveTool(EToolSide::Left))
+		{
+			// When swapping tools, default to accepting the active tool
+			EditorMode->GetToolManager()->DeactivateTool(EToolSide::Left, EToolShutdownType::Accept);
+		}
+	}
+}
+
+void FModeToolkit::SetModeSettingsObject(UObject* InSettingsObject)
+{
+	ModeDetailsView->SetObject(InSettingsObject);
+}
diff --git a/Engine/Source/Editor/UnrealEd/Private/Tools/EditorComponentSourceFactory.cpp b/Engine/Source/Editor/UnrealEd/Private/Tools/EditorComponentSourceFactory.cpp
--- a/Engine/Source/Editor/UnrealEd/Private/Tools/EditorComponentSourceFactory.cpp
+++ b/Engine/Source/Editor/UnrealEd/Private/Tools/EditorComponentSourceFactory.cpp
@@ -9,13 +9,63 @@
 #include "PhysicsEngine/BodySetup.h"
 
 
-FMeshDescription *
-FStaticMeshComponentTarget::GetMesh() {
+FMeshDescription* FStaticMeshComponentTarget::GetMesh() 
+{
 	return Cast<UStaticMeshComponent>(Component)->GetStaticMesh()->GetMeshDescription(LODIndex);
 }
 
-void
-FStaticMeshComponentTarget::CommitMesh( const FCommitter& Committer )
+
+void FStaticMeshComponentTarget::GetMaterialSet(FComponentMaterialSet& MaterialSetOut, bool bAssetMaterials) const
+{
+	if (bAssetMaterials)
+	{
+		UStaticMesh* StaticMesh = Cast<UStaticMeshComponent>(Component)->GetStaticMesh();
+		int32 NumMaterials = Component->GetNumMaterials();
+		MaterialSetOut.Materials.SetNum(NumMaterials);
+		for (int32 k = 0; k < NumMaterials; ++k)
+		{
+			MaterialSetOut.Materials[k] = StaticMesh->GetMaterial(k);
+		}
+	}
+	else
+	{
+		FPrimitiveComponentTarget::GetMaterialSet(MaterialSetOut, false);
+	}
+}
+
+
+void FStaticMeshComponentTarget::CommitMaterialSetUpdate(const FComponentMaterialSet& MaterialSet, bool bApplyToAsset)
+{
+	// we only support this right now...
+	check(bApplyToAsset == true);
+
+	UStaticMesh* StaticMesh = Cast<UStaticMeshComponent>(Component)->GetStaticMesh();
+
+	// make sure transactional flag is on
+	StaticMesh->SetFlags(RF_Transactional);
+
+	bool bSavedToTransactionBuffer = StaticMesh->Modify();
+	check(bSavedToTransactionBuffer);
+	
+	int NewNumMaterials = MaterialSet.Materials.Num();
+	if (NewNumMaterials != StaticMesh->StaticMaterials.Num())
+	{
+		StaticMesh->StaticMaterials.SetNum(NewNumMaterials);
+	}
+	for (int k = 0; k < NewNumMaterials; ++k)
+	{
+		if (StaticMesh->GetMaterial(k) != MaterialSet.Materials[k])
+		{
+			StaticMesh->SetMaterial(k, MaterialSet.Materials[k]);
+		}
+	}
+
+	// right?
+	StaticMesh->PostEditChange();
+}
+
+
+void FStaticMeshComponentTarget::CommitMesh( const FCommitter& Committer )
 {
 	//bool bSaved = Component->Modify();
 	//check(bSaved);
@@ -26,9 +76,11 @@ FStaticMeshComponentTarget::CommitMesh( const FCommitter& Committer )
 
 	bool bSavedToTransactionBuffer = StaticMesh->Modify();
 	check(bSavedToTransactionBuffer);
-	FMeshDescription* MeshDescription = StaticMesh->GetMeshDescription(LODIndex);
 
-	Committer(MeshDescription);
+	FCommitParams CommitParams;
+	CommitParams.MeshDescription = StaticMesh->GetMeshDescription(LODIndex);
+
+	Committer(CommitParams);
 
 	StaticMesh->CommitMeshDescription(LODIndex);
 	StaticMesh->PostEditChange();
@@ -39,20 +91,18 @@ FStaticMeshComponentTarget::CommitMesh( const FCommitter& Committer )
 	//Component->PostEditChange();
 }
 
-bool
-FStaticMeshComponentTargetFactory::CanBuild(UActorComponent* Component)
+bool FStaticMeshComponentTargetFactory::CanBuild(UActorComponent* Component)
 {
 	UStaticMeshComponent* StaticMeshComponent = Cast<UStaticMeshComponent>(Component);
 	return StaticMeshComponent != nullptr;
 }
 
-TUniquePtr< FPrimitiveComponentTarget >
-FStaticMeshComponentTargetFactory::Build(UPrimitiveComponent* Component)
+TUniquePtr<FPrimitiveComponentTarget> FStaticMeshComponentTargetFactory::Build(UPrimitiveComponent* Component)
 {
 	UStaticMeshComponent* StaticMeshComponent = Cast<UStaticMeshComponent>(Component);
 	if (StaticMeshComponent != nullptr)
 	{
-		return TUniquePtr< FPrimitiveComponentTarget > { new FStaticMeshComponentTarget{Component} };
+		return TUniquePtr<FPrimitiveComponentTarget> { new FStaticMeshComponentTarget{Component} };
 	}
 	return {};
 }
diff --git a/Engine/Source/Editor/UnrealEd/Private/Tools/UEdMode.cpp b/Engine/Source/Editor/UnrealEd/Private/Tools/UEdMode.cpp
--- a/Engine/Source/Editor/UnrealEd/Private/Tools/UEdMode.cpp
+++ b/Engine/Source/Editor/UnrealEd/Private/Tools/UEdMode.cpp
@@ -23,6 +23,8 @@
 #include "LevelEditorViewport.h"
 #include "InteractiveToolManager.h"
 #include "InteractiveToolObjects.h"
+#include "Editor/EditorEngine.h"
+#include "GameFramework/Actor.h"
 
 
 class FEdModeQueriesImpl : public IToolsContextQueriesAPI
@@ -41,9 +43,42 @@ public:
 	void CacheCurrentViewState(FEditorViewportClient* ViewportClient)
 	{
 		FViewportCameraTransform ViewTransform = ViewportClient->GetViewTransform();
-		CachedViewState.Position = ViewTransform.GetLocation();
-		CachedViewState.Orientation = ViewTransform.GetRotation().Quaternion();
 		CachedViewState.bIsOrthographic = ViewportClient->IsOrtho();
+		CachedViewState.Position = ViewTransform.GetLocation();
+
+		// ViewTransform rotation is only initialized for perspective!
+		if (CachedViewState.bIsOrthographic == false)
+		{
+			CachedViewState.Orientation = ViewTransform.GetRotation().Quaternion();
+		}
+		else
+		{
+			// These rotations are based on hardcoded values in EditorViewportClient.cpp, see switches in FEditorViewportClient::CalcSceneView and FEditorViewportClient::Draw
+			switch (ViewportClient->ViewportType)
+			{
+			case LVT_OrthoXY:
+				CachedViewState.Orientation = FQuat(FRotator(-90.0f, -90.0f, 0.0f));
+				break;
+			case LVT_OrthoNegativeXY:
+				CachedViewState.Orientation = FQuat(FRotator(90.0f, 90.0f, 0.0f));
+				break;
+			case LVT_OrthoXZ:
+				CachedViewState.Orientation = FQuat(FRotator(0.0f, -90.0f, 0.0f));
+				break;
+			case LVT_OrthoNegativeXZ:
+				CachedViewState.Orientation = FQuat(FRotator(0.0f, 90.0f, 0.0f));
+				break;
+			case LVT_OrthoYZ:
+				CachedViewState.Orientation = FQuat(FRotator(0.0f, 0.0f, 0.0f));
+				break;
+			case LVT_OrthoNegativeYZ:
+				CachedViewState.Orientation = FQuat(FRotator(0.0f, 180.0f, 0.0f));
+				break;
+			default:
+				CachedViewState.Orientation = FQuat::Identity;
+			}
+		}
+
 		CachedViewState.bIsVR = false;
 	}
 
@@ -52,8 +87,11 @@ public:
 		StateOut.ToolManager = ToolsContext->ToolManager;
 		StateOut.GizmoManager = ToolsContext->GizmoManager;
 		StateOut.World = GEditor->GetWorld();
-		GEditor->GetSelectedActors()->GetSelectedObjects(StateOut.SelectedActors);
-		GEditor->GetSelectedComponents()->GetSelectedObjects(StateOut.SelectedComponents);
+		USelection* SelectionActorSet = GEditor->GetSelectedActors();
+		SelectionActorSet->GetSelectedObjects(StateOut.SelectedActors);
+		USelection* SelectionComponentSet = GEditor->GetSelectedComponents();
+		SelectionComponentSet->GetSelectedObjects(StateOut.SelectedComponents);
+
 	}
 
 
@@ -345,6 +383,9 @@ UEdMode::UEdMode()
 {
 	bDrawKillZ = true;
 	ToolsContext = nullptr;
+	ToolsContextClass = UInteractiveToolsContext::StaticClass();
+	ToolCommandList = MakeShareable(new FUICommandList);
+	bCheckIfDefaultToolNeeded = true;
 }
 
 
@@ -379,8 +420,6 @@ FRay UEdMode::GetRayFromMousePos(FEditorViewportClient* ViewportClient, FViewpor
 	return FRay(MouseViewportRay.GetOrigin(), MouseViewportRay.GetDirection(), true);
 }
 
-
-
 bool UEdMode::MouseEnter(FEditorViewportClient* ViewportClient, FViewport* Viewport, int32 x, int32 y)
 {
 	CurrentMouseState.Mouse.Position2D = FVector2D(x, y);
@@ -592,6 +631,14 @@ bool UEdMode::InputKey(FEditorViewportClient* ViewportClient, FViewport* Viewpor
 
 	if(!bHandled)
 	{
+
+		if (Event != IE_Released)
+		{
+			if (ToolCommandList->ProcessCommandBindings(Key, FSlateApplication::Get().GetModifierKeys(), false/*Event == IE_Repeat*/))
+			{
+				return true;
+			}
+		}
 		// Next pass input to the mode toolkit
 		if (Toolkit.IsValid() && ((Event == IE_Pressed) || (Event == IE_Repeat)))
 		{
@@ -612,7 +659,7 @@ bool UEdMode::InputKey(FEditorViewportClient* ViewportClient, FViewport* Viewpor
 		}
 	}
 
-	return false;
+	return bHandled;
 }
 
 bool UEdMode::InputAxis(FEditorViewportClient* InViewportClient, FViewport* Viewport, int32 ControllerId, FKey Key, float Delta, float DeltaTime)
@@ -663,6 +710,11 @@ void UEdMode::Tick(FEditorViewportClient* ViewportClient, float DeltaTime)
 		ToolsContext->ToolManager->Tick(DeltaTime);
 		ToolsContext->GizmoManager->Tick(DeltaTime);
 
+		if(bCheckIfDefaultToolNeeded && ToolsContext->ToolManager->GetActiveTool(EToolSide::Left) == nullptr)
+		{
+			ActivateDefaultTool();
+		}
+
 		if (bInvalidationPending)
 		{
 			ViewportClient->Invalidate();
@@ -679,13 +731,10 @@ void UEdMode::Tick(FEditorViewportClient* ViewportClient, float DeltaTime)
 	}
 }
 
-void UEdMode::ActorSelectionChangeNotify()
-{
-}
 
 bool UEdMode::HandleClick(FEditorViewportClient* InViewportClient, HHitProxy *HitProxy, const FViewportClick &Click)
 {
-	return false;
+	return true;
 }
 
 void UEdMode::Enter()
@@ -700,23 +749,32 @@ void UEdMode::Enter()
 
 	bPendingDeletion = false;
 
-	if (!Toolkit.IsValid())
-	{
-		Toolkit = MakeShareable(new FModeToolkit);
-		Toolkit->Init(Owner->GetToolkitHost());
-	}
-
 	// initialize the adapter that attaches the ToolsContext to this FEdMode
-	ToolsContext = NewObject<UInteractiveToolsContext>(GetTransientPackage(), TEXT("ToolsContext"), RF_Transient);
+	ToolsContext = NewObject<UInteractiveToolsContext>(GetTransientPackage(), ToolsContextClass.Get(), TEXT("ToolsContext"), RF_Transient);
 	TransactionAPI = new FEdModeTransactionImpl(ToolsContext, this);
 	QueriesAPI = new FEdModeQueriesImpl(ToolsContext, this);
 	AssetAPI = new FEditorToolAssetAPI();
 	ToolsContext->Initialize(QueriesAPI, TransactionAPI);
 
+	ToolsContext->OnToolNotificationMessage.AddLambda([this](const FText& Message)
+	{
+		OnToolNotificationMessage.Broadcast(Message);
+	});
+	ToolsContext->OnToolWarningMessage.AddLambda([this](const FText& Message)
+	{
+		OnToolWarningMessage.Broadcast(Message);
+	});
+
+
+	GetToolManager()->OnToolStarted.AddUObject(this, &UEdMode::OnToolStarted);
+	GetToolManager()->OnToolEnded.AddUObject(this, &UEdMode::OnToolEnded);
+
 	// enable auto invalidation in Editor, because invalidating for all hover and capture events is unpleasant
 	ToolsContext->InputRouter->bAutoInvalidateOnHover = true;
 	ToolsContext->InputRouter->bAutoInvalidateOnCapture = true;
 
+	// Now that the context is ready, make the toolkit
+	CreateToolkit();
 
 	// set up standard materials
 	StandardVertexColorMaterial = LoadObject<UMaterial>(nullptr, TEXT("/Game/Materials/VertexColor"));
@@ -734,13 +792,25 @@ void UEdMode::Enter()
 	FEditorDelegates::EditorModeIDEnter.Broadcast(GetID());
 	const bool bIsEnteringMode = true;
 	Owner->BroadcastEditorModeIDChanged(GetID(), bIsEnteringMode);
+
+}
+
+void UEdMode::RegisterTool(TSharedPtr<FUICommandInfo> UICommand, FString ToolIdentifier, UInteractiveToolBuilder* Builder)
+{
+	const TSharedRef<FUICommandList>& CommandList = Toolkit->GetToolkitCommands();
+	ToolsContext->ToolManager->RegisterToolType(ToolIdentifier, Builder);
+	CommandList->MapAction(UICommand,
+		FExecuteAction::CreateLambda([this, ToolIdentifier]() { ToolsContext->StartTool(EToolSide::Mouse, ToolIdentifier); }),
+		FCanExecuteAction::CreateLambda([this, ToolIdentifier]() { return GetToolManager()->CanActivateTool(EToolSide::Mouse, ToolIdentifier); }),
+		FIsActionChecked::CreateLambda([this, Builder]() { return ToolsContext->IsToolBuilderActive(EToolSide::Mouse, Builder); }));
 }
 
 void UEdMode::Exit()
 {
 	FEditorDelegates::BeginPIE.Remove(BeginPIEDelegateHandle);
 	FEditorDelegates::PreSaveWorld.Remove(PreSaveWorldDelegateHandle);
-
+	GetToolManager()->OnToolStarted.RemoveAll(this);
+	GetToolManager()->OnToolEnded.RemoveAll(this);
 	// auto-accept any in-progress tools
 	DeactivateAllActiveTools();
 
@@ -768,6 +838,9 @@ void UEdMode::Exit()
 		Toolkit.Reset();
 	}
 
+	OnToolNotificationMessage.Clear();
+	OnToolWarningMessage.Clear();
+
 	const bool bIsEnteringMode = false;
 	Owner->BroadcastEditorModeIDChanged(GetID(), bIsEnteringMode);
 	FEditorDelegates::EditorModeIDExit.Broadcast(GetID());
@@ -955,7 +1028,10 @@ AActor* UEdMode::GetFirstSelectedActorInstance() const
 
 void UEdMode::DeactivateAllActiveTools()
 {
-	ToolsContext->DeactivateAllActiveTools();
+	if (ToolsContext)
+	{
+		ToolsContext->DeactivateAllActiveTools();
+	}
 	RestoreEditorState();
 }
 
@@ -964,6 +1040,22 @@ UInteractiveToolManager* UEdMode::GetToolManager() const
 	return ToolsContext->ToolManager;
 }
 
+void UEdMode::CreateToolkit()
+{
+	if (!Toolkit.IsValid())
+	{
+		Toolkit = MakeShareable(new FModeToolkit);
+		Toolkit->Init(Owner->GetToolkitHost());
+	}
+
+	UClass* LoadedSettingsObject = SettingsClass.LoadSynchronous();
+	SettingsObject = NewObject<UObject>(this, LoadedSettingsObject);
+	if (SettingsObject)
+	{
+		Toolkit->SetModeSettingsObject(SettingsObject);
+	}
+}
+
 bool UEdMode::IsSnapRotationEnabled()
 {
 	return GetDefault<ULevelEditorViewportSettings>()->RotGridEnabled;
diff --git a/Engine/Source/Editor/UnrealEd/Public/Kismet2/DebuggerCommands.h b/Engine/Source/Editor/UnrealEd/Public/Kismet2/DebuggerCommands.h
--- a/Engine/Source/Editor/UnrealEd/Public/Kismet2/DebuggerCommands.h
+++ b/Engine/Source/Editor/UnrealEd/Public/Kismet2/DebuggerCommands.h
@@ -91,7 +91,6 @@ public:
 	TSharedPtr<FUICommandInfo> PlayInDefaultPlayerStart;
 	TSharedPtr<FUICommandInfo> PlayInSettings;
 	TSharedPtr<FUICommandInfo> PlayInNetworkSettings;
-	TSharedPtr<FUICommandInfo> PlayInNetworkDedicatedServer;
 
 	TArray< TSharedPtr< FUICommandInfo > > PlayInTargetedMobilePreviewDevices;
 
diff --git a/Engine/Source/Editor/UnrealEd/Public/SkelImport.h b/Engine/Source/Editor/UnrealEd/Public/SkelImport.h
--- a/Engine/Source/Editor/UnrealEd/Public/SkelImport.h
+++ b/Engine/Source/Editor/UnrealEd/Public/SkelImport.h
@@ -91,6 +91,8 @@ struct ExistingSkelMeshData
 	TMap<UAssetUserData*, bool>				ExistingAssetUserData;
 
 	USkeletalMesh::FOnMeshChanged			ExistingOnMeshChanged;
+
+	TMap<FName, FString>* ExistingUMetaDataTagValues;
 };
 
 /** 
diff --git a/Engine/Source/Editor/UnrealEd/Public/Toolkits/BaseToolkit.h b/Engine/Source/Editor/UnrealEd/Public/Toolkits/BaseToolkit.h
--- a/Engine/Source/Editor/UnrealEd/Public/Toolkits/BaseToolkit.h
+++ b/Engine/Source/Editor/UnrealEd/Public/Toolkits/BaseToolkit.h
@@ -10,6 +10,7 @@
 #include "Toolkits/IToolkitHost.h"
 
 class UInteractiveTool;
+class UInteractiveToolManager;
 class SDockableTab;
 class UEdMode;
 class IDetailsView;
@@ -63,6 +64,7 @@ public:
 	 */
 	void AddToolkitTab(const TSharedRef<SDockableTab>& TabToAdd, const EToolkitTabSpot::Type TabSpot);
 
+
 protected:
 
 	/** @return Returns the prefix string to use for tabs created for this toolkit.  In world-centric mode, tabs get a
@@ -106,6 +108,7 @@ public:
 
 	/** Initializes the mode toolkit */
 	virtual void Init(const TSharedPtr<IToolkitHost>& InitToolkitHost);
+	~FModeToolkit();
 
 public:
 
@@ -126,21 +129,43 @@ public:
 	virtual const TArray<UObject*>* GetObjectsCurrentlyBeingEdited() const override;
 	virtual FLinearColor GetWorldCentricTabColorScale() const override;
 	virtual class FEdMode* GetEditorMode() const override;
+	virtual FText GetEditorModeDisplayName() const override;
+	virtual FSlateIcon GetEditorModeIcon() const override;
 
 	virtual UEdMode* GetScriptableEditorMode() const;
+	virtual TSharedPtr<SWidget> GetInlineContent() const override;
 
-private:
+	/** Returns the number of Mode specific tabs in the mode toolbar **/
+	virtual void GetToolPaletteNames(TArray<FName>& PaletteNames) const {}
 
-	TSharedPtr<SWidget> ToolkitWidget;
-	TSharedPtr<IDetailsView> DetailsView;
+	/**
+	 * @param PaletteIndex      The index of the ToolPalette to build
+	 * @returns the name of Tool Palette
+	 **/
+	virtual FText GetToolPaletteDisplayName(FName Palette) const { return FText(); }
+
+	/**
+	 * @param PaletteIndex      The index of the ToolPalette to build
+	 * @param ToolbarBuilder    The builder to use for given PaletteIndex
+	**/
+	virtual void BuildToolPalette(FName Palette, class FToolBarBuilder& ToolbarBuilder);
 
-	// these functions just forward calls to the ToolsContext / ToolManager
 
+	virtual void OnToolPaletteChanged(FName PaletteName);
+
+	void SetModeSettingsObject(UObject* InSettingsObject);
+
+protected:
 	bool CanStartTool(const FString& ToolTypeIdentifier);
 	bool CanAcceptActiveTool();
 	bool CanCancelActiveTool();
 	bool CanCompleteActiveTool();
 
-	FReply StartTool(const FString& ToolTypeIdentifier);
-	FReply EndTool(EToolShutdownType ShutdownType);
+	virtual void OnToolStarted(UInteractiveToolManager* Manager, UInteractiveTool* Tool);
+	virtual void OnToolEnded(UInteractiveToolManager* Manager, UInteractiveTool* Tool);
+
+protected:
+	TSharedPtr<SWidget> ToolkitWidget;
+	TSharedPtr<IDetailsView> ModeDetailsView;
+	TSharedPtr<IDetailsView> DetailsView;
 };
diff --git a/Engine/Source/Editor/UnrealEd/Public/Tools/EditorComponentSourceFactory.h b/Engine/Source/Editor/UnrealEd/Public/Tools/EditorComponentSourceFactory.h
--- a/Engine/Source/Editor/UnrealEd/Public/Tools/EditorComponentSourceFactory.h
+++ b/Engine/Source/Editor/UnrealEd/Public/Tools/EditorComponentSourceFactory.h
@@ -17,9 +17,17 @@ public:
 class UNREALED_API FStaticMeshComponentTarget : public FPrimitiveComponentTarget
 {
 public:
+
 	FStaticMeshComponentTarget( UPrimitiveComponent* Component )
 		: FPrimitiveComponentTarget( Cast<UStaticMeshComponent>(Component) ){}
+
+	virtual void GetMaterialSet(FComponentMaterialSet& MaterialSetOut, bool bAssetMaterials) const override;
+
 	FMeshDescription* GetMesh() override;
+
 	void CommitMesh( const FCommitter& ) override;
+
+	virtual void CommitMaterialSetUpdate(const FComponentMaterialSet& MaterialSet, bool bApplyToAsset) override;
+
 	static const int LODIndex{0};
 };
diff --git a/Engine/Source/Editor/UnrealEd/Public/Tools/UEdMode.h b/Engine/Source/Editor/UnrealEd/Public/Tools/UEdMode.h
--- a/Engine/Source/Editor/UnrealEd/Public/Tools/UEdMode.h
+++ b/Engine/Source/Editor/UnrealEd/Public/Tools/UEdMode.h
@@ -11,6 +11,7 @@
 #include "UObject/Object.h"
 #include "Math/Ray.h"
 #include "InputState.h"
+#include "UObject/SoftObjectPtr.h"
 #include "UEdMode.generated.h"
 
 
@@ -27,6 +28,7 @@ struct FConvexVolume;
 struct FViewportClick;
 class UInteractiveToolsContext;
 class UInteractiveToolManager;
+class UInteractiveTool;
 class IToolsContextQueriesAPI;
 class IToolsContextTransactionsAPI;
 class IToolsContextAssetAPI;
@@ -38,6 +40,9 @@ struct FViewportClick;
 class FModeTool;
 class FEditorViewportClient;
 struct FViewportClick;
+class UInteractiveToolBuilder;
+class FUICommandInfo;
+class FUICommandList;
 
 
 /** Outcomes when determining whether it's possible to perform an action on the edit modes*/
@@ -57,7 +62,7 @@ namespace EEditAction
 /**
  * Base class for all editor modes.
  */
-UCLASS()
+UCLASS(Abstract)
 class UNREALED_API UEdMode : public UObject, public FEditorCommonDrawHelper
 {
 	GENERATED_BODY()
@@ -117,7 +122,7 @@ public:
 	virtual bool IsCompatibleWith(FEditorModeID OtherModeID) const { return false; }
 	virtual void ActorMoveNotify() {}
 	virtual void ActorsDuplicatedNotify(TArray<AActor*>& PreDuplicateSelection, TArray<AActor*>& PostDuplicateSelection, bool bOffsetLocations) {}
-	virtual void ActorSelectionChangeNotify();
+	virtual void ActorSelectionChangeNotify() {};
 	virtual void ActorPropChangeNotify() {}
 	virtual void MapChangeNotify() {}
 
@@ -164,6 +169,10 @@ public:
 	virtual void UpdateInternalData() {}
 
 	virtual void Enter();
+
+	virtual void RegisterTool(TSharedPtr<FUICommandInfo> UICommand, FString ToolIdentifier, UInteractiveToolBuilder* Builder);
+
+
 	virtual void Exit();
 	virtual UTexture2D* GetVertexTexture() { return GEngine->DefaultBSPVertexTexture; }
 
@@ -232,13 +241,6 @@ public:
 	/** Returns the owning mode manager for this mode */
 	class FEditorModeTools* GetModeManager() const;
 
-	/**
-	 * Called when the editor mode should rebuild its toolbar
-	 *
-	 * @param ToolbarBuilder	The builder which should be used to add toolbar widgets
-	 */
-	virtual void BuildModeToolbar(class FToolBarBuilder& ToolbarBuilder) {}
-
 public:
 
 	/** Request that this mode be deleted at the next convenient opportunity (FEditorModeTools::Tick) */
@@ -261,6 +263,7 @@ private:
 	virtual void TerminateActiveToolsOnSaveWorld();
 
 	FRay GetRayFromMousePos(FEditorViewportClient* ViewportClient, FViewport* Viewport, int MouseX, int MouseY);
+	
 protected:
 
 	/** Information pertaining to this mode. Assigned by FEditorModeRegistry. */
@@ -299,10 +302,31 @@ public:
 	 */
 	virtual UInteractiveToolManager* GetToolManager() const;
 
+	virtual TMap<FName, TArray<TSharedPtr<FUICommandInfo>>> GetModeCommands() const
+	{
+		return TMap<FName, TArray<TSharedPtr<FUICommandInfo>>>();
+	};
+
+	void SetCurrentPaletteName(FName InName)
+	{
+		CurrentPaletteName = InName;
+		UpdateOnPaletteChange();
+	};
+
 protected:
+	virtual void CreateToolkit();
+	virtual void OnToolStarted(UInteractiveToolManager* Manager, UInteractiveTool* Tool) {};
+	virtual void OnToolEnded(UInteractiveToolManager* Manager, UInteractiveTool* Tool) { bCheckIfDefaultToolNeeded = true; };
+	virtual void ActivateDefaultTool() {};
+	virtual void UpdateOnPaletteChange() {};
 
-	UInteractiveToolsContext* ToolsContext;
+protected:
 
+	UPROPERTY()
+	UInteractiveToolsContext* ToolsContext;
+	TSoftClassPtr<UInteractiveToolsContext> ToolsContextClass;
+	/** Command list lives here so that the key bindings on the commands can be processed in the viewport. */
+	TSharedPtr<FUICommandList> ToolCommandList;
 	bool bInvalidationPending;
 	IToolsContextQueriesAPI* QueriesAPI;
 	IToolsContextTransactionsAPI* TransactionAPI;
@@ -313,5 +337,17 @@ protected:
 	// called before a Save starts. This currently shuts down active tools.
 	FDelegateHandle PreSaveWorldDelegateHandle;
 
+	UPROPERTY()
+	TSoftClassPtr<UObject> SettingsClass;
+
+	UPROPERTY(Transient)
+	UObject* SettingsObject;
+
+	DECLARE_MULTICAST_DELEGATE_OneParam(FOnModeToolNotification, const FText&);
+	FOnModeToolNotification OnToolNotificationMessage;
+	FOnModeToolNotification OnToolWarningMessage;
+
+	FName CurrentPaletteName;
+	bool bCheckIfDefaultToolNeeded;
 };
 
diff --git a/Engine/Source/Editor/VREditor/VREditorMode.cpp b/Engine/Source/Editor/VREditor/VREditorMode.cpp
--- a/Engine/Source/Editor/VREditor/VREditorMode.cpp
+++ b/Engine/Source/Editor/VREditor/VREditorMode.cpp
@@ -898,9 +898,11 @@ void UVREditorMode::ToggleSIEAndVREditor()
 {
 	if (GEditor->EditorWorld == nullptr && !GEditor->bIsSimulatingInEditor)
 	{
-		const FVector* StartLoc = NULL;
-		const FRotator* StartRot = NULL;
-		GEditor->RequestPlaySession(false, VREditorLevelViewportWeakPtr.Pin(), true /*bSimulateInEditor*/, StartLoc, StartRot, -1);
+		FRequestPlaySessionParams SessionParams;
+		SessionParams.DestinationSlateViewport = VREditorLevelViewportWeakPtr;
+		SessionParams.WorldType = EPlaySessionWorldType::SimulateInEditor;
+
+		GEditor->RequestPlaySession(SessionParams);
 	}
 	else if (GEditor->PlayWorld != nullptr && GEditor->bIsSimulatingInEditor)
 	{
@@ -913,10 +915,16 @@ void UVREditorMode::TogglePIEAndVREditor()
 	bool bRequestedPIE = false;
 	if (GEditor->EditorWorld == nullptr && GEditor->PlayWorld == nullptr && !GEditor->bIsSimulatingInEditor)
 	{
-		const FVector* StartLoc = NULL;
-		const FRotator* StartRot = NULL;
+		FRequestPlaySessionParams SessionParams;
+		SessionParams.DestinationSlateViewport = VREditorLevelViewportWeakPtr;
+		SessionParams.WorldType = EPlaySessionWorldType::PlayInEditor;
+		
 		const bool bHMDIsReady = (GEngine && GEngine->XRSystem.IsValid() && GEngine->XRSystem->GetHMDDevice() && GEngine->XRSystem->GetHMDDevice()->IsHMDConnected());
-		GEditor->RequestPlaySession(true, VREditorLevelViewportWeakPtr.Pin(), false /*bSimulateInEditor*/, StartLoc, StartRot, -1, false, bHMDIsReady);
+		if (bHMDIsReady)
+		{
+			SessionParams.SessionPreviewTypeOverride = EPlaySessionPreviewType::VRPreview;
+		}
+		GEditor->RequestPlaySession(SessionParams);
 		bRequestedPIE = true;
 	}
 	else if (GEditor->PlayWorld != nullptr)
diff --git a/Engine/Source/Programs/Enterprise/Datasmith/DatasmithRevitExporter/Private/DatasmithRevitExportContext.cs b/Engine/Source/Programs/Enterprise/Datasmith/DatasmithRevitExporter/Private/DatasmithRevitExportContext.cs
--- a/Engine/Source/Programs/Enterprise/Datasmith/DatasmithRevitExporter/Private/DatasmithRevitExportContext.cs
+++ b/Engine/Source/Programs/Enterprise/Datasmith/DatasmithRevitExporter/Private/DatasmithRevitExportContext.cs
@@ -193,7 +193,7 @@ namespace DatasmithRevitExporter
 				return RenderNodeAction.Proceed;
 			}
 
-			return RenderNodeAction.Skip;
+            return RenderNodeAction.Skip;
 		}
 
 		// OnElementEnd marks the end of an element being exported.
@@ -237,7 +237,7 @@ namespace DatasmithRevitExporter
 				return RenderNodeAction.Proceed;
 			}
 
-			return RenderNodeAction.Skip;
+            return RenderNodeAction.Skip;
 		}
 
 		// OnInstanceEnd marks the end of a family instance being exported.
@@ -724,11 +724,11 @@ namespace DatasmithRevitExporter
 
 				CreateMeshActor(InWorldTransform, out InstanceData.InstanceMesh, out InstanceData.InstanceActor);
 
-				// The Datasmith instance actor is a component in the hierarchy.
-				InstanceData.InstanceActor.SetIsComponent(true);
-			}
+                // The Datasmith instance actor is a component in the hierarchy.
+                //InstanceData.InstanceActor.SetIsComponent(true);
+            }
 
-			public FDatasmithFacadeMesh PopInstance()
+            public FDatasmithFacadeMesh PopInstance()
 			{
 				FInstanceData InstanceData = InstanceDataStack.Pop();
 
@@ -765,6 +765,10 @@ namespace DatasmithRevitExporter
 				string LayerName = Category.GetCategory(CurrentElement.Document, BuiltInCategory.OST_LightingFixtureSource)?.Name ?? "Light Sources";
 				SetActorProperties(LayerName, LightActor);
 
+				// Set the Datasmith light actor layer to its predefined name.
+				string CategoryName = Category.GetCategory(CurrentElement.Document, BuiltInCategory.OST_LightingFixtureSource)?.Name ?? "Light Sources";
+				LightActor.SetLayer(CategoryName);
+
 				// Set the specific properties of the Datasmith light actor.
 				FDatasmithRevitLight.SetLightProperties(InLightAsset, CurrentElement, LightActor);
 
diff --git a/Engine/Source/Programs/Enterprise/Datasmith/DatasmithRevitExporter/Private/DatasmithRevitMaterial.cs b/Engine/Source/Programs/Enterprise/Datasmith/DatasmithRevitExporter/Private/DatasmithRevitMaterial.cs
--- a/Engine/Source/Programs/Enterprise/Datasmith/DatasmithRevitExporter/Private/DatasmithRevitMaterial.cs
+++ b/Engine/Source/Programs/Enterprise/Datasmith/DatasmithRevitExporter/Private/DatasmithRevitMaterial.cs
@@ -55,7 +55,7 @@ namespace DatasmithRevitExporter
 			ExtraTexturePaths = InExtraTexturePaths;
 
 			// Create a new Datasmith master material.
-			MasterMaterial = new FDatasmithFacadeMaterial(GetMaterialName(InMaterialNode, CurrentMaterial), MaterialLabel);
+			MasterMaterial = new FDatasmithFacadeMaterial(GetMaterialName(InMaterialNode, CurrentMaterial), GetMaterialLabel(InMaterialNode, CurrentMaterial));
 
 			// Hash the Datasmith master material name to shorten it.
 			MasterMaterial.HashName();
diff --git a/Engine/Source/Programs/UnrealMultiUserServer/Private/UnrealMultiUserServer.cpp b/Engine/Source/Programs/UnrealMultiUserServer/Private/UnrealMultiUserServer.cpp
--- a/Engine/Source/Programs/UnrealMultiUserServer/Private/UnrealMultiUserServer.cpp
+++ b/Engine/Source/Programs/UnrealMultiUserServer/Private/UnrealMultiUserServer.cpp
@@ -113,6 +113,18 @@ static CommandLineArguments GSavedCommandLine;
 	NSAppleEventManager* appleEventManager = [NSAppleEventManager sharedAppleEventManager];
 	[appleEventManager setEventHandler:self andSelector:@selector(handleQuitEvent:withReplyEvent:) forEventClass:kCoreEventClass andEventID:kAEQuitApplication];
 
+	// Add a menu bar to the application.
+	id menubar = [[NSMenu new] autorelease];
+	id appMenuItem = [[NSMenuItem new] autorelease];
+	[menubar addItem:appMenuItem];
+	[NSApp setMainMenu:menubar];
+
+	// Populate the menu bar.
+	id appMenu = [[NSMenu new] autorelease];
+	id quitMenuItem = [[[NSMenuItem alloc] initWithTitle:NSLOCTEXT("UMUS_Quit", "QuitApp", "Quit").ToString().GetNSString() action:@selector(terminate:) keyEquivalent:@"q"] autorelease];
+	[appMenu addItem:quitMenuItem];
+	[appMenuItem setSubmenu:appMenu];
+
 	RunGameThread(self, @selector(runGameThread:));
 }
 
diff --git a/Engine/Source/Runtime/ApplicationCore/Private/IOS/Accessibility/IOSAccessibilityElement.cpp b/Engine/Source/Runtime/ApplicationCore/Private/IOS/Accessibility/IOSAccessibilityElement.cpp
--- a/Engine/Source/Runtime/ApplicationCore/Private/IOS/Accessibility/IOSAccessibilityElement.cpp
+++ b/Engine/Source/Runtime/ApplicationCore/Private/IOS/Accessibility/IOSAccessibilityElement.cpp
@@ -5,7 +5,7 @@
 #include "IOS/Accessibility/IOSAccessibilityElement.h"
 
 #include "Async/TaskGraphInterfaces.h"
-#include "GenericPlatform/GenericAccessibleInterfaces.h"
+#include "GenericPlatform/Accessibility/GenericAccessibleInterfaces.h"
 #include "IOS/Accessibility/IOSAccessibilityCache.h"
 #include "IOS/IOSApplication.h"
 #include "IOS/IOSAppDelegate.h"
diff --git a/Engine/Source/Runtime/ApplicationCore/Private/Windows/Accessibility/WindowsUIABaseProvider.cpp b/Engine/Source/Runtime/ApplicationCore/Private/Windows/Accessibility/WindowsUIABaseProvider.cpp
--- a/Engine/Source/Runtime/ApplicationCore/Private/Windows/Accessibility/WindowsUIABaseProvider.cpp
+++ b/Engine/Source/Runtime/ApplicationCore/Private/Windows/Accessibility/WindowsUIABaseProvider.cpp
@@ -6,7 +6,7 @@
 
 #include "Misc/AssertionMacros.h"
 #include "Windows/Accessibility/WindowsUIAManager.h"
-#include "GenericPlatform/GenericAccessibleInterfaces.h"
+#include "GenericPlatform/Accessibility/GenericAccessibleInterfaces.h"
 
 FWindowsUIABaseProvider::FWindowsUIABaseProvider(FWindowsUIAManager& InManager, TSharedRef<IAccessibleWidget> InWidget)
 	: UIAManager(&InManager)
diff --git a/Engine/Source/Runtime/ApplicationCore/Private/Windows/Accessibility/WindowsUIAControlProvider.cpp b/Engine/Source/Runtime/ApplicationCore/Private/Windows/Accessibility/WindowsUIAControlProvider.cpp
--- a/Engine/Source/Runtime/ApplicationCore/Private/Windows/Accessibility/WindowsUIAControlProvider.cpp
+++ b/Engine/Source/Runtime/ApplicationCore/Private/Windows/Accessibility/WindowsUIAControlProvider.cpp
@@ -6,7 +6,7 @@
 #include "Windows/Accessibility/WindowsUIAWidgetProvider.h"
 #include "Windows/Accessibility/WindowsUIAPropertyGetters.h"
 #include "Windows/Accessibility/WindowsUIAManager.h"
-#include "GenericPlatform/GenericAccessibleInterfaces.h"
+#include "GenericPlatform/Accessibility/GenericAccessibleInterfaces.h"
 
 // FWindowsUIATextRangeProvider
 
@@ -66,15 +66,16 @@ ULONG STDCALL FWindowsUIATextRangeProvider::Release()
 
 HRESULT STDCALL FWindowsUIATextRangeProvider::Clone(ITextRangeProvider** pRetVal)
 {
-	if (IsValid())
-	{
-		*pRetVal = new FWindowsUIATextRangeProvider(*UIAManager, Widget, TextRange);
-		return S_OK;
-	}
-	else
-	{
-		return UIA_E_ELEMENTNOTAVAILABLE;
-	}
+	HRESULT ReturnValue = UIA_E_ELEMENTNOTAVAILABLE;
+	UIAManager->RunInGameThreadBlocking(
+		[this, &ReturnValue, &pRetVal]() {
+		if (IsValid())
+		{
+			*pRetVal = new FWindowsUIATextRangeProvider(*UIAManager, Widget, TextRange);
+			ReturnValue = S_OK;
+		}
+	});
+	return ReturnValue;
 }
 
 HRESULT STDCALL FWindowsUIATextRangeProvider::Compare(ITextRangeProvider* range, BOOL* pRetVal)
@@ -82,7 +83,11 @@ HRESULT STDCALL FWindowsUIATextRangeProvider::Compare(ITextRangeProvider* range,
 	// The documentation states that different endpoints that produce the same text are not equal,
 	// but doesn't say anything about the same endpoints that come from different control providers.
 	// Perhaps we can assume that comparing text ranges from different Widgets is not valid.
-	*pRetVal = TextRange == static_cast<FWindowsUIATextRangeProvider*>(range)->TextRange;
+	UIAManager->RunInGameThreadBlocking(
+		[this, &range, &pRetVal]() {
+
+		*pRetVal = TextRange == static_cast<FWindowsUIATextRangeProvider*>(range)->TextRange;
+	});
 	return S_OK;
 }
 
@@ -97,34 +102,39 @@ HRESULT STDCALL FWindowsUIATextRangeProvider::CompareEndpoints(TextPatternRangeE
 
 HRESULT STDCALL FWindowsUIATextRangeProvider::ExpandToEnclosingUnit(TextUnit unit)
 {
-	if (IsValid())
-	{
-		switch (unit)
+	HRESULT ReturnValue = UIA_E_ELEMENTNOTAVAILABLE;
+	UIAManager->RunInGameThreadBlocking(
+		[this, &ReturnValue, &unit]() {
+		if (IsValid())
 		{
-		case TextUnit_Character:
-			TextRange.EndIndex = FMath::Min(TextRange.BeginIndex + 1, Widget->AsText()->GetText().Len());
-			break;
-		case TextUnit_Format:
-			return E_NOTIMPL;
-		case TextUnit_Word:
-			return E_NOTIMPL;
-		case TextUnit_Line:
-			return E_NOTIMPL;
-		case TextUnit_Paragraph:
-			return E_NOTIMPL;
-		case TextUnit_Page:
-			return E_NOTIMPL;
-		case TextUnit_Document:
-			TextRange = FTextRange(0, Widget->AsText()->GetText().Len());
-			break;
+			switch (unit)
+			{
+			case TextUnit_Character:
+				TextRange.EndIndex = FMath::Min(TextRange.BeginIndex + 1, Widget->AsText()->GetText().Len());
+				break;
+			case TextUnit_Format:
+				ReturnValue = E_NOTIMPL;
+				break;
+			case TextUnit_Word:
+				ReturnValue = E_NOTIMPL;
+				break;
+			case TextUnit_Line:
+				ReturnValue = E_NOTIMPL;
+				break;
+			case TextUnit_Paragraph:
+				ReturnValue = E_NOTIMPL;
+				break;
+			case TextUnit_Page:
+				ReturnValue = E_NOTIMPL;
+				break;
+			case TextUnit_Document:
+				TextRange = FTextRange(0, Widget->AsText()->GetText().Len());
+				break;
+			}
+			ReturnValue = S_OK;
 		}
-
-		return S_OK;
-	}
-	else
-	{
-		return UIA_E_ELEMENTNOTAVAILABLE;
-	}
+	});
+	return ReturnValue;
 }
 
 HRESULT STDCALL FWindowsUIATextRangeProvider::FindAttribute(TEXTATTRIBUTEID attributeId, VARIANT val, BOOL backward, ITextRangeProvider** pRetVal)
@@ -134,25 +144,26 @@ HRESULT STDCALL FWindowsUIATextRangeProvider::FindAttribute(TEXTATTRIBUTEID attr
 
 HRESULT STDCALL FWindowsUIATextRangeProvider::FindText(BSTR text, BOOL backward, BOOL ignoreCase, ITextRangeProvider** pRetVal)
 {
-	if (IsValid())
-	{
-		FString TextToSearch(text);
-		int32 FoundIndex = TextFromTextRange().Find(TextToSearch, ignoreCase ? ESearchCase::IgnoreCase : ESearchCase::CaseSensitive, backward ? ESearchDir::FromEnd : ESearchDir::FromStart);
-		if (FoundIndex == INDEX_NONE)
-		{
-			*pRetVal = nullptr;
-		}
-		else
+	HRESULT ReturnValue = UIA_E_ELEMENTNOTAVAILABLE;
+	UIAManager->RunInGameThreadBlocking(
+		[this, &ReturnValue, &text, &backward, &ignoreCase, &pRetVal]() {
+		if (IsValid())
 		{
-			const int32 StartIndex = TextRange.BeginIndex + FoundIndex;
-			*pRetVal = new FWindowsUIATextRangeProvider(*UIAManager, Widget, FTextRange(StartIndex, StartIndex + TextToSearch.Len()));
+			FString TextToSearch(text);
+			int32 FoundIndex = TextFromTextRange().Find(TextToSearch, ignoreCase ? ESearchCase::IgnoreCase : ESearchCase::CaseSensitive, backward ? ESearchDir::FromEnd : ESearchDir::FromStart);
+			if (FoundIndex == INDEX_NONE)
+			{
+				*pRetVal = nullptr;
+			}
+			else
+			{
+				const int32 StartIndex = TextRange.BeginIndex + FoundIndex;
+				*pRetVal = new FWindowsUIATextRangeProvider(*UIAManager, Widget, FTextRange(StartIndex, StartIndex + TextToSearch.Len()));
+			}
+			ReturnValue = S_OK;
 		}
-		return S_OK;
-	}
-	else
-	{
-		return UIA_E_ELEMENTNOTAVAILABLE;
-	}
+	});
+	return ReturnValue;
 }
 
 HRESULT STDCALL FWindowsUIATextRangeProvider::GetAttributeValue(TEXTATTRIBUTEID attributeId, VARIANT* pRetVal)
@@ -167,28 +178,30 @@ HRESULT STDCALL FWindowsUIATextRangeProvider::GetBoundingRectangles(SAFEARRAY**
 
 HRESULT STDCALL FWindowsUIATextRangeProvider::GetEnclosingElement(IRawElementProviderSimple** pRetVal)
 {
-	if (IsValid())
-	{
-		*pRetVal = static_cast<IRawElementProviderSimple*>(&UIAManager->GetWidgetProvider(Widget));
-		return S_OK;
-	}
-	else
-	{
-		return UIA_E_ELEMENTNOTAVAILABLE;
-	}
+	HRESULT ReturnValue = UIA_E_ELEMENTNOTAVAILABLE;
+	UIAManager->RunInGameThreadBlocking(
+		[this, &ReturnValue, &pRetVal]() {
+		if (IsValid())
+		{
+			*pRetVal = static_cast<IRawElementProviderSimple*>(&UIAManager->GetWidgetProvider(Widget));
+			ReturnValue = S_OK;
+		}
+	});
+	return ReturnValue;
 }
 
 HRESULT STDCALL FWindowsUIATextRangeProvider::GetText(int maxLength, BSTR* pRetVal)
 {
-	if (IsValid())
-	{
-		*pRetVal = SysAllocString(*TextFromTextRange().Left(maxLength));
-		return S_OK;
-	}
-	else
-	{
-		return UIA_E_ELEMENTNOTAVAILABLE;
-	}
+	HRESULT ReturnValue = UIA_E_ELEMENTNOTAVAILABLE;
+	UIAManager->RunInGameThreadBlocking(
+		[this, &ReturnValue, &maxLength, &pRetVal]() {
+		if (IsValid())
+		{
+			*pRetVal = SysAllocString(*TextFromTextRange().Left(maxLength));
+			ReturnValue = S_OK;
+		}
+	});
+	return ReturnValue;
 }
 
 HRESULT STDCALL FWindowsUIATextRangeProvider::Move(TextUnit unit, int count, int* pRetVal)
@@ -203,24 +216,28 @@ HRESULT STDCALL FWindowsUIATextRangeProvider::MoveEndpointByUnit(TextPatternRang
 
 HRESULT STDCALL FWindowsUIATextRangeProvider::MoveEndpointByRange(TextPatternRangeEndpoint endpoint, ITextRangeProvider* targetRange, TextPatternRangeEndpoint targetEndpoint)
 {
-	FWindowsUIATextRangeProvider* CastedRange = static_cast<FWindowsUIATextRangeProvider*>(targetRange);
-	int32 NewIndex = (targetEndpoint == TextPatternRangeEndpoint_Start) ? CastedRange->TextRange.BeginIndex : CastedRange->TextRange.EndIndex;
-	if (endpoint == TextPatternRangeEndpoint_Start)
-	{
-		TextRange.BeginIndex = NewIndex;
-		if (TextRange.BeginIndex > TextRange.EndIndex)
+	UIAManager->RunInGameThreadBlocking(
+		[this, &endpoint, &targetRange, &targetEndpoint]() {
+
+		FWindowsUIATextRangeProvider* CastedRange = static_cast<FWindowsUIATextRangeProvider*>(targetRange);
+		int32 NewIndex = (targetEndpoint == TextPatternRangeEndpoint_Start) ? CastedRange->TextRange.BeginIndex : CastedRange->TextRange.EndIndex;
+		if (endpoint == TextPatternRangeEndpoint_Start)
 		{
-			TextRange.EndIndex = TextRange.BeginIndex;
+			TextRange.BeginIndex = NewIndex;
+			if (TextRange.BeginIndex > TextRange.EndIndex)
+			{
+				TextRange.EndIndex = TextRange.BeginIndex;
+			}
 		}
-	}
-	else
-	{
-		TextRange.EndIndex = NewIndex;
-		if (TextRange.BeginIndex > TextRange.EndIndex)
+		else
 		{
-			TextRange.BeginIndex = TextRange.EndIndex;
+			TextRange.EndIndex = NewIndex;
+			if (TextRange.BeginIndex > TextRange.EndIndex)
+			{
+				TextRange.BeginIndex = TextRange.EndIndex;
+			}
 		}
-	}
+	});
 	return S_OK;
 }
 
@@ -315,119 +332,132 @@ ULONG STDCALL FWindowsUIAControlProvider::Release()
 
 HRESULT STDCALL FWindowsUIAControlProvider::Invoke()
 {
-	if (IsValid())
-	{
-		Widget->AsActivatable()->Activate();
-		return S_OK;
-	}
-	else
-	{
-		return UIA_E_ELEMENTNOTAVAILABLE;
-	}
+	HRESULT ReturnValue = UIA_E_ELEMENTNOTAVAILABLE;
+	UIAManager->RunInGameThreadBlocking(
+		[this, &ReturnValue]() {
+		if (IsValid())
+		{
+			Widget->AsActivatable()->Activate();
+			ReturnValue = S_OK;
+		}
+	});
+	return ReturnValue;
+	
 }
 
 HRESULT STDCALL FWindowsUIAControlProvider::SetValue(double val)
 {
-	if (IsValid())
-	{
-		Widget->AsProperty()->SetValue(FString::SanitizeFloat(val));
-		return S_OK;
-	}
-	else
-	{
-		return UIA_E_ELEMENTNOTAVAILABLE;
-	}
+	HRESULT ReturnValue = UIA_E_ELEMENTNOTAVAILABLE;
+	UIAManager->RunInGameThreadBlocking(
+		[this, &ReturnValue, &val]() {
+		if (IsValid())
+		{
+			Widget->AsProperty()->SetValue(FString::SanitizeFloat(val));
+			ReturnValue = S_OK;
+		}
+	});
+	return ReturnValue;
 }
 
 HRESULT STDCALL FWindowsUIAControlProvider::get_Value(double* pRetVal)
 {
-	if (IsValid())
-	{
-		*pRetVal = WindowsUIAPropertyGetters::GetPropertyValue(Widget, UIA_RangeValueValuePropertyId).GetValue<double>();
-		return S_OK;
-	}
-	else
-	{
-		return UIA_E_ELEMENTNOTAVAILABLE;
-	}
+	HRESULT ReturnValue = UIA_E_ELEMENTNOTAVAILABLE;
+	UIAManager->RunInGameThreadBlocking(
+		[this, &ReturnValue, &pRetVal]() {
+		if (IsValid())
+		{
+			*pRetVal = WindowsUIAPropertyGetters::GetPropertyValue(Widget, UIA_RangeValueValuePropertyId).GetValue<double>();
+			ReturnValue = S_OK;
+		}
+	});
+	return ReturnValue;
 }
 
 HRESULT STDCALL FWindowsUIAControlProvider::get_IsReadOnly(BOOL* pRetVal)
 {
-	if (IsValid())
-	{
-		*pRetVal = WindowsUIAPropertyGetters::GetPropertyValue(Widget, UIA_ValueIsReadOnlyPropertyId).GetValue<bool>();
-		return S_OK;
-	}
-	else
-	{
-		return UIA_E_ELEMENTNOTAVAILABLE;
-	}
+	HRESULT ReturnValue = UIA_E_ELEMENTNOTAVAILABLE;
+	UIAManager->RunInGameThreadBlocking(
+		[this, &ReturnValue, &pRetVal]() {
+		if (IsValid())
+		{
+			*pRetVal = WindowsUIAPropertyGetters::GetPropertyValue(Widget, UIA_ValueIsReadOnlyPropertyId).GetValue<bool>();
+			ReturnValue = S_OK;
+		}
+	});
+	return ReturnValue;
 }
 
 HRESULT STDCALL FWindowsUIAControlProvider::get_Maximum(double* pRetVal)
 {
-	if (IsValid())
-	{
-		*pRetVal = WindowsUIAPropertyGetters::GetPropertyValue(Widget, UIA_RangeValueMaximumPropertyId).GetValue<double>();
-		return S_OK;
-	}
-	else
-	{
-		return UIA_E_ELEMENTNOTAVAILABLE;
-	}
+	HRESULT ReturnValue = UIA_E_ELEMENTNOTAVAILABLE;
+	UIAManager->RunInGameThreadBlocking(
+		[this, &ReturnValue, &pRetVal]() {
+		if (IsValid())
+		{
+			*pRetVal = WindowsUIAPropertyGetters::GetPropertyValue(Widget, UIA_RangeValueMaximumPropertyId).GetValue<double>();
+			ReturnValue = S_OK;
+		}
+	});
+	return ReturnValue;
 }
 
 HRESULT STDCALL FWindowsUIAControlProvider::get_Minimum(double* pRetVal)
 {
-	if (IsValid())
-	{
-		*pRetVal = WindowsUIAPropertyGetters::GetPropertyValue(Widget, UIA_RangeValueMinimumPropertyId).GetValue<double>();
-		return S_OK;
-	}
-	else
-	{
-		return UIA_E_ELEMENTNOTAVAILABLE;
-	}
+	HRESULT ReturnValue = UIA_E_ELEMENTNOTAVAILABLE;
+	UIAManager->RunInGameThreadBlocking(
+		[this, &ReturnValue, &pRetVal]() {
+		if (IsValid())
+		{
+			*pRetVal = WindowsUIAPropertyGetters::GetPropertyValue(Widget, UIA_RangeValueMinimumPropertyId).GetValue<double>();
+			ReturnValue = S_OK;
+		}
+
+	});
+	return ReturnValue;
 }
 
 HRESULT STDCALL FWindowsUIAControlProvider::get_LargeChange(double* pRetVal)
 {
-	if (IsValid())
-	{
-		*pRetVal = WindowsUIAPropertyGetters::GetPropertyValue(Widget, UIA_RangeValueLargeChangePropertyId).GetValue<double>();
-		return S_OK;
-	}
-	else
-	{
-		return UIA_E_ELEMENTNOTAVAILABLE;
-	}
+	HRESULT ReturnValue = UIA_E_ELEMENTNOTAVAILABLE;
+	UIAManager->RunInGameThreadBlocking(
+		[this, &ReturnValue, &pRetVal]() {
+		if (IsValid())
+		{
+			*pRetVal = WindowsUIAPropertyGetters::GetPropertyValue(Widget, UIA_RangeValueLargeChangePropertyId).GetValue<double>();
+			ReturnValue = S_OK;
+		}
+
+	});
+	return ReturnValue;
 }
 
 HRESULT STDCALL FWindowsUIAControlProvider::get_SmallChange(double* pRetVal)
 {
-	if (IsValid())
-	{
-		*pRetVal = WindowsUIAPropertyGetters::GetPropertyValue(Widget, UIA_RangeValueSmallChangePropertyId).GetValue<double>();
-		return S_OK;
-	}
-	else
-	{
-		return UIA_E_ELEMENTNOTAVAILABLE;
-	}
+	HRESULT ReturnValue = UIA_E_ELEMENTNOTAVAILABLE;
+	UIAManager->RunInGameThreadBlocking(
+		[this, &ReturnValue, &pRetVal]() {
+		if (IsValid())
+		{
+			*pRetVal = WindowsUIAPropertyGetters::GetPropertyValue(Widget, UIA_RangeValueSmallChangePropertyId).GetValue<double>();
+			ReturnValue = S_OK;
+		}
+	});
+	return ReturnValue;
 }
 
 HRESULT STDCALL FWindowsUIAControlProvider::get_DocumentRange(ITextRangeProvider** pRetVal)
 {
-	if (IsValid())
-	{
-		*pRetVal = static_cast<ITextRangeProvider*>(new FWindowsUIATextRangeProvider(*UIAManager, Widget, FTextRange(0, Widget->AsText()->GetText().Len())));
-		return S_OK;
-	}
-	else
-	{
-		return UIA_E_ELEMENTNOTAVAILABLE;
-	}
+	HRESULT ReturnValue = UIA_E_ELEMENTNOTAVAILABLE;
+	UIAManager->RunInGameThreadBlocking(
+		[this, &ReturnValue, &pRetVal]() {
+		if (IsValid())
+		{
+			*pRetVal = static_cast<ITextRangeProvider*>(new FWindowsUIATextRangeProvider(*UIAManager, Widget, FTextRange(0, Widget->AsText()->GetText().Len())));
+			ReturnValue = S_OK;
+		}
+
+	});
+	return ReturnValue;
 }
 
 HRESULT STDCALL FWindowsUIAControlProvider::get_SupportedTextSelection(SupportedTextSelection* pRetVal)
@@ -472,67 +502,73 @@ HRESULT STDCALL FWindowsUIAControlProvider::get_ToggleState(ToggleState* pRetVal
 	//	*pRetVal = ToggleState_Indeterminate;
 	//	break;
 	//}
-	if (IsValid())
-	{
-		*pRetVal = static_cast<ToggleState>(WindowsUIAPropertyGetters::GetPropertyValue(Widget, UIA_ToggleToggleStatePropertyId).GetValue<int32>());
-		return S_OK;
-	}
-	else
-	{
-		return UIA_E_ELEMENTNOTAVAILABLE;
-	}
+
+	HRESULT ReturnValue = UIA_E_ELEMENTNOTAVAILABLE;
+	UIAManager->RunInGameThreadBlocking(
+		[this, &ReturnValue, &pRetVal]() {
+		if (IsValid())
+		{
+			*pRetVal = static_cast<ToggleState>(WindowsUIAPropertyGetters::GetPropertyValue(Widget, UIA_ToggleToggleStatePropertyId).GetValue<int32>());
+			ReturnValue = S_OK;
+		}
+	});
+	return ReturnValue;
 }
 
 HRESULT STDCALL FWindowsUIAControlProvider::Toggle()
 {
-	if (IsValid())
-	{
-		Widget->AsActivatable()->Activate();
-		return S_OK;
-	}
-	else
-	{
-		return UIA_E_ELEMENTNOTAVAILABLE;
-	}
+	HRESULT ReturnValue = UIA_E_ELEMENTNOTAVAILABLE;
+	UIAManager->RunInGameThreadBlocking(
+		[this, &ReturnValue]() {
+		if (IsValid())
+		{
+			Widget->AsActivatable()->Activate();
+			ReturnValue = S_OK;
+		}
+	});
+	return ReturnValue;
 }
 
 HRESULT STDCALL FWindowsUIAControlProvider::get_CanMove(BOOL *pRetVal)
 {
-	if (IsValid())
-	{
-		*pRetVal = WindowsUIAPropertyGetters::GetPropertyValue(Widget, UIA_TransformCanMovePropertyId).GetValue<bool>();
-		return S_OK;
-	}
-	else
-	{
-		return UIA_E_ELEMENTNOTAVAILABLE;
-	}
+	HRESULT ReturnValue = UIA_E_ELEMENTNOTAVAILABLE;
+	UIAManager->RunInGameThreadBlocking(
+		[this, &ReturnValue, &pRetVal]() {
+		if (IsValid())
+		{
+			*pRetVal = WindowsUIAPropertyGetters::GetPropertyValue(Widget, UIA_TransformCanMovePropertyId).GetValue<bool>();
+			ReturnValue = S_OK;
+		}
+	});
+	return ReturnValue;
 }
 
 HRESULT STDCALL FWindowsUIAControlProvider::get_CanResize(BOOL *pRetVal)
 {
-	if (IsValid())
-	{
-		*pRetVal = WindowsUIAPropertyGetters::GetPropertyValue(Widget, UIA_TransformCanResizePropertyId).GetValue<bool>();
-		return S_OK;
-	}
-	else
-	{
-		return UIA_E_ELEMENTNOTAVAILABLE;
-	}
+	HRESULT ReturnValue = UIA_E_ELEMENTNOTAVAILABLE;
+	UIAManager->RunInGameThreadBlocking(
+		[this, &ReturnValue, &pRetVal]() {
+		if (IsValid())
+		{
+			*pRetVal = WindowsUIAPropertyGetters::GetPropertyValue(Widget, UIA_TransformCanResizePropertyId).GetValue<bool>();
+			ReturnValue = S_OK;
+		}
+	});
+	return ReturnValue;
 }
 
 HRESULT STDCALL FWindowsUIAControlProvider::get_CanRotate(BOOL *pRetVal)
 {
-	if (IsValid())
-	{
-		*pRetVal = WindowsUIAPropertyGetters::GetPropertyValue(Widget, UIA_TransformCanRotatePropertyId).GetValue<bool>();
-		return S_OK;
-	}
-	else
-	{
-		return UIA_E_ELEMENTNOTAVAILABLE;
-	}
+	HRESULT ReturnValue = UIA_E_ELEMENTNOTAVAILABLE;
+	UIAManager->RunInGameThreadBlocking(
+		[this, &ReturnValue, &pRetVal]() {
+		if (IsValid())
+		{
+			*pRetVal = WindowsUIAPropertyGetters::GetPropertyValue(Widget, UIA_TransformCanRotatePropertyId).GetValue<bool>();
+			ReturnValue = S_OK;
+		}
+	});
+	return ReturnValue;
 }
 
 HRESULT STDCALL FWindowsUIAControlProvider::Move(double x, double y)
@@ -552,28 +588,31 @@ HRESULT STDCALL FWindowsUIAControlProvider::Rotate(double degrees)
 
 HRESULT STDCALL FWindowsUIAControlProvider::SetValue(LPCWSTR val)
 {
-	if (IsValid())
-	{
-		Widget->AsProperty()->SetValue(FString(val));
-		return S_OK;
-	}
-	else
-	{
-		return UIA_E_ELEMENTNOTAVAILABLE;
-	}
+	HRESULT ReturnValue = UIA_E_ELEMENTNOTAVAILABLE;
+	UIAManager->RunInGameThreadBlocking(
+		[this, &ReturnValue, &val]() {
+		if (IsValid())
+		{
+			Widget->AsProperty()->SetValue(FString(val));
+			ReturnValue = S_OK;
+		}
+	});
+	return ReturnValue;
 }
 
 HRESULT STDCALL FWindowsUIAControlProvider::get_Value(BSTR* pRetVal)
 {
-	if (IsValid())
-	{
-		*pRetVal = SysAllocString(*WindowsUIAPropertyGetters::GetPropertyValue(Widget, UIA_ValueValuePropertyId).GetValue<FString>());
-		return S_OK;
-	}
-	else
-	{
-		return UIA_E_ELEMENTNOTAVAILABLE;
-	}
+	HRESULT ReturnValue = UIA_E_ELEMENTNOTAVAILABLE;
+	UIAManager->RunInGameThreadBlocking(
+		[this, &ReturnValue, &pRetVal]() {
+		if (IsValid())
+		{
+			*pRetVal = SysAllocString(*WindowsUIAPropertyGetters::GetPropertyValue(Widget, UIA_ValueValuePropertyId).GetValue<FString>());
+			ReturnValue = S_OK;
+		}
+
+	});
+	return ReturnValue;
 }
 
 //HRESULT STDCALL FWindowsUIAControlProvider::get_IsReadOnly(BOOL* pRetVal)
@@ -584,112 +623,130 @@ HRESULT STDCALL FWindowsUIAControlProvider::get_Value(BSTR* pRetVal)
 
 HRESULT STDCALL FWindowsUIAControlProvider::Close()
 {
-	if (IsValid())
-	{
-		Widget->AsWindow()->Close();
-		return S_OK;
-	}
-	else
-	{
-		return UIA_E_ELEMENTNOTAVAILABLE;
-	}
+	HRESULT ReturnValue = UIA_E_ELEMENTNOTAVAILABLE;
+	UIAManager->RunInGameThreadBlocking(
+		[this, &ReturnValue]() {
+		if (IsValid())
+		{
+			Widget->AsWindow()->Close();
+			ReturnValue = S_OK;
+		}
+
+	});
+	return ReturnValue;
 }
 
 HRESULT STDCALL FWindowsUIAControlProvider::get_CanMaximize(BOOL* pRetVal)
 {
-	if (IsValid())
-	{
-		*pRetVal = WindowsUIAPropertyGetters::GetPropertyValue(Widget, UIA_WindowCanMaximizePropertyId).GetValue<bool>();
-		return S_OK;
-	}
-	else
-	{
-		return UIA_E_ELEMENTNOTAVAILABLE;
-	}
+	HRESULT ReturnValue = UIA_E_ELEMENTNOTAVAILABLE;
+	UIAManager->RunInGameThreadBlocking(
+		[this, &ReturnValue, &pRetVal]() {
+		if (IsValid())
+		{
+			*pRetVal = WindowsUIAPropertyGetters::GetPropertyValue(Widget, UIA_WindowCanMaximizePropertyId).GetValue<bool>();
+			ReturnValue = S_OK;
+		}
+
+	});
+	return ReturnValue;
 }
 
 HRESULT STDCALL FWindowsUIAControlProvider::get_CanMinimize(BOOL* pRetVal)
 {
-	if (IsValid())
-	{
-		*pRetVal = WindowsUIAPropertyGetters::GetPropertyValue(Widget, UIA_WindowCanMinimizePropertyId).GetValue<bool>();
-		return S_OK;
-	}
-	else
-	{
-		return UIA_E_ELEMENTNOTAVAILABLE;
-	}
+	HRESULT ReturnValue = UIA_E_ELEMENTNOTAVAILABLE;
+	UIAManager->RunInGameThreadBlocking(
+		[this, &ReturnValue, &pRetVal]() {
+		if (IsValid())
+		{
+			*pRetVal = WindowsUIAPropertyGetters::GetPropertyValue(Widget, UIA_WindowCanMinimizePropertyId).GetValue<bool>();
+			ReturnValue = S_OK;
+		}
+	});
+	return ReturnValue;
 }
 
 HRESULT STDCALL FWindowsUIAControlProvider::get_IsModal(BOOL* pRetVal)
 {
-	if (IsValid())
-	{
-		*pRetVal = WindowsUIAPropertyGetters::GetPropertyValue(Widget, UIA_WindowIsModalPropertyId).GetValue<bool>();
-		return S_OK;
-	}
-	else
-	{
-		return UIA_E_ELEMENTNOTAVAILABLE;
-	}
+	HRESULT ReturnValue = UIA_E_ELEMENTNOTAVAILABLE;
+	UIAManager->RunInGameThreadBlocking(
+		[this, &ReturnValue, &pRetVal]() {
+		if (IsValid())
+		{
+			*pRetVal = WindowsUIAPropertyGetters::GetPropertyValue(Widget, UIA_WindowIsModalPropertyId).GetValue<bool>();
+			ReturnValue = S_OK;
+		}
+	});
+	return ReturnValue;
 }
 
 HRESULT STDCALL FWindowsUIAControlProvider::get_IsTopmost(BOOL* pRetVal)
 {
-	// todo: not 100% sure what this is looking for. top window in hierarchy of child windows? on top of all other windows in Windows OS?
-	*pRetVal = WindowsUIAPropertyGetters::GetPropertyValue(Widget, UIA_WindowIsTopmostPropertyId).GetValue<bool>();
-	return E_NOTIMPL;
+	HRESULT ReturnValue = E_NOTIMPL;
+	UIAManager->RunInGameThreadBlocking(
+		[this, &ReturnValue, &pRetVal]() {
+		// todo: not 100% sure what this is looking for. top window in hierarchy of child windows? on top of all other windows in Windows OS?
+		*pRetVal = WindowsUIAPropertyGetters::GetPropertyValue(Widget, UIA_WindowIsTopmostPropertyId).GetValue<bool>();
+	});
+	return ReturnValue;
 }
 
 HRESULT STDCALL FWindowsUIAControlProvider::get_WindowInteractionState(WindowInteractionState* pRetVal)
 {
-	if (IsValid())
-	{
-		// todo: do we have a way to identify if the app is processing data vs idling?
-		*pRetVal = static_cast<WindowInteractionState>(WindowsUIAPropertyGetters::GetPropertyValue(Widget, UIA_WindowWindowInteractionStatePropertyId).GetValue<int32>());
-	}
-	else
-	{
-		*pRetVal = WindowInteractionState_Closing;
-	}
+	HRESULT ReturnValue = UIA_E_ELEMENTNOTAVAILABLE;
+	UIAManager->RunInGameThreadBlocking(
+		[this, &pRetVal]() {
+		if (IsValid())
+		{
+			// todo: do we have a way to identify if the app is processing data vs idling?
+			*pRetVal = static_cast<WindowInteractionState>(WindowsUIAPropertyGetters::GetPropertyValue(Widget, UIA_WindowWindowInteractionStatePropertyId).GetValue<int32>());
+		}
+		else
+		{
+			*pRetVal = WindowInteractionState_Closing;
+		}
+
+	});
 	return S_OK;
 }
 
 HRESULT STDCALL FWindowsUIAControlProvider::get_WindowVisualState(WindowVisualState* pRetVal)
 {
-	if (IsValid())
-	{
-		*pRetVal = static_cast<WindowVisualState>(WindowsUIAPropertyGetters::GetPropertyValue(Widget, UIA_WindowWindowVisualStatePropertyId).GetValue<int32>());
-		return S_OK;
-	}
-	else
-	{
-		return UIA_E_ELEMENTNOTAVAILABLE;
-	}
+	HRESULT ReturnValue = UIA_E_ELEMENTNOTAVAILABLE;
+	UIAManager->RunInGameThreadBlocking(
+		[this, &ReturnValue, &pRetVal]() {
+		if (IsValid())
+		{
+			*pRetVal = static_cast<WindowVisualState>(WindowsUIAPropertyGetters::GetPropertyValue(Widget, UIA_WindowWindowVisualStatePropertyId).GetValue<int32>());
+			ReturnValue = S_OK;
+		}
+
+	});
+	return ReturnValue;
 }
 
 HRESULT STDCALL FWindowsUIAControlProvider::SetVisualState(WindowVisualState state)
 {
-	if (IsValid())
-	{
-		switch (state)
+	HRESULT ReturnValue = UIA_E_ELEMENTNOTAVAILABLE;
+	UIAManager->RunInGameThreadBlocking(
+		[this, &ReturnValue, &state]() {
+		if (IsValid())
 		{
-		case WindowVisualState_Normal:
-			Widget->AsWindow()->SetDisplayState(IAccessibleWindow::EWindowDisplayState::Normal);
-			break;
-		case WindowVisualState_Minimized:
-			Widget->AsWindow()->SetDisplayState(IAccessibleWindow::EWindowDisplayState::Minimize);
-			break;
-		case WindowVisualState_Maximized:
-			Widget->AsWindow()->SetDisplayState(IAccessibleWindow::EWindowDisplayState::Maximize);
-			break;
+			switch (state)
+			{
+			case WindowVisualState_Normal:
+				Widget->AsWindow()->SetDisplayState(IAccessibleWindow::EWindowDisplayState::Normal);
+				break;
+			case WindowVisualState_Minimized:
+				Widget->AsWindow()->SetDisplayState(IAccessibleWindow::EWindowDisplayState::Minimize);
+				break;
+			case WindowVisualState_Maximized:
+				Widget->AsWindow()->SetDisplayState(IAccessibleWindow::EWindowDisplayState::Maximize);
+				break;
+			}
+			ReturnValue = S_OK;
 		}
-		return S_OK;
-	}
-	else
-	{
-		return UIA_E_ELEMENTNOTAVAILABLE;
-	}
+	});
+	return ReturnValue;
 }
 
 HRESULT STDCALL FWindowsUIAControlProvider::WaitForInputIdle(int milliseconds, BOOL* pRetVal)
diff --git a/Engine/Source/Runtime/ApplicationCore/Private/Windows/Accessibility/WindowsUIAManager.cpp b/Engine/Source/Runtime/ApplicationCore/Private/Windows/Accessibility/WindowsUIAManager.cpp
--- a/Engine/Source/Runtime/ApplicationCore/Private/Windows/Accessibility/WindowsUIAManager.cpp
+++ b/Engine/Source/Runtime/ApplicationCore/Private/Windows/Accessibility/WindowsUIAManager.cpp
@@ -18,6 +18,7 @@
 #include "Internationalization/Internationalization.h"
 #include "Internationalization/Culture.h"
 
+
 TMap<EAccessibleWidgetType, ULONG> FWindowsUIAManager::WidgetTypeToWindowsTypeMap;
 TMap<EAccessibleWidgetType, FText> FWindowsUIAManager::WidgetTypeToTextMap;
 
@@ -277,6 +278,12 @@ void FWindowsUIAManager::UpdateCachedCurrentLocaleLCID()
 		CachedCurrentLocaleLCID = FInternationalization::Get().GetDefaultLocale()->GetLCID();
 	}
 }
+
+void FWindowsUIAManager::RunInGameThreadBlocking(CONST TFunction<void()>& Function) const
+{
+		WindowsApplication.GetAccessibleMessageHandler()->RunInThread(Function, true, ENamedThreads::GameThread);
+}
+
 #if !UE_BUILD_SHIPPING
 void FWindowsUIAManager::DumpAccessibilityStats() const
 {
diff --git a/Engine/Source/Runtime/ApplicationCore/Private/Windows/Accessibility/WindowsUIAWidgetProvider.cpp b/Engine/Source/Runtime/ApplicationCore/Private/Windows/Accessibility/WindowsUIAWidgetProvider.cpp
--- a/Engine/Source/Runtime/ApplicationCore/Private/Windows/Accessibility/WindowsUIAWidgetProvider.cpp
+++ b/Engine/Source/Runtime/ApplicationCore/Private/Windows/Accessibility/WindowsUIAWidgetProvider.cpp
@@ -3,8 +3,7 @@
 #if WITH_ACCESSIBILITY
 
 #include "Windows/Accessibility/WindowsUIAWidgetProvider.h"
-
-#include "GenericPlatform/GenericAccessibleInterfaces.h"
+#include "GenericPlatform/Accessibility/GenericAccessibleInterfaces.h"
 #include "Windows/Accessibility/WindowsUIAControlProvider.h"
 #include "Windows/Accessibility/WindowsUIAManager.h"
 #include "Windows/Accessibility/WindowsUIAPropertyGetters.h"
diff --git a/Engine/Source/Runtime/ApplicationCore/Public/GenericPlatform/GenericApplication.h b/Engine/Source/Runtime/ApplicationCore/Public/GenericPlatform/GenericApplication.h
--- a/Engine/Source/Runtime/ApplicationCore/Public/GenericPlatform/GenericApplication.h
+++ b/Engine/Source/Runtime/ApplicationCore/Public/GenericPlatform/GenericApplication.h
@@ -10,7 +10,7 @@
 #include "Templates/SharedPointer.h"
 #include "Delegates/Delegate.h"
 #if WITH_ACCESSIBILITY
-#include "GenericPlatform/GenericAccessibleInterfaces.h"
+#include "GenericPlatform/Accessibility/GenericAccessibleInterfaces.h"
 #endif
 #include "GenericPlatform/GenericApplicationMessageHandler.h"
 #include "GenericPlatform/GenericWindowDefinition.h"
@@ -442,7 +442,6 @@ public:
 	{
 
 	}
-
 	virtual ~GenericApplication() {}
 
 	virtual void SetMessageHandler( const TSharedRef< FGenericApplicationMessageHandler >& InMessageHandler ) { MessageHandler = InMessageHandler; }
@@ -562,7 +561,9 @@ protected:
 
 #if WITH_ACCESSIBILITY
 	TSharedRef<FGenericAccessibleMessageHandler> AccessibleMessageHandler;
+	
 #endif
+
 	
 	/** Trigger the OnDisplayMetricsChanged event with the argument 'InMetrics' */
 	void BroadcastDisplayMetricsChanged( const FDisplayMetrics& InMetrics ){ OnDisplayMetricsChangedEvent.Broadcast( InMetrics ); }
diff --git a/Engine/Source/Runtime/ApplicationCore/Public/Windows/Accessibility/WindowsUIAManager.h b/Engine/Source/Runtime/ApplicationCore/Public/Windows/Accessibility/WindowsUIAManager.h
--- a/Engine/Source/Runtime/ApplicationCore/Public/Windows/Accessibility/WindowsUIAManager.h
+++ b/Engine/Source/Runtime/ApplicationCore/Public/Windows/Accessibility/WindowsUIAManager.h
@@ -5,7 +5,7 @@
 #if WITH_ACCESSIBILITY
 
 #include "Templates/SharedPointer.h"
-#include "GenericPlatform/GenericAccessibleInterfaces.h"
+#include "GenericPlatform/Accessibility/GenericAccessibleInterfaces.h"
 
 class FWindowsApplication;
 class FWindowsWindow;
@@ -69,6 +69,14 @@ public:
 
 	uint32 GetCachedCurrentLocaleLCID() const { return CachedCurrentLocaleLCID; }
 
+	/**
+	* Runs the passed in function in the game thread. 
+	* Blocks until the function completes execution in the game thread.
+
+	* @param Function The function to run in the game thread 
+	*/
+	void RunInGameThreadBlocking(const TFunction<void()>& Function) const;
+
 	static TMap<EAccessibleWidgetType, ULONG> WidgetTypeToWindowsTypeMap;
 	static TMap<EAccessibleWidgetType, FText> WidgetTypeToTextMap;
 
diff --git a/Engine/Source/Runtime/Core/Private/Misc/StringFormatter.cpp b/Engine/Source/Runtime/Core/Private/Misc/StringFormatter.cpp
--- a/Engine/Source/Runtime/Core/Private/Misc/StringFormatter.cpp
+++ b/Engine/Source/Runtime/Core/Private/Misc/StringFormatter.cpp
@@ -439,13 +439,13 @@ FStringFormatter& GetDefaultFormatter()
 	return DefaultFormatter;
 }
 
-FString FString::Format(const TCHAR* InFormatString, const TMap<FString, FStringFormatArg>& InNamedArguments)
+FString FString::Format(const TCHAR* InFormatString, const FStringFormatNamedArguments& InNamedArguments)
 {
 	FStringFormatter& DefaultFormatter = GetDefaultFormatter();
 	return DefaultFormatter.Format(InFormatString, InNamedArguments);
 }
 
-FString FString::Format(const TCHAR* InFormatString, const TArray<FStringFormatArg>& InOrderedArguments)
+FString FString::Format(const TCHAR* InFormatString, const FStringFormatOrderedArguments& InOrderedArguments)
 {
 	FStringFormatter& DefaultFormatter = GetDefaultFormatter();
 	return DefaultFormatter.Format(InFormatString, InOrderedArguments);
diff --git a/Engine/Source/Runtime/Core/Public/Containers/UnrealString.h b/Engine/Source/Runtime/Core/Public/Containers/UnrealString.h
--- a/Engine/Source/Runtime/Core/Public/Containers/UnrealString.h
+++ b/Engine/Source/Runtime/Core/Public/Containers/UnrealString.h
@@ -30,6 +30,9 @@ class FStringView;
 struct FStringFormatArg;
 template<typename KeyType,typename ValueType,typename SetAllocator ,typename KeyFuncs > class TMap;
 
+typedef TMap<FString, FStringFormatArg> FStringFormatNamedArguments;
+typedef TArray<FStringFormatArg> FStringFormatOrderedArguments;
+
 template <>
 struct TIsContiguousContainer<FString>
 {
@@ -1509,7 +1512,7 @@ public:
 	 * @param InNamedArguments		A map of named arguments that match the tokens specified in InExpression
 	 * @return A string containing the formatted text
 	 */
-	static FString Format(const TCHAR* InFormatString, const TMap<FString, FStringFormatArg>& InNamedArguments);
+	static FString Format(const TCHAR* InFormatString, const FStringFormatNamedArguments& InNamedArguments);
 
 	/**
 	 * Format the specified string using the specified arguments. Replaces instances of {0} with indices from the given array matching the index specified in the token
@@ -1517,7 +1520,7 @@ public:
 	 * @param InOrderedArguments	An array of ordered arguments that match the tokens specified in InExpression
 	 * @return A string containing the formatted text
 	 */
-	static FString Format(const TCHAR* InFormatString, const TArray<FStringFormatArg>& InOrderedArguments);
+	static FString Format(const TCHAR* InFormatString, const FStringFormatOrderedArguments& InOrderedArguments);
 
 	/** Returns a string containing only the Ch character */
 	static FString Chr( TCHAR Ch );
diff --git a/Engine/Source/Runtime/Core/Public/UObject/EditorObjectVersion.h b/Engine/Source/Runtime/Core/Public/UObject/EditorObjectVersion.h
--- a/Engine/Source/Runtime/Core/Public/UObject/EditorObjectVersion.h
+++ b/Engine/Source/Runtime/Core/Public/UObject/EditorObjectVersion.h
@@ -85,6 +85,8 @@ struct CORE_API FEditorObjectVersion
 		ComputeWeightedNormals,
 		// SkeletalMesh now can be rebuild in editor, no more need to re-import
 		SkeletalMeshBuildRefactor,
+		// Move all SkeletalMesh source data into a private uasset in the same package has the skeletalmesh
+		SkeletalMeshMoveEditorSourceDataToPrivateAsset,
 		// -----<new versions can be added above this line>-------------------------------------------------
 		VersionPlusOne,
 		LatestVersion = VersionPlusOne - 1
diff --git a/Engine/Source/Runtime/Engine/Classes/Engine/GameInstance.h b/Engine/Source/Runtime/Engine/Classes/Engine/GameInstance.h
--- a/Engine/Source/Runtime/Engine/Classes/Engine/GameInstance.h
+++ b/Engine/Source/Runtime/Engine/Classes/Engine/GameInstance.h
@@ -14,6 +14,10 @@
 #include "Subsystems/GameInstanceSubsystem.h"
 #include "Subsystems/SubsystemCollection.h"
 
+#if WITH_EDITOR
+#include "Settings/LevelEditorPlaySettings.h"
+#endif 
+
 #include "GameInstance.generated.h"
 
 class AGameModeBase;
@@ -24,6 +28,7 @@ class FUniqueNetId;
 class ULocalPlayer;
 class UOnlineSession;
 struct FLatentActionManager;
+class ULevelEditorPlaySettings;
 
 // 
 // 	EWelcomeScreen, 	//initial screen.  Used for platforms where we may not have a signed in user yet.
@@ -90,6 +95,16 @@ private:
 //@TODO: Some of these are really mutually exclusive and should be refactored (put into a struct to make this easier in the future)
 struct FGameInstancePIEParameters
 {
+	FGameInstancePIEParameters()
+		: bSimulateInEditor(false)
+		, bAnyBlueprintErrors(false)
+		, bStartInSpectatorMode(false)
+		, bRunAsDedicated(false)
+		, WorldFeatureLevel(ERHIFeatureLevel::Num)
+		, EditorPlaySettings(nullptr)
+		, NetMode(EPlayNetMode::PIE_Standalone)
+	{}
+
 	// Are we doing SIE instead of PIE?
 	bool bSimulateInEditor;
 
@@ -103,7 +118,16 @@ struct FGameInstancePIEParameters
 	bool bRunAsDedicated;
 
 	// The feature level that PIE world should use
-	ERHIFeatureLevel::Type WorldFeatureLevel = ERHIFeatureLevel::Num;
+	ERHIFeatureLevel::Type WorldFeatureLevel;
+
+	// Kept alive externally.
+	ULevelEditorPlaySettings* EditorPlaySettings;
+
+	// Which net mode should this PIE instance start in? Affects which maps are loaded.
+	EPlayNetMode NetMode;
+
+	// The map we should force the game to load instead of the one currently running in the editor. Blank for no override
+	FString OverrideMapURL;
 };
 
 #endif
