IOS thumbstick button and special left support.

diff --git a/Engine/Source/Runtime/ApplicationCore/Private/IOS/IOSInputInterface.cpp b/Engine/Source/Runtime/ApplicationCore/Private/IOS/IOSInputInterface.cpp
--- a/Engine/Source/Runtime/ApplicationCore/Private/IOS/IOSInputInterface.cpp
+++ b/Engine/Source/Runtime/ApplicationCore/Private/IOS/IOSInputInterface.cpp
@@ -12,6 +12,10 @@
 
 DECLARE_LOG_CATEGORY_EXTERN(LogIOSInput, Log, All);
 
+#ifndef __IPHONE_OS_VERSION_MAX_ALLOWED
+#define __IPHONE_OS_VERSION_MAX_ALLOWED 0
+#endif
+
 #ifndef __MAC_OS_VERSION_MAX_ALLOWED
 #define __MAC_OS_VERSION_MAX_ALLOWED 0
 #endif
@@ -20,27 +24,18 @@ DECLARE_LOG_CATEGORY_EXTERN(LogIOSInput, Log, All);
 #define __APPLETV_OS_VERSION_MAX_ALLOWED 0
 #endif
 
-#if PLATFORM_TVOS
-#define TVOS_THUMBSTICKS_UNDECLARED (__APPLETV_OS_VERSION_MAX_ALLOWED < 121000)
-#else
-#define TVOS_THUMBSTICKS_UNDECLARED false
-#endif
-
-#if (__IPHONE_OS_VERSION_MAX_ALLOWED < 121000 || TVOS_THUMBSTICKS_UNDECLARED || __MAC_OS_VERSION_MAX_ALLOWED < 1401000)
-@interface GCExtendedGamepadSnapshot()
-@property (nonatomic, readwrite, nullable) GCControllerButtonInput *leftThumbstickButton;
-@property (nonatomic, readwrite, nullable) GCControllerButtonInput *rightThumbstickButton;
-@end
-
 @interface GCExtendedGamepad()
+#if (__IPHONE_OS_VERSION_MAX_ALLOWED < 121000 || __APPLETV_OS_VERSION_MAX_ALLOWED < 121000 || __MAC_OS_VERSION_MAX_ALLOWED < 1401000)
 @property (nonatomic, readwrite, nullable) GCControllerButtonInput *leftThumbstickButton;
 @property (nonatomic, readwrite, nullable) GCControllerButtonInput *rightThumbstickButton;
-@end
+#endif
 
-@interface GCControllerButtonInput (capture)
-- (bool)_setValue:(float)arg1;
-@end
+#if (__IPHONE_OS_VERSION_MAX_ALLOWED < 130000 || __APPLETV_OS_VERSION_MAX_ALLOWED < 130000 || __MAC_OS_VERSION_MAX_ALLOWED < 1500000)
+@property (nonatomic, readwrite, nullable) GCControllerButtonInput *buttonMenu;
+@property (nonatomic, readwrite, nullable) GCControllerButtonInput *buttonOptions;
 #endif
+@end
+
 
 #if (__IPHONE_OS_VERSION_MAX_ALLOWED < 130000)
 // only redefine these values pre-iOS13 SDK
@@ -112,6 +107,11 @@ FIOSInputInterface::FIOSInputInterface( const TSharedRef< FGenericApplicationMes
 	   });
 	
 	FMemory::Memzero(Controllers, sizeof(Controllers));
+    
+    for (GCController* Cont in [GCController controllers])
+    {
+        HandleConnection(Cont);
+    }
 	
 	FEmbeddedDelegates::GetNativeToEmbeddedParamsDelegateForSubsystem(TEXT("iosinput")).AddLambda([this](const FEmbeddedCallParamsHelper& Message)
 	{
@@ -181,6 +181,8 @@ void FIOSInputInterface::HandleConnection(GCController* Controller)
 			( bIsTreatedAsGamepad && Controllers[ControllerIndex].bIsGamepadConnected == false))
 		{
 			Controller.playerIndex = (GCControllerPlayerIndex)ControllerIndex;
+            
+            Controllers[ControllerIndex].Controller = Controller;
 #if PLATFORM_TVOS
 			if (Controller.microGamepad != nil)
 			{
@@ -377,9 +379,11 @@ void FIOSInputInterface::SendControllerEvents()
 		MessageHandler->OnMotionDetected(Attitude, RotationRate, Gravity, Acceleration, 0);
 	}
 #endif
-		
-	for (GCController* Cont in [GCController controllers])
+    for(int32 i = 0; i < UE_ARRAY_COUNT(Controllers); ++i)
  	{
+        if(!Controllers[i].bIsGamepadConnected) continue;
+        
+        GCController* Cont = Controllers[i].Controller;
 		GCExtendedGamepad* ExtendedGamepad = Cont.extendedGamepad;
 #if PLATFORM_TVOS
 		GCMicroGamepad* MicroGamepad = Cont.microGamepad;
@@ -400,7 +404,8 @@ void FIOSInputInterface::SendControllerEvents()
 
 		FUserController& Controller = Controllers[Cont.playerIndex];
 		
-		if (Controller.bPauseWasPressed)
+        // If buttonMenu is defined, we will handle it like a regular button.
+		if (Controller.bPauseWasPressed && ExtendedGamepad.buttonMenu == nil)
 		{
 			MessageHandler->OnControllerButtonPressed(FGamepadKeyNames::SpecialRight, Cont.playerIndex, false);
 			MessageHandler->OnControllerButtonReleased(FGamepadKeyNames::SpecialRight, Cont.playerIndex, false);
@@ -489,11 +494,26 @@ if ((Previous##Gamepad != nil && Gamepad.GCAxis.value != Previous##Gamepad.GCAxi
 			HANDLE_ANALOG_VIRTUAL_BUTTONS(ExtendedGamepad, rightThumbstick.xAxis, FGamepadKeyNames::RightStickLeft, FGamepadKeyNames::RightStickRight);
 			HANDLE_ANALOG_VIRTUAL_BUTTONS(ExtendedGamepad, rightThumbstick.yAxis, FGamepadKeyNames::RightStickDown, FGamepadKeyNames::RightStickUp);
             
+            if(ExtendedGamepad.buttonMenu != nil)
+            {
+                HANDLE_BUTTON(ExtendedGamepad, buttonMenu,          FGamepadKeyNames::SpecialRight);
+            }
+            
+            if(ExtendedGamepad.buttonOptions != nil)
+            {
+                HANDLE_BUTTON(ExtendedGamepad, buttonOptions,       FGamepadKeyNames::SpecialLeft);
+            }
+            
             if(ExtendedGamepad.leftThumbstickButton != nil)
             {
-                
-                HANDLE_BUTTON(ExtendedGamepad, leftThumbstickButton,        FGamepadKeyNames::LeftThumb);
-                HANDLE_BUTTON(ExtendedGamepad, rightThumbstickButton,        FGamepadKeyNames::RightThumb);
+                HANDLE_BUTTON_INTERNAL(ExtendedGamepad, Controller.bLeftThumbstickWasPressed, ExtendedGamepad.leftThumbstickButton.pressed, FGamepadKeyNames::LeftThumb);
+                Controller.bLeftThumbstickWasPressed = ExtendedGamepad.leftThumbstickButton.pressed;
+            }
+            
+            if(ExtendedGamepad.rightThumbstickButton != nil)
+            {
+                HANDLE_BUTTON_INTERNAL(ExtendedGamepad, Controller.bRightThumbstickWasPressed, ExtendedGamepad.rightThumbstickButton.pressed, FGamepadKeyNames::RightThumb);
+                Controller.bRightThumbstickWasPressed = ExtendedGamepad.rightThumbstickButton.pressed;
             }
 
             [Controller.PreviousExtendedGamepad release];
@@ -504,11 +524,6 @@ if ((Previous##Gamepad != nil && Gamepad.GCAxis.value != Previous##Gamepad.GCAxi
             else
             {
                 Controller.PreviousExtendedGamepad = [ExtendedGamepad saveSnapshot];
-                if(ExtendedGamepad.leftThumbstickButton != nil)
-                {
-                    [Controller.PreviousExtendedGamepad.leftThumbstickButton _setValue:ExtendedGamepad.leftThumbstickButton.value];
-                    [Controller.PreviousExtendedGamepad.rightThumbstickButton _setValue:ExtendedGamepad.rightThumbstickButton.value];
-                }
             }
             [Controller.PreviousExtendedGamepad retain];
 		}
diff --git a/Engine/Source/Runtime/ApplicationCore/Public/IOS/IOSInputInterface.h b/Engine/Source/Runtime/ApplicationCore/Public/IOS/IOSInputInterface.h
--- a/Engine/Source/Runtime/ApplicationCore/Public/IOS/IOSInputInterface.h
+++ b/Engine/Source/Runtime/ApplicationCore/Public/IOS/IOSInputInterface.h
@@ -117,10 +117,10 @@ private:
 	/** Game controller objects (per user)*/
 	struct FUserController
 	{
-		GCGamepadSnapshot* PreviousGamepad;
-		GCExtendedGamepadSnapshot* PreviousExtendedGamepad;
+        GCController* Controller;
+		GCExtendedGamepad* PreviousExtendedGamepad;
 #if PLATFORM_TVOS
-		GCMicroGamepadSnapshot* PreviousMicroGamepad;
+		GCMicroGamepad* PreviousMicroGamepad;
 #endif
 		FQuat ReferenceAttitude;
 		bool bNeedsReferenceAttitude;
@@ -128,6 +128,9 @@ private:
 		bool bIsGamepadConnected;
 		bool bIsRemoteConnected;
 		bool bPauseWasPressed;
+        // Tracked outside the gamepads to allow us to support thumbsticks on more devices.
+        bool bRightThumbstickWasPressed;
+        bool bLeftThumbstickWasPressed;
 	};
 	// there is a hardcoded limit of 4 controllers in the API
 	FUserController Controllers[4];
