Stratus XL triggers are not appropriately triggering trigger events.

diff --git a/Engine/Source/Runtime/ApplicationCore/Private/Android/AndroidInputInterface.cpp b/Engine/Source/Runtime/ApplicationCore/Private/Android/AndroidInputInterface.cpp
--- a/Engine/Source/Runtime/ApplicationCore/Private/Android/AndroidInputInterface.cpp
+++ b/Engine/Source/Runtime/ApplicationCore/Private/Android/AndroidInputInterface.cpp
@@ -917,6 +917,7 @@ void FAndroidInputInterface::SendControllerEvents()
 						else if (CurrentDevice.DeviceInfo.Name.StartsWith(TEXT("SteelSeries Stratus XL")))
 						{
 							CurrentDevice.bSupportsHat = true;
+							CurrentDevice.bTriggersUseThresholdForClick = true;
 
 							// For some reason the left trigger is at 0.5 when at rest so we have to adjust for that.
 							CurrentDevice.LTAnalogRangeMinimum = 0.5f;
