Fixing crash when using controllers with IOS devices on versions less than 13.

diff --git a/Engine/Source/Runtime/ApplicationCore/Private/IOS/IOSInputInterface.cpp b/Engine/Source/Runtime/ApplicationCore/Private/IOS/IOSInputInterface.cpp
--- a/Engine/Source/Runtime/ApplicationCore/Private/IOS/IOSInputInterface.cpp
+++ b/Engine/Source/Runtime/ApplicationCore/Private/IOS/IOSInputInterface.cpp
@@ -404,8 +404,10 @@ void FIOSInputInterface::SendControllerEvents()
 
 		FUserController& Controller = Controllers[Cont.playerIndex];
 		
+        static bool bSystemSupportsMenuButtons = [GCExtendedGamepad instancesRespondToSelector:@selector(buttonOptions)];
+        
         // If buttonMenu is defined, we will handle it like a regular button.
-		if (Controller.bPauseWasPressed && ExtendedGamepad.buttonMenu == nil)
+		if (Controller.bPauseWasPressed && !bSystemSupportsMenuButtons)
 		{
 			MessageHandler->OnControllerButtonPressed(FGamepadKeyNames::SpecialRight, Cont.playerIndex, false);
 			MessageHandler->OnControllerButtonReleased(FGamepadKeyNames::SpecialRight, Cont.playerIndex, false);
@@ -494,30 +496,26 @@ if ((Previous##Gamepad != nil && Gamepad.GCAxis.value != Previous##Gamepad.GCAxi
 			HANDLE_ANALOG_VIRTUAL_BUTTONS(ExtendedGamepad, rightThumbstick.xAxis, FGamepadKeyNames::RightStickLeft, FGamepadKeyNames::RightStickRight);
 			HANDLE_ANALOG_VIRTUAL_BUTTONS(ExtendedGamepad, rightThumbstick.yAxis, FGamepadKeyNames::RightStickDown, FGamepadKeyNames::RightStickUp);
             
-            if(ExtendedGamepad.buttonMenu != nil)
+            if(bSystemSupportsMenuButtons)
             {
                 HANDLE_BUTTON(ExtendedGamepad, buttonMenu,          FGamepadKeyNames::SpecialRight);
-            }
-            
-            if(ExtendedGamepad.buttonOptions != nil)
-            {
                 HANDLE_BUTTON(ExtendedGamepad, buttonOptions,       FGamepadKeyNames::SpecialLeft);
             }
             
-            if(ExtendedGamepad.leftThumbstickButton != nil)
+            static bool bSystemSupportsThumbsticks = [GCExtendedGamepad instancesRespondToSelector:@selector(leftThumbstickButton)];
+            
+            if(bSystemSupportsThumbsticks)
             {
                 HANDLE_BUTTON_INTERNAL(ExtendedGamepad, Controller.bLeftThumbstickWasPressed, ExtendedGamepad.leftThumbstickButton.pressed, FGamepadKeyNames::LeftThumb);
                 Controller.bLeftThumbstickWasPressed = ExtendedGamepad.leftThumbstickButton.pressed;
-            }
-            
-            if(ExtendedGamepad.rightThumbstickButton != nil)
-            {
+                
                 HANDLE_BUTTON_INTERNAL(ExtendedGamepad, Controller.bRightThumbstickWasPressed, ExtendedGamepad.rightThumbstickButton.pressed, FGamepadKeyNames::RightThumb);
                 Controller.bRightThumbstickWasPressed = ExtendedGamepad.rightThumbstickButton.pressed;
             }
 
             [Controller.PreviousExtendedGamepad release];
-            if ([Cont respondsToSelector:@selector(capture)])
+            static bool bSupportsGamepadCapture = [Cont respondsToSelector:@selector(capture)];
+            if (bSupportsGamepadCapture)
             {
                 Controller.PreviousExtendedGamepad = (GCExtendedGamepadSnapshot*)[ExtendedGamepad.controller capture].extendedGamepad;
             }
