Adding option to disable phone vibration when controllers are attached to mobile devices

diff --git a/Engine/Source/Runtime/Android/AndroidRuntimeSettings/Classes/AndroidRuntimeSettings.h b/Engine/Source/Runtime/Android/AndroidRuntimeSettings/Classes/AndroidRuntimeSettings.h
--- a/Engine/Source/Runtime/Android/AndroidRuntimeSettings/Classes/AndroidRuntimeSettings.h
+++ b/Engine/Source/Runtime/Android/AndroidRuntimeSettings/Classes/AndroidRuntimeSettings.h
@@ -450,6 +450,10 @@ public:
 	UPROPERTY(GlobalConfig, EditAnywhere, Category = Input, meta = (DisplayName = "Block Android system keys being sent from controllers"))
 	bool bBlockAndroidKeysOnControllers;
 
+	// Block force feedback on the device when controllers are attached.
+	UPROPERTY(GlobalConfig, EditAnywhere, Category = Input, meta = (DisplayName = "Block force feedback on the device when controllers are attached"))
+	bool bControllersBlockDeviceFeedback;
+
 	/** Android encoding options. */
 	UPROPERTY(GlobalConfig, EditAnywhere, Category = Audio, meta = (DisplayName = "Encoding Format"))
 	TEnumAsByte<EAndroidAudio::Type> AndroidAudio;
diff --git a/Engine/Source/Runtime/ApplicationCore/Private/Android/AndroidInputInterface.cpp b/Engine/Source/Runtime/ApplicationCore/Private/Android/AndroidInputInterface.cpp
--- a/Engine/Source/Runtime/ApplicationCore/Private/Android/AndroidInputInterface.cpp
+++ b/Engine/Source/Runtime/ApplicationCore/Private/Android/AndroidInputInterface.cpp
@@ -26,6 +26,7 @@ FForceFeedbackValues FAndroidInputInterface::VibeValues;
 
 bool FAndroidInputInterface::bAllowControllers = true;
 bool FAndroidInputInterface::bBlockAndroidKeysOnControllers = false;
+bool FAndroidInputInterface::bControllersBlockDeviceFeedback = false;
 
 FAndroidControllerData FAndroidInputInterface::OldControllerData[MAX_NUM_CONTROLLERS];
 FAndroidControllerData FAndroidInputInterface::NewControllerData[MAX_NUM_CONTROLLERS];
@@ -72,6 +73,7 @@ FAndroidInputInterface::FAndroidInputInterface(const TSharedRef< FGenericApplica
 {
 	GConfig->GetBool(TEXT("/Script/AndroidRuntimeSettings.AndroidRuntimeSettings"), TEXT("bAllowControllers"), bAllowControllers, GEngineIni);
 	GConfig->GetBool(TEXT("/Script/AndroidRuntimeSettings.AndroidRuntimeSettings"), TEXT("bBlockAndroidKeysOnControllers"), bBlockAndroidKeysOnControllers, GEngineIni);
+	GConfig->GetBool(TEXT("/Script/AndroidRuntimeSettings.AndroidRuntimeSettings"), TEXT("bControllersBlockDeviceFeedback"), bControllersBlockDeviceFeedback, GEngineIni);
 
 	ButtonMapping[0] = FGamepadKeyNames::FaceButtonBottom;
 	ButtonMapping[1] = FGamepadKeyNames::FaceButtonRight;
@@ -224,6 +226,8 @@ void FAndroidInputInterface::SetForceFeedbackChannelValue(int32 ControllerId, FF
 		}
 	}
 
+	bDidFeedback |= IsGamepadAttached() && bControllersBlockDeviceFeedback;
+
 	// If controller handled force feedback don't do it on the phone
 	if (bDidFeedback)
 	{
@@ -273,17 +277,21 @@ void FAndroidInputInterface::SetForceFeedbackChannelValues(int32 ControllerId, c
 		}
 	}
 
+	bDidFeedback |= IsGamepadAttached() && bControllersBlockDeviceFeedback;
+
 	// If controller handled force feedback don't do it on the phone
 	if (bDidFeedback)
 	{
-		return;
+		VibeValues.LeftLarge = VibeValues.RightLarge = VibeValues.LeftSmall = VibeValues.RightSmall = 0.0f;
+	}
+	else
+	{
+		VibeValues = Values;
 	}
 
 	// Note: only one motor on Android at the moment, but remember all the settings
 	// update will look at combination of all values to pick state
 
-	VibeValues = Values;
-
 	// Update with the latest values (wait for SendControllerEvents later?)
 	UpdateVibeMotors();
 }
diff --git a/Engine/Source/Runtime/ApplicationCore/Private/Android/AndroidPlatformApplicationMisc.cpp b/Engine/Source/Runtime/ApplicationCore/Private/Android/AndroidPlatformApplicationMisc.cpp
--- a/Engine/Source/Runtime/ApplicationCore/Private/Android/AndroidPlatformApplicationMisc.cpp
+++ b/Engine/Source/Runtime/ApplicationCore/Private/Android/AndroidPlatformApplicationMisc.cpp
@@ -86,6 +86,15 @@ void FAndroidApplicationMisc::SetGamepadsAllowed(bool bAllowed)
 	}
 }
 
+
+void FAndroidApplicationMisc::SetGamepadsBlockDeviceFeedback(bool bBlock)
+{
+	if (FAndroidInputInterface* InputInterface = (FAndroidInputInterface*)FAndroidApplication::Get()->GetInputInterface())
+	{
+		InputInterface->SetGamepadsBlockDeviceFeedback(bBlock);
+	}
+}
+
 void FAndroidApplicationMisc::ResetGamepadAssignments()
 {
 	FAndroidInputInterface::ResetGamepadAssignments();
diff --git a/Engine/Source/Runtime/ApplicationCore/Private/IOS/IOSInputInterface.cpp b/Engine/Source/Runtime/ApplicationCore/Private/IOS/IOSInputInterface.cpp
--- a/Engine/Source/Runtime/ApplicationCore/Private/IOS/IOSInputInterface.cpp
+++ b/Engine/Source/Runtime/ApplicationCore/Private/IOS/IOSInputInterface.cpp
@@ -89,6 +89,7 @@ FIOSInputInterface::FIOSInputInterface( const TSharedRef< FGenericApplicationMes
 	GConfig->GetBool(TEXT("/Script/IOSRuntimeSettings.IOSRuntimeSettings"), TEXT("bUseRemoteAsVirtualJoystick"), bUseRemoteAsVirtualJoystick, GEngineIni);
 	GConfig->GetBool(TEXT("/Script/IOSRuntimeSettings.IOSRuntimeSettings"), TEXT("bUseRemoteAbsoluteDpadValues"), bUseRemoteAbsoluteDpadValues, GEngineIni);
 	GConfig->GetBool(TEXT("/Script/IOSRuntimeSettings.IOSRuntimeSettings"), TEXT("bAllowControllers"), bAllowControllers, GEngineIni);
+	GConfig->GetBool(TEXT("/Script/IOSRuntimeSettings.IOSRuntimeSettings"), TEXT("bControllersBlockDeviceFeedback"), bControllersBlockDeviceFeedback, GEngineIni);
 	
 	[[NSNotificationCenter defaultCenter] addObserverForName:GCControllerDidConnectNotification object:nil queue:[NSOperationQueue currentQueue] usingBlock:^(NSNotification* Notification)
 	 {
@@ -797,6 +798,11 @@ bool FIOSInputInterface::IsGamepadAttached() const
 
 void FIOSInputInterface::SetForceFeedbackChannelValue(int32 ControllerId, FForceFeedbackChannelType ChannelType, float Value)
 {
+	if(IsGamepadAttached() && bControllersBlockDeviceFeedback)
+	{
+		Value = 0.0f;
+	}
+
 	if(HapticFeedbackSupportLevel >= 2)
 	{
 		// if we are at rest, then kick when we are over the Kick cutoff
diff --git a/Engine/Source/Runtime/ApplicationCore/Private/IOS/IOSPlatformApplicationMisc.cpp b/Engine/Source/Runtime/ApplicationCore/Private/IOS/IOSPlatformApplicationMisc.cpp
--- a/Engine/Source/Runtime/ApplicationCore/Private/IOS/IOSPlatformApplicationMisc.cpp
+++ b/Engine/Source/Runtime/ApplicationCore/Private/IOS/IOSPlatformApplicationMisc.cpp
@@ -191,6 +191,14 @@ void FIOSPlatformApplicationMisc::SetGamepadsAllowed(bool bAllowed)
 	}
 }
 
+void FIOSPlatformApplicationMisc::SetGamepadsBlockDeviceFeedback(bool bBlock)
+{
+	if (FIOSInputInterface* InputInterface = (FIOSInputInterface*)CachedApplication->GetInputInterface())
+	{
+		InputInterface->SetGamepadsBlockDeviceFeedback(bBlock);
+	}
+}
+
 void FIOSPlatformApplicationMisc::ResetGamepadAssignments()
 {
 	UE_LOG(LogIOS, Warning, TEXT("Restting gamepad assignments is not allowed in IOS"))
diff --git a/Engine/Source/Runtime/ApplicationCore/Public/Android/AndroidInputInterface.h b/Engine/Source/Runtime/ApplicationCore/Public/Android/AndroidInputInterface.h
--- a/Engine/Source/Runtime/ApplicationCore/Public/Android/AndroidInputInterface.h
+++ b/Engine/Source/Runtime/ApplicationCore/Public/Android/AndroidInputInterface.h
@@ -267,6 +267,8 @@ public:
 	virtual void ResetLightColor(int32 ControllerId) override;
 
 	void SetGamepadsAllowed(bool bAllowed) { bAllowControllers = bAllowed; }
+	void SetGamepadsBlockDeviceFeedback(bool bBlock) { bControllersBlockDeviceFeedback = bBlock; }
+
 	virtual bool IsGamepadAttached() const;
 
 
@@ -332,6 +334,9 @@ private:
 	// should we allow controllers to send input
 	static bool bAllowControllers;
 
+	// bluetooth connected controllers will block force feedback.
+	static bool bControllersBlockDeviceFeedback;
+
 	// should we allow controllers to send Android_Back and Android_Menu events
 	static bool bBlockAndroidKeysOnControllers;
 
diff --git a/Engine/Source/Runtime/ApplicationCore/Public/Android/AndroidPlatformApplicationMisc.h b/Engine/Source/Runtime/ApplicationCore/Public/Android/AndroidPlatformApplicationMisc.h
--- a/Engine/Source/Runtime/ApplicationCore/Public/Android/AndroidPlatformApplicationMisc.h
+++ b/Engine/Source/Runtime/ApplicationCore/Public/Android/AndroidPlatformApplicationMisc.h
@@ -14,6 +14,7 @@ struct APPLICATIONCORE_API FAndroidApplicationMisc : public FGenericPlatformAppl
 	static bool IsScreensaverEnabled();
 	static bool ControlScreensaver(EScreenSaverAction Action);
 	static void SetGamepadsAllowed(bool bAllowed);
+	static void SetGamepadsBlockDeviceFeedback(bool bBlock);
 	static void ResetGamepadAssignments();
 	static void ResetGamepadAssignmentToController(int32 ControllerId);
 	static bool IsControllerAssignedToGamepad(int32 ControllerId);
diff --git a/Engine/Source/Runtime/ApplicationCore/Public/GenericPlatform/GenericPlatformApplicationMisc.h b/Engine/Source/Runtime/ApplicationCore/Public/GenericPlatform/GenericPlatformApplicationMisc.h
--- a/Engine/Source/Runtime/ApplicationCore/Public/GenericPlatform/GenericPlatformApplicationMisc.h
+++ b/Engine/Source/Runtime/ApplicationCore/Public/GenericPlatform/GenericPlatformApplicationMisc.h
@@ -178,6 +178,12 @@ struct APPLICATIONCORE_API FGenericPlatformApplicationMisc
 	static void SetGamepadsAllowed(bool bAllowed)
 	{}
 
+	/*
+	* Set whether gamepads are allowed at the platform level.
+	*/
+	static void SetGamepadsBlockDeviceFeedback(bool bAllowed)
+	{}
+
 	/*
 	 * Resets the gamepad to player controller id assignments
 	 */
diff --git a/Engine/Source/Runtime/ApplicationCore/Public/IOS/IOSInputInterface.h b/Engine/Source/Runtime/ApplicationCore/Public/IOS/IOSInputInterface.h
--- a/Engine/Source/Runtime/ApplicationCore/Public/IOS/IOSInputInterface.h
+++ b/Engine/Source/Runtime/ApplicationCore/Public/IOS/IOSInputInterface.h
@@ -79,6 +79,7 @@ public:
 	virtual void ResetLightColor(int32 ControllerId) override {}
 
 	void SetGamepadsAllowed(bool bAllowed) { bAllowControllers = bAllowed; }
+	void SetGamepadsBlockDeviceFeedback(bool bBlock) { bControllersBlockDeviceFeedback = bBlock; }
 	bool IsControllerAssignedToGamepad(int32 ControllerId) const;
 	bool IsGamepadAttached() const;
 
@@ -149,6 +150,9 @@ private:
 	
 	// should we allow controllers to send input
 	bool bAllowControllers;
+
+	// bluetooth connected controllers will block force feedback.
+	bool bControllersBlockDeviceFeedback;
 	
 	/** Is motion paused or not? */
 	bool bPauseMotion;
diff --git a/Engine/Source/Runtime/ApplicationCore/Public/IOS/IOSPlatformApplicationMisc.h b/Engine/Source/Runtime/ApplicationCore/Public/IOS/IOSPlatformApplicationMisc.h
--- a/Engine/Source/Runtime/ApplicationCore/Public/IOS/IOSPlatformApplicationMisc.h
+++ b/Engine/Source/Runtime/ApplicationCore/Public/IOS/IOSPlatformApplicationMisc.h
@@ -15,6 +15,7 @@ struct APPLICATIONCORE_API FIOSPlatformApplicationMisc : public FGenericPlatform
 	static bool ControlScreensaver(EScreenSaverAction Action);
 
 	static void SetGamepadsAllowed(bool bAllowed);
+	static void SetGamepadsBlockDeviceFeedback(bool bBlock);
 	static void ResetGamepadAssignments();
 	static void ResetGamepadAssignmentToController(int32 ControllerId);
 	static bool IsControllerAssignedToGamepad(int32 ControllerId);
diff --git a/Engine/Source/Runtime/Engine/Classes/Kismet/KismetSystemLibrary.h b/Engine/Source/Runtime/Engine/Classes/Kismet/KismetSystemLibrary.h
--- a/Engine/Source/Runtime/Engine/Classes/Kismet/KismetSystemLibrary.h
+++ b/Engine/Source/Runtime/Engine/Classes/Kismet/KismetSystemLibrary.h
@@ -1626,6 +1626,12 @@ class ENGINE_API UKismetSystemLibrary : public UBlueprintFunctionLibrary
 	UFUNCTION(BlueprintPure, Category = "Utilities|Platform")
 	static bool GetVolumeButtonsHandledBySystem();
 
+	/**
+	 * Sets whether attached gamepads will block feedback from the device itself (Mobile only).
+	 */
+	UFUNCTION(BlueprintCallable, Category = "Utilities|Platform")
+	static void SetGamepadsBlockDeviceFeedback(bool bBlock);
+
 	/**
 	 * Resets the gamepad to player controller id assignments (Android and iOS only)
 	 */
diff --git a/Engine/Source/Runtime/Engine/Private/KismetSystemLibrary.cpp b/Engine/Source/Runtime/Engine/Private/KismetSystemLibrary.cpp
--- a/Engine/Source/Runtime/Engine/Private/KismetSystemLibrary.cpp
+++ b/Engine/Source/Runtime/Engine/Private/KismetSystemLibrary.cpp
@@ -2486,6 +2486,12 @@ bool UKismetSystemLibrary::GetVolumeButtonsHandledBySystem()
 	return FPlatformMisc::GetVolumeButtonsHandledBySystem();
 }
 
+void UKismetSystemLibrary::SetGamepadsBlockDeviceFeedback(bool bBlock)
+{
+	FPlatformApplicationMisc::SetGamepadsBlockDeviceFeedback(bBlock);
+
+}
+
 void UKismetSystemLibrary::ResetGamepadAssignments()
 {
 	FPlatformApplicationMisc::ResetGamepadAssignments();
diff --git a/Engine/Source/Runtime/IOS/IOSRuntimeSettings/Classes/IOSRuntimeSettings.h b/Engine/Source/Runtime/IOS/IOSRuntimeSettings/Classes/IOSRuntimeSettings.h
--- a/Engine/Source/Runtime/IOS/IOSRuntimeSettings/Classes/IOSRuntimeSettings.h
+++ b/Engine/Source/Runtime/IOS/IOSRuntimeSettings/Classes/IOSRuntimeSettings.h
@@ -348,6 +348,10 @@ public:
 	UPROPERTY(GlobalConfig, EditAnywhere, Category = Input, meta = (DisplayName = "Allow MFi (Bluetooth) controllers"))
 	bool bAllowControllers;
 
+	// Block force feedback on the device when controllers are attached.
+	UPROPERTY(GlobalConfig, EditAnywhere, Category = Input, meta = (DisplayName = "Block force feedback on the device when controllers are attached"))
+	bool bControllersBlockDeviceFeedback;
+
 	// Disables usage of device motion data. If application does not use motion data disabling it will improve battery life
 	UPROPERTY(GlobalConfig, EditAnywhere, Category = Input, meta = (DisplayName = "Disable Motion Controls"))
 	bool bDisableMotionData;
