Adding SetGamepadsBlockDeviceFeedback to lumin input interface.

diff --git a/Engine/Source/Runtime/ApplicationCore/Public/Lumin/LuminInputInterface.h b/Engine/Source/Runtime/ApplicationCore/Public/Lumin/LuminInputInterface.h
--- a/Engine/Source/Runtime/ApplicationCore/Public/Lumin/LuminInputInterface.h
+++ b/Engine/Source/Runtime/ApplicationCore/Public/Lumin/LuminInputInterface.h
@@ -38,6 +38,7 @@ public:
 	virtual void ResetLightColor(int32 ControllerId) override {};
 
 	void SetGamepadsAllowed(bool bAllowed) {}
+	void SetGamepadsBlockDeviceFeedback(bool bBlock) { }
 	virtual void AddExternalInputDevice(TSharedPtr<class IInputDevice>& InputDevice);
 	virtual bool IsGamepadAttached() const;
 
