fix for fps analytics on IOS when greater than 60Hz

diff --git a/Engine/Source/Runtime/ApplicationCore/Public/GenericPlatform/GenericPlatformFramePacer.h b/Engine/Source/Runtime/ApplicationCore/Public/GenericPlatform/GenericPlatformFramePacer.h
--- a/Engine/Source/Runtime/ApplicationCore/Public/GenericPlatform/GenericPlatformFramePacer.h
+++ b/Engine/Source/Runtime/ApplicationCore/Public/GenericPlatform/GenericPlatformFramePacer.h
@@ -38,4 +38,9 @@ struct APPLICATIONCORE_API FGenericPlatformRHIFramePacer
 	 * The pace we are running at (30 = 30fps, 0 = unpaced)
 	 */
 	static uint32 GetFramePace() { return 0; };
+
+	/**
+	 * The maximum refresh rate the pacer supports
+	 */
+	static uint32 GetMaxRefreshRate() { return 60; }
 };
