addition of API to query update status on the app store for IOS, this allows the app to know whether or not an update is available on the store before sending the user there

diff --git a/Engine/Source/Runtime/ApplicationCore/Private/IOS/IOSAppDelegate.cpp b/Engine/Source/Runtime/ApplicationCore/Private/IOS/IOSAppDelegate.cpp
--- a/Engine/Source/Runtime/ApplicationCore/Private/IOS/IOSAppDelegate.cpp
+++ b/Engine/Source/Runtime/ApplicationCore/Private/IOS/IOSAppDelegate.cpp
@@ -298,6 +298,11 @@ static IOSAppDelegate* CachedDelegate = nil;
 	// make sure this thread has an auto release pool setup
 	NSAutoreleasePool* AutoreleasePool = [[NSAutoreleasePool alloc] init];
 
+	// check for update on app store if cvar is enabled
+	dispatch_async(dispatch_get_main_queue(), ^{
+		[[IOSAppDelegate GetDelegate] DoUpdateCheck];
+	});
+	
 	{
 		SCOPED_BOOT_TIMING("[IOSAppDelegate MainAppThread setup]");
 
@@ -1285,6 +1290,61 @@ static FAutoConsoleVariableRef CVarGEnableThermalsReport(
 }
 #endif
 
+// checking for update on the app store
+-(void)DoUpdateCheck
+{
+	static bool bInit = false;
+	static NSString* CurrentVersion = [[NSBundle mainBundle] infoDictionary][@"CFBundleShortVersionString"];
+	static NSString* BundleID = [[NSBundle mainBundle] infoDictionary][@"CFBundleIdentifier"];
+	bool bEnableUpdateCheck = NO;
+	if (!bInit)
+	{
+		GConfig->GetBool(TEXT("/Script/IOSRuntimeSettings.IOSRuntimeSettings"), TEXT("bEnableUpdateCheck"), bEnableUpdateCheck, GEngineIni);
+		self.bUpdateAvailable = false;
+		bInit = true;
+	}
+	if (bEnableUpdateCheck)
+	{
+		// kick off a check on the app store for an update
+		NSLocale* Locale = [NSLocale autoupdatingCurrentLocale];
+		NSURL* StoreURL = [NSURL URLWithString: [NSString stringWithFormat: @"http://itunes.apple.com/%@/lookup?bundleId=%@", Locale.countryCode, BundleID]];
+		
+		// kick off an NSURLSession to read the data
+		NSURLSession* Session = [NSURLSession sharedSession];
+		NSURLSessionDataTask* SessionTask = [Session dataTaskWithRequest: [NSURLRequest requestWithURL: StoreURL] completionHandler:^(NSData* _Nullable data, NSURLResponse* _Nullable response, NSError* _Nullable error) {
+			NSDictionary* StoreDictionary = [NSJSONSerialization JSONObjectWithData: data options: 0 error: nil];
+			
+			if ([StoreDictionary[@"resultCount"] integerValue] == 1)
+			{
+				// get the store version
+				NSString* StoreVersion = StoreDictionary[@"results"][0][@"version"];
+				if ([StoreVersion compare: CurrentVersion options: NSNumericSearch] == NSOrderedDescending)
+				{
+					self.bUpdateAvailable = true;
+				}
+				else
+				{
+					self.bUpdateAvailable = false;
+				}
+			}
+			else
+			{
+				self.bUpdateAvailable = false;
+			}
+		}];
+		
+		[SessionTask resume];
+	}
+}
+
+-(bool)IsUpdateAvailable
+{
+	dispatch_async(dispatch_get_main_queue(), ^{
+		[[IOSAppDelegate GetDelegate] DoUpdateCheck];
+	});
+	return self.bUpdateAvailable;
+}
+
 - (void) StartGameThread
 {
 	// create a new thread (the pointer will be retained forever)
diff --git a/Engine/Source/Runtime/ApplicationCore/Public/IOS/IOSAppDelegate.h b/Engine/Source/Runtime/ApplicationCore/Public/IOS/IOSAppDelegate.h
--- a/Engine/Source/Runtime/ApplicationCore/Public/IOS/IOSAppDelegate.h
+++ b/Engine/Source/Runtime/ApplicationCore/Public/IOS/IOSAppDelegate.h
@@ -188,6 +188,8 @@ APPLICATIONCORE_API
 @property (assign) bool bBatteryState;
 @property (assign) int BatteryLevel;
 
+@property (assign) bool bUpdateAvailable;
+
 /**
  * @return the single app delegate object
  */
@@ -209,6 +211,8 @@ APPLICATIONCORE_API
 -(float)GetBackgroundingMainThreadBlockTime;
 -(void)OverrideBackgroundingMainThreadBlockTime:(float)BlockTime;
 
+-(bool)IsUpdateAvailable;
+
 @property (assign) bool bAudioSessionInitialized;
 
 /** TRUE if the device is playing background music and we want to allow that */
diff --git a/Engine/Source/Runtime/Core/Private/IOS/IOSPlatformMisc.cpp b/Engine/Source/Runtime/Core/Private/IOS/IOSPlatformMisc.cpp
--- a/Engine/Source/Runtime/Core/Private/IOS/IOSPlatformMisc.cpp
+++ b/Engine/Source/Runtime/Core/Private/IOS/IOSPlatformMisc.cpp
@@ -846,6 +846,11 @@ void FIOSPlatformMisc::RequestStoreReview()
 #endif
 }
 
+bool FIOSPlatformMisc::IsUpdateAvailable()
+{
+	return [[IOSAppDelegate GetDelegate] IsUpdateAvailable];
+}
+
 /**
 * Returns a unique string for advertising identification
 *
diff --git a/Engine/Source/Runtime/Core/Public/IOS/IOSPlatformMisc.h b/Engine/Source/Runtime/Core/Public/IOS/IOSPlatformMisc.h
--- a/Engine/Source/Runtime/Core/Public/IOS/IOSPlatformMisc.h
+++ b/Engine/Source/Runtime/Core/Public/IOS/IOSPlatformMisc.h
@@ -101,6 +101,8 @@ struct CORE_API FIOSPlatformMisc : public FApplePlatformMisc
 	
 	static void RequestStoreReview();
 
+	static bool IsUpdateAvailable();
+	
 	// Possible iOS devices
 	enum EIOSDevice
 	{
