fix for properly starting the cache of the update data if the ini data is loaded

diff --git a/Engine/Source/Runtime/ApplicationCore/Private/IOS/IOSAppDelegate.cpp b/Engine/Source/Runtime/ApplicationCore/Private/IOS/IOSAppDelegate.cpp
--- a/Engine/Source/Runtime/ApplicationCore/Private/IOS/IOSAppDelegate.cpp
+++ b/Engine/Source/Runtime/ApplicationCore/Private/IOS/IOSAppDelegate.cpp
@@ -298,11 +298,6 @@ static IOSAppDelegate* CachedDelegate = nil;
 	// make sure this thread has an auto release pool setup
 	NSAutoreleasePool* AutoreleasePool = [[NSAutoreleasePool alloc] init];
 
-	// check for update on app store if cvar is enabled
-	dispatch_async(dispatch_get_main_queue(), ^{
-		[[IOSAppDelegate GetDelegate] DoUpdateCheck];
-	});
-	
 	{
 		SCOPED_BOOT_TIMING("[IOSAppDelegate MainAppThread setup]");
 
@@ -319,6 +314,11 @@ static IOSAppDelegate* CachedDelegate = nil;
 
 	FAppEntry::Init();
 
+	// check for update on app store if cvar is enabled
+	dispatch_async(dispatch_get_main_queue(), ^{
+		[[IOSAppDelegate GetDelegate] DoUpdateCheck];
+	});
+	
 	// now that GConfig has been loaded, load the EnabledAudioFeatures from ini
 	TArray<FString> EnabledAudioFeatures;
 	GConfig->GetArray(TEXT("Audio"), TEXT("EnabledAudioFeatures"), EnabledAudioFeatures, GEngineIni);
@@ -1299,11 +1299,11 @@ static FAutoConsoleVariableRef CVarGEnableThermalsReport(
 	bool bEnableUpdateCheck = NO;
 	if (!bInit)
 	{
-		GConfig->GetBool(TEXT("/Script/IOSRuntimeSettings.IOSRuntimeSettings"), TEXT("bEnableUpdateCheck"), bEnableUpdateCheck, GEngineIni);
+		bool bReadData = GConfig->GetBool(TEXT("/Script/IOSRuntimeSettings.IOSRuntimeSettings"), TEXT("bEnableUpdateCheck"), bEnableUpdateCheck, GEngineIni);
 		self.bUpdateAvailable = false;
-		bInit = true;
+		bInit = bReadData;
 	}
-	if (bEnableUpdateCheck)
+	if (bEnableUpdateCheck && bInit)
 	{
 		// kick off a check on the app store for an update
 		NSLocale* Locale = [NSLocale autoupdatingCurrentLocale];
