Fixing mac build.  Casting the int32 inputs to floats to use the GetRangePct function.

diff --git a/Engine/Source/Runtime/ApplicationCore/Private/Mac/HIDInputInterface.cpp b/Engine/Source/Runtime/ApplicationCore/Private/Mac/HIDInputInterface.cpp
--- a/Engine/Source/Runtime/ApplicationCore/Private/Mac/HIDInputInterface.cpp
+++ b/Engine/Source/Runtime/ApplicationCore/Private/Mac/HIDInputInterface.cpp
@@ -456,7 +456,7 @@ void HIDInputInterface::SendControllerEvents()
 					else
 					{
 						const bool bIsTrigger = Element.Usage == ControllerState.Device.LeftTriggerAnalogMapping || Element.Usage == ControllerState.Device.RightTriggerAnalogMapping;
-						const float Percentage = FMath::GetRangePct(Element.MinValue, Element.MaxValue, NewValue);
+						const float Percentage = FMath::GetRangePct((float)Element.MinValue, (float)Element.MaxValue, (float)NewValue);
 						const float FloatValue = bIsTrigger ? Percentage : Percentage * 2.0f - 1.0f;
 
 						if (Element.Usage == ControllerState.Device.LeftAnalogXMapping && ControllerState.LeftAnalogX != NewValue)
