Try to fix IOS13 input issue.

diff --git a/Engine/Source/Runtime/ApplicationCore/Private/IOS/IOSView.cpp b/Engine/Source/Runtime/ApplicationCore/Private/IOS/IOSView.cpp
--- a/Engine/Source/Runtime/ApplicationCore/Private/IOS/IOSView.cpp
+++ b/Engine/Source/Runtime/ApplicationCore/Private/IOS/IOSView.cpp
@@ -177,6 +177,8 @@ id<MTLDevice> GMetalDevice = nil;
 {
 	SCOPED_BOOT_TIMING("[FIOSView initInternal]");
 
+	CachedMarkedText = nil;
+
 	// figure out if we should start up GL or Metal
 #if HAS_METAL
 	// if the device is valid, we know Metal is usable (see +layerClass)
@@ -260,6 +262,7 @@ id<MTLDevice> GMetalDevice = nil;
 
 -(void)dealloc
 {
+	[CachedMarkedText release];
 	[markedTextStyle release];
 #if WITH_ACCESSIBILITY
 	[_accessibilityElements release];
@@ -629,6 +632,11 @@ id<MTLDevice> GMetalDevice = nil;
 
 - (void)insertText:(NSString *)theText
 {
+	if (nil != CachedMarkedText) {
+		[CachedMarkedText release];
+		CachedMarkedText = nil;
+	}
+
 	// insert text one key at a time, as chars, not keydowns
 	for (int32 CharIndex = 0; CharIndex < [theText length]; CharIndex++)
 	{
@@ -653,6 +661,10 @@ id<MTLDevice> GMetalDevice = nil;
 
 - (void)deleteBackward
 {
+	if (nil != CachedMarkedText) {
+		[CachedMarkedText release];
+		CachedMarkedText = nil;
+	}
 	FIOSInputInterface::QueueKeyInput(KEYCODE_BACKSPACE, '\b');
 }
 
@@ -739,8 +751,10 @@ id<MTLDevice> GMetalDevice = nil;
 
 - (NSString *)textInRange:(UITextRange *)range
 {
-	// @todo keyboard: This is called
-	return @"";
+	if (nil != CachedMarkedText) {
+		return CachedMarkedText;
+	}
+	return nil;
 	//IndexedRange *r = (IndexedRange *)range;
 	//return ([textStore substringWithRange:r.range]);
 }
@@ -770,14 +784,23 @@ id<MTLDevice> GMetalDevice = nil;
 
 - (UITextRange *)markedTextRange
 {
-	// @todo keyboard: This is called
-	return nil;
+	if (nil != CachedMarkedText) {
+		return[[[UITextRange alloc] init] autorelease];
+	}
+	return nil; // Nil if no marked text.
 }
 
 
 - (void)setMarkedText:(NSString *)markedText selectedRange:(NSRange)selectedRange
 {
+	if (markedText == CachedMarkedText) {
+		return;
+	}
+	if (nil != CachedMarkedText) {
+		[CachedMarkedText release];
+	}
 	CachedMarkedText = markedText;
+	[CachedMarkedText retain];
 	//NSLog(@"setting marked text to %@", markedText);
 }
 
@@ -787,6 +810,7 @@ id<MTLDevice> GMetalDevice = nil;
 	if (CachedMarkedText != nil)
 	{
 		[self insertText:CachedMarkedText];
+		[CachedMarkedText release];
 		CachedMarkedText = nil;
 	}
 }
@@ -863,15 +887,15 @@ id<MTLDevice> GMetalDevice = nil;
 }
 
 
-- (UITextWritingDirection)baseWritingDirectionForPosition:(UITextPosition *)position inDirection:(UITextStorageDirection)direction
+- (NSWritingDirection)baseWritingDirectionForPosition:(UITextPosition *)position inDirection:(UITextStorageDirection)direction
 {
 	REPORT_EVENT;
 	// assume left to right for now
-	return UITextWritingDirectionLeftToRight;
+	return NSWritingDirectionLeftToRight;
 }
 
 
-- (void)setBaseWritingDirection:(UITextWritingDirection)writingDirection forRange:(UITextRange *)range
+- (void)setBaseWritingDirection:(NSWritingDirection)writingDirection forRange:(UITextRange *)range
 {
 	// @todo keyboard: This is called
 }
