Copying Private-LoadTimes-4.24 stream to Main. Biggest changes are in Materials/Shader memory freezing.

diff --git a/Engine/Source/Runtime/ApplicationCore/Private/Windows/TextStoreACP.cpp b/Engine/Source/Runtime/ApplicationCore/Private/Windows/TextStoreACP.cpp
--- a/Engine/Source/Runtime/ApplicationCore/Private/Windows/TextStoreACP.cpp
+++ b/Engine/Source/Runtime/ApplicationCore/Private/Windows/TextStoreACP.cpp
@@ -79,7 +79,7 @@ STDAPI_(ULONG) FTextStoreACP::Release()
 
 STDAPI FTextStoreACP::AdviseSink(__RPC__in REFIID riid, __RPC__in_opt IUnknown *punk, DWORD dwMask)
 {
-	UE_LOG(LogTextStoreACP, Verbose, TEXT("AdviseSink"));
+	UE_LOG_CLINKAGE(LogTextStoreACP, Verbose, TEXT("AdviseSink"));
 
 	// punk must not be null.
 	if(!punk)
@@ -137,7 +137,7 @@ STDAPI FTextStoreACP::AdviseSink(__RPC__in REFIID riid, __RPC__in_opt IUnknown *
 
 STDAPI FTextStoreACP::UnadviseSink(__RPC__in_opt IUnknown *punk)
 {
-	UE_LOG(LogTextStoreACP, Verbose, TEXT("UnadviseSink"));
+	UE_LOG_CLINKAGE(LogTextStoreACP, Verbose, TEXT("UnadviseSink"));
 
 	// punk must not be null.
 	if(!punk)
@@ -176,7 +176,7 @@ STDAPI FTextStoreACP::UnadviseSink(__RPC__in_opt IUnknown *punk)
 
 STDAPI FTextStoreACP::RequestLock(DWORD dwLockFlags, HRESULT *phrSession)
 {
-	UE_LOG(LogTextStoreACP, Verbose, TEXT("RequestLock"));
+	UE_LOG_CLINKAGE(LogTextStoreACP, Verbose, TEXT("RequestLock"));
 
 	// Must have a sink object.
 	if(!AdviseSinkObject.TextStoreACPSink)
@@ -194,32 +194,32 @@ STDAPI FTextStoreACP::RequestLock(DWORD dwLockFlags, HRESULT *phrSession)
 	if(LockManager.LockType == 0)
 	{
 		// Flag as locked.
-		UE_LOG(LogTextStoreACP, Verbose, TEXT("RequestLock - Locking..."));
+		UE_LOG_CLINKAGE(LogTextStoreACP, Verbose, TEXT("RequestLock - Locking..."));
 		LockManager.LockType = dwLockFlags & (~TS_LF_SYNC);
 
-		UE_LOG(LogTextStoreACP, Verbose, TEXT("RequestLock - Locked"));
+		UE_LOG_CLINKAGE(LogTextStoreACP, Verbose, TEXT("RequestLock - Locked"));
 
 		*phrSession = AdviseSinkObject.TextStoreACPSink->OnLockGranted(LockManager.LockType);
 
-		UE_LOG(LogTextStoreACP, Verbose, TEXT("RequestLock - Unlocking..."));
+		UE_LOG_CLINKAGE(LogTextStoreACP, Verbose, TEXT("RequestLock - Unlocking..."));
 
 		while(LockManager.IsPendingLockUpgrade)
 		{
-			UE_LOG(LogTextStoreACP, Verbose, TEXT("RequestLock - Upgrading..."));
+			UE_LOG_CLINKAGE(LogTextStoreACP, Verbose, TEXT("RequestLock - Upgrading..."));
 
 			LockManager.LockType = TS_LF_READWRITE;
 			LockManager.IsPendingLockUpgrade = false;
 
-			UE_LOG(LogTextStoreACP, Verbose, TEXT("RequestLock - Upgraded"));
+			UE_LOG_CLINKAGE(LogTextStoreACP, Verbose, TEXT("RequestLock - Upgraded"));
 
 			AdviseSinkObject.TextStoreACPSink->OnLockGranted(LockManager.LockType);
 
-			UE_LOG(LogTextStoreACP, Verbose, TEXT("RequestLock - Unlocking..."));
+			UE_LOG_CLINKAGE(LogTextStoreACP, Verbose, TEXT("RequestLock - Unlocking..."));
 		}
 
 		// Flag as unlocked.
 		LockManager.LockType = 0;
-		UE_LOG(LogTextStoreACP, Verbose, TEXT("RequestLock - Unlocked"));
+		UE_LOG_CLINKAGE(LogTextStoreACP, Verbose, TEXT("RequestLock - Unlocked"));
 
 		return S_OK;
 	}
@@ -241,7 +241,7 @@ STDAPI FTextStoreACP::RequestLock(DWORD dwLockFlags, HRESULT *phrSession)
 
 STDAPI FTextStoreACP::GetStatus(__RPC__out TS_STATUS *pdcs)
 {
-	UE_LOG(LogTextStoreACP, Verbose, TEXT("GetStatus"));
+	UE_LOG_CLINKAGE(LogTextStoreACP, Verbose, TEXT("GetStatus"));
 
 	// pdcs cannot be null.
 	if(!pdcs)
@@ -257,7 +257,7 @@ STDAPI FTextStoreACP::GetStatus(__RPC__out TS_STATUS *pdcs)
 
 STDAPI FTextStoreACP::GetEndACP(__RPC__out LONG *pacp)
 {
-	UE_LOG(LogTextStoreACP, Verbose, TEXT("GetEndACP"));
+	UE_LOG_CLINKAGE(LogTextStoreACP, Verbose, TEXT("GetEndACP"));
 
 	if(!IsFlaggedReadLocked(LockManager.LockType))
 	{
@@ -271,7 +271,7 @@ STDAPI FTextStoreACP::GetEndACP(__RPC__out LONG *pacp)
 
 STDAPI FTextStoreACP::GetSelection(ULONG ulIndex, ULONG ulCount, __RPC__out_ecount_part(ulCount, *pcFetched) TS_SELECTION_ACP *pSelection, __RPC__out ULONG *pcFetched)
 {
-	UE_LOG(LogTextStoreACP, Verbose, TEXT("GetSelection"));
+	UE_LOG_CLINKAGE(LogTextStoreACP, Verbose, TEXT("GetSelection"));
 
 	if(!IsFlaggedReadLocked(LockManager.LockType))
 	{
@@ -312,7 +312,7 @@ STDAPI FTextStoreACP::GetSelection(ULONG ulIndex, ULONG ulCount, __RPC__out_ecou
 
 STDAPI FTextStoreACP::SetSelection(ULONG ulCount, __RPC__in_ecount_full(ulCount) const TS_SELECTION_ACP *pSelection)
 {
-	UE_LOG(LogTextStoreACP, Verbose, TEXT("SetSelection - From %d to %d"), pSelection->acpStart, pSelection->acpEnd);
+	UE_LOG_CLINKAGE(LogTextStoreACP, Verbose, TEXT("SetSelection - From %d to %d"), pSelection->acpStart, pSelection->acpEnd);
 
 	if(!IsFlaggedReadWriteLocked(LockManager.LockType))
 	{
@@ -347,7 +347,7 @@ STDAPI FTextStoreACP::SetSelection(ULONG ulCount, __RPC__in_ecount_full(ulCount)
 
 STDAPI FTextStoreACP::RequestSupportedAttrs(DWORD dwFlags, ULONG cFilterAttrs, __RPC__in_ecount_full_opt(cFilterAttrs) const TS_ATTRID *paFilterAttrs)
 {
-	UE_LOG(LogTextStoreACP, Verbose, TEXT("RequestSupportedAttrs"));
+	UE_LOG_CLINKAGE(LogTextStoreACP, Verbose, TEXT("RequestSupportedAttrs"));
 
 	for(ULONG i = 0; i < cFilterAttrs; ++i)
 	{
@@ -368,7 +368,7 @@ STDAPI FTextStoreACP::RequestSupportedAttrs(DWORD dwFlags, ULONG cFilterAttrs, _
 
 STDAPI FTextStoreACP::RequestAttrsAtPosition(LONG acpPos, ULONG cFilterAttrs, __RPC__in_ecount_full_opt(cFilterAttrs) const TS_ATTRID *paFilterAttrs, DWORD dwFlags)
 {
-	UE_LOG(LogTextStoreACP, Verbose, TEXT("RequestAttrsAtPosition"));
+	UE_LOG_CLINKAGE(LogTextStoreACP, Verbose, TEXT("RequestAttrsAtPosition"));
 
 	for(ULONG i = 0; i < cFilterAttrs; ++i)
 	{
@@ -389,20 +389,20 @@ STDAPI FTextStoreACP::RequestAttrsAtPosition(LONG acpPos, ULONG cFilterAttrs, __
 
 STDAPI FTextStoreACP::RequestAttrsTransitioningAtPosition(LONG acpPos, ULONG cFilterAttrs, __RPC__in_ecount_full_opt(cFilterAttrs) const TS_ATTRID *paFilterAttrs, DWORD dwFlags)
 {
-	UE_LOG(LogTextStoreACP, Verbose, TEXT("RequestAttrsTransitioningAtPosition"));
+	UE_LOG_CLINKAGE(LogTextStoreACP, Verbose, TEXT("RequestAttrsTransitioningAtPosition"));
 	return E_NOTIMPL; // TODO: No meaningful documentation to this method, not implemented.
 }
 
 STDAPI FTextStoreACP::FindNextAttrTransition(LONG acpStart, LONG acpHalt, ULONG cFilterAttrs, __RPC__in_ecount_full_opt(cFilterAttrs) const TS_ATTRID *paFilterAttrs, DWORD dwFlags, __RPC__out LONG *pacpNext, __RPC__out BOOL *pfFound, __RPC__out LONG *plFoundOffset)
 {
-	UE_LOG(LogTextStoreACP, Verbose, TEXT("FindNextAttrTransition"));
+	UE_LOG_CLINKAGE(LogTextStoreACP, Verbose, TEXT("FindNextAttrTransition"));
 	return E_NOTIMPL; // TODO: No meaningful documentation to this method, not implemented.
 }
 
 CA_SUPPRESS(6101)
 STDAPI FTextStoreACP::RetrieveRequestedAttrs(ULONG ulCount, __RPC__out_ecount_part(ulCount, *pcFetched) TS_ATTRVAL *paAttrVals, __RPC__out ULONG *pcFetched)
 {
-	UE_LOG(LogTextStoreACP, Verbose, TEXT("RetrieveRequestedAttrs"));
+	UE_LOG_CLINKAGE(LogTextStoreACP, Verbose, TEXT("RetrieveRequestedAttrs"));
 
 	*pcFetched = 0;
 
@@ -424,7 +424,7 @@ STDAPI FTextStoreACP::RetrieveRequestedAttrs(ULONG ulCount, __RPC__out_ecount_pa
 
 STDAPI FTextStoreACP::GetActiveView(__RPC__out TsViewCookie *pvcView)
 {
-	UE_LOG(LogTextStoreACP, Verbose, TEXT("GetActiveView"));
+	UE_LOG_CLINKAGE(LogTextStoreACP, Verbose, TEXT("GetActiveView"));
 	*pvcView = 0;
 
 	return S_OK;
@@ -432,7 +432,7 @@ STDAPI FTextStoreACP::GetActiveView(__RPC__out TsViewCookie *pvcView)
 
 STDAPI FTextStoreACP::GetACPFromPoint(TsViewCookie vcView, __RPC__in const POINT *pt, DWORD dwFlags, __RPC__out LONG *pacp)
 {
-	UE_LOG(LogTextStoreACP, Verbose, TEXT("GetACPFromPoint - At (%d, %d)"), pt->x, pt->y);
+	UE_LOG_CLINKAGE(LogTextStoreACP, Verbose, TEXT("GetACPFromPoint - At (%d, %d)"), pt->x, pt->y);
 
 	if(!IsFlaggedReadLocked(LockManager.LockType))
 	{
@@ -449,7 +449,7 @@ STDAPI FTextStoreACP::GetACPFromPoint(TsViewCookie vcView, __RPC__in const POINT
 
 STDAPI FTextStoreACP::GetTextExt(TsViewCookie vcView, LONG acpStart, LONG acpEnd, __RPC__out RECT *prc, __RPC__out BOOL *pfClipped)
 {
-	UE_LOG(LogTextStoreACP, Verbose, TEXT("GetTextExt - From %d to %d"), acpStart, acpEnd);
+	UE_LOG_CLINKAGE(LogTextStoreACP, Verbose, TEXT("GetTextExt - From %d to %d"), acpStart, acpEnd);
 
 	if(!IsFlaggedReadLocked(LockManager.LockType))
 	{
@@ -483,7 +483,7 @@ STDAPI FTextStoreACP::GetTextExt(TsViewCookie vcView, LONG acpStart, LONG acpEnd
 
 STDAPI FTextStoreACP::GetScreenExt(TsViewCookie vcView, __RPC__out RECT *prc)
 {
-	UE_LOG(LogTextStoreACP, Verbose, TEXT("GetScreenExt"));
+	UE_LOG_CLINKAGE(LogTextStoreACP, Verbose, TEXT("GetScreenExt"));
 
 	// View cookie must be valid.
 	if(vcView != 0)
@@ -506,7 +506,7 @@ CA_SUPPRESS(28196)
 CA_SUPPRESS(6387)
 STDAPI FTextStoreACP::GetWnd(TsViewCookie vcView, __RPC__deref_out_opt HWND *phwnd)
 {
-	UE_LOG(LogTextStoreACP, Verbose, TEXT("GetWnd"));
+	UE_LOG_CLINKAGE(LogTextStoreACP, Verbose, TEXT("GetWnd"));
 
 	const TSharedPtr<FGenericWindow> GenericWindow = TextContext->GetWindow();
 	*phwnd = GenericWindow.IsValid() ? reinterpret_cast<HWND>(GenericWindow->GetOSWindowHandle()) : nullptr;
@@ -516,7 +516,7 @@ STDAPI FTextStoreACP::GetWnd(TsViewCookie vcView, __RPC__deref_out_opt HWND *phw
 
 STDAPI FTextStoreACP::GetText(LONG acpStart, LONG acpEnd, __RPC__out_ecount_part(cchPlainReq, *pcchPlainOut) WCHAR *pchPlain, ULONG cchPlainReq, __RPC__out ULONG *pcchPlainOut, __RPC__out_ecount_part(ulRunInfoReq, *pulRunInfoOut) TS_RUNINFO *prgRunInfo, ULONG ulRunInfoReq, __RPC__out ULONG *pulRunInfoOut, __RPC__out LONG *pacpNext)
 {
-	UE_LOG(LogTextStoreACP, Verbose, TEXT("GetText - From %d to %d"), acpStart, acpEnd);
+	UE_LOG_CLINKAGE(LogTextStoreACP, Verbose, TEXT("GetText - From %d to %d"), acpStart, acpEnd);
 
 	if(!IsFlaggedReadLocked(LockManager.LockType))
 	{
@@ -565,7 +565,7 @@ STDAPI FTextStoreACP::GetText(LONG acpStart, LONG acpEnd, __RPC__out_ecount_part
 
 STDAPI FTextStoreACP::QueryInsert(LONG acpInsertStart, LONG acpInsertEnd, ULONG cch, __RPC__out LONG *pacpResultStart, __RPC__out LONG *pacpResultEnd)
 {
-	UE_LOG(LogTextStoreACP, Verbose, TEXT("QueryInsert"));
+	UE_LOG_CLINKAGE(LogTextStoreACP, Verbose, TEXT("QueryInsert"));
 
 	// Can't successfully query if there's nowhere to write a result.
 	if(!pacpResultStart || !pacpResultEnd)
@@ -588,7 +588,7 @@ STDAPI FTextStoreACP::QueryInsert(LONG acpInsertStart, LONG acpInsertEnd, ULONG
 CA_SUPPRESS(6101)
 STDAPI FTextStoreACP::InsertTextAtSelection(DWORD dwFlags, __RPC__in_ecount_full(cch) const WCHAR *pchText, ULONG cch, __RPC__out LONG *pacpStart, __RPC__out LONG *pacpEnd, __RPC__out TS_TEXTCHANGE *pChange)
 {
-	UE_LOG(LogTextStoreACP, Verbose, TEXT("InsertTextAtSelection - %s"), pchText);
+	UE_LOG_CLINKAGE(LogTextStoreACP, Verbose, TEXT("InsertTextAtSelection - %s"), pchText);
 
 	// Must have a read lock if querying.
 	if(dwFlags == TS_IAS_QUERYONLY && !IsFlaggedReadLocked(LockManager.LockType))
@@ -670,7 +670,7 @@ STDAPI FTextStoreACP::InsertTextAtSelection(DWORD dwFlags, __RPC__in_ecount_full
 CA_SUPPRESS(6103)
 STDAPI FTextStoreACP::SetText(DWORD dwFlags, LONG acpStart, LONG acpEnd, __RPC__in_ecount_full(cch) const WCHAR *pchText, ULONG cch, __RPC__out TS_TEXTCHANGE *pChange)
 {
-	UE_LOG(LogTextStoreACP, Verbose, TEXT("SetText - From %d to %d, set %s"), acpStart, acpEnd, pchText);
+	UE_LOG_CLINKAGE(LogTextStoreACP, Verbose, TEXT("SetText - From %d to %d, set %s"), acpStart, acpEnd, pchText);
 
 	if(!IsFlaggedReadWriteLocked(LockManager.LockType))
 	{
@@ -701,7 +701,7 @@ STDAPI FTextStoreACP::SetText(DWORD dwFlags, LONG acpStart, LONG acpEnd, __RPC__
 
 STDAPI FTextStoreACP::GetEmbedded(LONG acpPos, __RPC__in REFGUID rguidService, __RPC__in REFIID riid, __RPC__deref_out_opt IUnknown **ppunk)
 {
-	UE_LOG(LogTextStoreACP, Verbose, TEXT("GetEmbedded"));
+	UE_LOG_CLINKAGE(LogTextStoreACP, Verbose, TEXT("GetEmbedded"));
 
 	if(!IsFlaggedReadLocked(LockManager.LockType))
 	{
@@ -713,7 +713,7 @@ STDAPI FTextStoreACP::GetEmbedded(LONG acpPos, __RPC__in REFGUID rguidService, _
 
 STDAPI FTextStoreACP::GetFormattedText(LONG acpStart, LONG acpEnd, __RPC__deref_out_opt IDataObject **ppDataObject)
 {
-	UE_LOG(LogTextStoreACP, Verbose, TEXT("GetFormattedText"));
+	UE_LOG_CLINKAGE(LogTextStoreACP, Verbose, TEXT("GetFormattedText"));
 
 	if(!IsFlaggedReadLocked(LockManager.LockType))
 	{
@@ -725,14 +725,14 @@ STDAPI FTextStoreACP::GetFormattedText(LONG acpStart, LONG acpEnd, __RPC__deref_
 
 STDAPI FTextStoreACP::QueryInsertEmbedded(__RPC__in const GUID *pguidService, __RPC__in const FORMATETC *pFormatEtc, __RPC__out BOOL *pfInsertable)
 {
-	UE_LOG(LogTextStoreACP, Verbose, TEXT("QueryInsertEmbedded"));
+	UE_LOG_CLINKAGE(LogTextStoreACP, Verbose, TEXT("QueryInsertEmbedded"));
 
 	return E_NOTIMPL;
 }
 
 STDAPI FTextStoreACP::InsertEmbedded(DWORD dwFlags, LONG acpStart, LONG acpEnd, __RPC__in_opt IDataObject *pDataObject, __RPC__out TS_TEXTCHANGE *pChange)
 {
-	UE_LOG(LogTextStoreACP, Verbose, TEXT("InsertEmbedded"));
+	UE_LOG_CLINKAGE(LogTextStoreACP, Verbose, TEXT("InsertEmbedded"));
 
 	if(!IsFlaggedReadWriteLocked(LockManager.LockType))
 	{
@@ -744,7 +744,7 @@ STDAPI FTextStoreACP::InsertEmbedded(DWORD dwFlags, LONG acpStart, LONG acpEnd,
 
 STDAPI FTextStoreACP::InsertEmbeddedAtSelection(DWORD dwFlags, __RPC__in_opt IDataObject *pDataObject, __RPC__out LONG *pacpStart, __RPC__out LONG *pacpEnd, __RPC__out TS_TEXTCHANGE *pChange)
 {
-	UE_LOG(LogTextStoreACP, Verbose, TEXT("InsertEmbeddedAtSelection"));
+	UE_LOG_CLINKAGE(LogTextStoreACP, Verbose, TEXT("InsertEmbeddedAtSelection"));
 
 	if(!IsFlaggedReadLocked(LockManager.LockType))
 	{
@@ -756,7 +756,7 @@ STDAPI FTextStoreACP::InsertEmbeddedAtSelection(DWORD dwFlags, __RPC__in_opt IDa
 
 STDAPI FTextStoreACP::OnStartComposition(__RPC__in_opt ITfCompositionView *pComposition, __RPC__out BOOL *pfOk)
 {
-	UE_LOG(LogTextStoreACP, Verbose, TEXT("OnStartComposition"));
+	UE_LOG_CLINKAGE(LogTextStoreACP, Verbose, TEXT("OnStartComposition"));
 
 	// Can only handle 1 composition. This is not an error or failure, however.
 	if(!Composition.TSFCompositionView)
@@ -781,7 +781,7 @@ STDAPI FTextStoreACP::OnStartComposition(__RPC__in_opt ITfCompositionView *pComp
 
 STDAPI FTextStoreACP::OnUpdateComposition(__RPC__in_opt ITfCompositionView *pComposition, __RPC__in_opt ITfRange *pRangeNew)
 {
-	UE_LOG(LogTextStoreACP, Verbose, TEXT("OnUpdateComposition"));
+	UE_LOG_CLINKAGE(LogTextStoreACP, Verbose, TEXT("OnUpdateComposition"));
 
 	// Can not update without an active composition.
 	if(!Composition.TSFCompositionView)
@@ -809,7 +809,7 @@ STDAPI FTextStoreACP::OnUpdateComposition(__RPC__in_opt ITfCompositionView *pCom
 			return E_FAIL;
 		}
 
-		UE_LOG(LogTextStoreACP, Verbose, TEXT("OnUpdateComposition - From %d to %d"), BeginIndex, BeginIndex + Length);
+		UE_LOG_CLINKAGE(LogTextStoreACP, Verbose, TEXT("OnUpdateComposition - From %d to %d"), BeginIndex, BeginIndex + Length);
 
 		TextContext->UpdateCompositionRange(BeginIndex, Length);
 	}
@@ -819,7 +819,7 @@ STDAPI FTextStoreACP::OnUpdateComposition(__RPC__in_opt ITfCompositionView *pCom
 
 STDAPI FTextStoreACP::OnEndComposition(__RPC__in_opt ITfCompositionView *pComposition)
 {
-	UE_LOG(LogTextStoreACP, Verbose, TEXT("OnEndComposition"));
+	UE_LOG_CLINKAGE(LogTextStoreACP, Verbose, TEXT("OnEndComposition"));
 
 	// Can not update without an active composition.
 	if(!Composition.TSFCompositionView)
diff --git a/Engine/Source/Runtime/AssetRegistry/Private/AssetRegistry.cpp b/Engine/Source/Runtime/AssetRegistry/Private/AssetRegistry.cpp
--- a/Engine/Source/Runtime/AssetRegistry/Private/AssetRegistry.cpp
+++ b/Engine/Source/Runtime/AssetRegistry/Private/AssetRegistry.cpp
@@ -21,6 +21,7 @@
 #include "UObject/ConstructorHelpers.h"
 #include "Misc/RedirectCollector.h"
 #include "AssetRegistryModule.h"
+#include "Serialization/LargeMemoryReader.h"
 
 #if WITH_EDITOR
 #include "IDirectoryWatcher.h"
@@ -56,6 +57,9 @@ UAssetRegistry::UAssetRegistry(const FObjectInitializer& ObjectInitializer)
 {
 }
 
+#include "Misc/PreLoadFile.h"
+static FPreLoadFile GPreLoadAssetRegistry(TEXT("{PROJECT}AssetRegistry.bin"));
+
 UAssetRegistryImpl::UAssetRegistryImpl(const FObjectInitializer& ObjectInitializer)
 	: Super(ObjectInitializer)
 {
@@ -91,19 +95,29 @@ UAssetRegistryImpl::UAssetRegistryImpl(const FObjectInitializer& ObjectInitializ
 	else if (FPlatformProperties::RequiresCookedData())
 	{
 		// load the cooked data
-		FArrayReader SerializedAssetData;
+		int64 Size;
+		void* PreloadedData = GPreLoadAssetRegistry.TakeOwnershipOfLoadedData(&Size);
 
-		FString AssetRegistryFilename = (FPaths::ProjectDir() / TEXT("AssetRegistry.bin"));
-		if (SerializationOptions.bSerializeAssetRegistry && IFileManager::Get().FileExists(*AssetRegistryFilename) && FFileHelper::LoadFileToArray(SerializedAssetData, *AssetRegistryFilename))
+		if (SerializationOptions.bSerializeAssetRegistry)
 		{
-			// serialize the data with the memory reader (will convert FStrings to FNames, etc)
-			Serialize(SerializedAssetData);
+			if (PreloadedData != nullptr)
+			{
+				FLargeMemoryReader SerializedAssetData((uint8*)PreloadedData, Size, ELargeMemoryReaderFlags::TakeOwnership);
+				// serialize the data with the memory reader (will convert FStrings to FNames, etc)
+				Serialize(SerializedAssetData);
+			}
 		}
+		else
+		{
+			FMemory::Free(PreloadedData);
+		}
+
 		TArray<TSharedRef<IPlugin>> ContentPlugins = IPluginManager::Get().GetEnabledPluginsWithContent();
 		for (TSharedRef<IPlugin> ContentPlugin : ContentPlugins)
 		{
 			if (ContentPlugin->CanContainContent())
 			{
+				FArrayReader SerializedAssetData;
 				FString PluginAssetRegistry = ContentPlugin->GetBaseDir() / TEXT("AssetRegistry.bin");
 				if (IFileManager::Get().FileExists(*PluginAssetRegistry) && FFileHelper::LoadFileToArray(SerializedAssetData, *PluginAssetRegistry))
 				{
diff --git a/Engine/Source/Runtime/Core/Private/Async/TaskGraph.cpp b/Engine/Source/Runtime/Core/Private/Async/TaskGraph.cpp
--- a/Engine/Source/Runtime/Core/Private/Async/TaskGraph.cpp
+++ b/Engine/Source/Runtime/Core/Private/Async/TaskGraph.cpp
@@ -440,9 +440,10 @@ public:
 	virtual void ProcessTasksUntilQuit(int32 QueueIndex) = 0;
 
 	/** Used for named threads to start processing tasks until the thread is idle and RequestQuit has been called. **/
-	virtual void ProcessTasksUntilIdle(int32 QueueIndex)
+	virtual uint64 ProcessTasksUntilIdle(int32 QueueIndex)
 	{
 		check(0);
+		return 0;
 	}
 
 	/** 
@@ -583,19 +584,21 @@ public:
 		verify(!--Queue(QueueIndex).RecursionGuard);
 	}
 
-	virtual void ProcessTasksUntilIdle(int32 QueueIndex) override
+	virtual uint64 ProcessTasksUntilIdle(int32 QueueIndex) override
 	{
 		check(Queue(QueueIndex).StallRestartEvent); // make sure we are started up
 
 		Queue(QueueIndex).QuitForReturn = false;
 		verify(++Queue(QueueIndex).RecursionGuard == 1);
-		ProcessTasksNamedThread(QueueIndex, false);
+		uint64 ProcessedTasks = ProcessTasksNamedThread(QueueIndex, false);
 		verify(!--Queue(QueueIndex).RecursionGuard);
+		return ProcessedTasks;
 	}
 
 
-	void ProcessTasksNamedThread(int32 QueueIndex, bool bAllowStall)
+	uint64 ProcessTasksNamedThread(int32 QueueIndex, bool bAllowStall)
 	{
+		uint64 ProcessedTasks = 0;
 #if UE_EXTERNAL_PROFILING_ENABLED
 		static thread_local bool bOnce = false;
 		if (!bOnce)
@@ -662,7 +665,7 @@ public:
 						Queue(QueueIndex).StallRestartEvent->Wait(MAX_uint32, bCountAsStall);
 						if (Queue(QueueIndex).QuitForShutdown)
 						{
-							return;
+							return ProcessedTasks;
 						}
 						TestRandomizedThreads();
 					}
@@ -683,6 +686,7 @@ public:
 			else
 			{
 				Task->Execute(NewTasks, ENamedThreads::Type(ThreadId | (QueueIndex << ENamedThreads::QueueIndexShift)));
+				ProcessedTasks++;
 				TestRandomizedThreads();
 			}
 		}
@@ -693,6 +697,7 @@ public:
 			bTasksOpen = false;
 		}
 #endif
+		return ProcessedTasks;
 	}
 	virtual void EnqueueFromThisThread(int32 QueueIndex, FBaseGraphTask* Task) override
 	{
@@ -854,15 +859,16 @@ public:
 		} while (!Queue.QuitForShutdown && FPlatformProcess::SupportsMultithreading()); // @Hack - quit now when running with only one thread.
 	}
 
-	virtual void ProcessTasksUntilIdle(int32 QueueIndex) override
+	virtual uint64 ProcessTasksUntilIdle(int32 QueueIndex) override
 	{
 		if (!FPlatformProcess::SupportsMultithreading())
 		{
-			ProcessTasks();
+			return ProcessTasks();
 		}
 		else
 		{
 			check(0);
+			return 0;
 		}
 	}
 
@@ -961,12 +967,13 @@ private:
 	*	@param QueueIndex, Queue to process tasks from
 	*	@param bAllowStall,  if true, the thread will block on the stall event when it runs out of tasks.
 	**/
-	void ProcessTasks()
+	uint64 ProcessTasks()
 	{
 		LLM_SCOPE(ELLMTag::TaskGraphTasksMisc);
 
 		TStatId StallStatId;
 		bool bCountAsStall = true;
+		uint64 ProcessedTasks = 0;
 #if STATS
 		TStatId StatName;
 		FCycleCounter ProcessingTasks;
@@ -1026,6 +1033,7 @@ private:
 #endif
 			bDidStall = false;
 			Task->Execute(NewTasks, ENamedThreads::Type(ThreadId));
+			ProcessedTasks++;
 			TestRandomizedThreads();
 			if (Queue.bStallForTuning)
 			{
@@ -1049,6 +1057,7 @@ private:
 			}
 		}
 		verify(!--Queue.RecursionGuard);
+		return ProcessedTasks;
 	}
 
 	/** Grouping of the data for an individual queue. **/
@@ -1395,13 +1404,14 @@ public:
 		Thread(CurrentThread).InitializeForCurrentThread();
 	}
 
-	virtual void ProcessThreadUntilIdle(ENamedThreads::Type CurrentThread) final override
+	virtual uint64 ProcessThreadUntilIdle(ENamedThreads::Type CurrentThread) final override
 	{
+		SCOPED_NAMED_EVENT(ProcessThreadUntilIdle, FColor::Red);
 		int32 QueueIndex = ENamedThreads::GetQueueIndex(CurrentThread);
 		CurrentThread = ENamedThreads::GetThreadIndex(CurrentThread);
 		check(CurrentThread >= 0 && CurrentThread < NumNamedThreads);
 		check(CurrentThread == GetCurrentThread());
-		Thread(CurrentThread).ProcessTasksUntilIdle(QueueIndex);
+		return Thread(CurrentThread).ProcessTasksUntilIdle(QueueIndex);
 	}
 
 	virtual void ProcessThreadUntilRequestReturn(ENamedThreads::Type CurrentThread) final override
@@ -1423,6 +1433,7 @@ public:
 
 	virtual void WaitUntilTasksComplete(const FGraphEventArray& Tasks, ENamedThreads::Type CurrentThreadIfKnown = ENamedThreads::AnyThread) final override
 	{
+		TRACE_CPUPROFILER_EVENT_SCOPE(WaitUntilTasksComplete);
 		ENamedThreads::Type CurrentThread = CurrentThreadIfKnown;
 		if (ENamedThreads::GetThreadIndex(CurrentThreadIfKnown) == ENamedThreads::AnyThread)
 		{
diff --git a/Engine/Source/Runtime/Core/Private/Containers/MemoryImageAllocator.cpp b/Engine/Source/Runtime/Core/Private/Containers/MemoryImageAllocator.cpp
new file mode 100644
--- /dev/null
+++ b/Engine/Source/Runtime/Core/Private/Containers/MemoryImageAllocator.cpp
@@ -0,0 +1,82 @@
+#include "Containers/ContainerAllocationPolicies.h"
+#include "Containers/ResourceArray.h"
+#include "Serialization/MemoryImage.h"
+
+IMPLEMENT_ABSTRACT_TYPE_LAYOUT(FResourceArrayInterface);
+
+FMemoryImageAllocatorBase::~FMemoryImageAllocatorBase()
+{
+	if (!Data.IsFrozen())
+	{
+		FScriptContainerElement* LocalData = Data.Get();
+		if (LocalData)
+		{
+			FMemory::Free(LocalData);
+		}
+	}
+}
+
+void FMemoryImageAllocatorBase::MoveToEmpty(FMemoryImageAllocatorBase& Other)
+{
+	checkSlow(this != &Other);
+	if (!Data.IsFrozen())
+	{
+		FScriptContainerElement* LocalData = Data.Get();
+		if (LocalData)
+		{
+			FMemory::Free(LocalData);
+		}
+	}
+	Data = Other.Data;
+	Other.Data = nullptr;
+}
+
+void FMemoryImageAllocatorBase::ResizeAllocation(int32 PreviousNumElements, int32 NumElements, SIZE_T NumBytesPerElement, uint32 Alignment)
+{
+	FScriptContainerElement* LocalData = Data.Get();
+	// Avoid calling FMemory::Realloc( nullptr, 0 ) as ANSI C mandates returning a valid pointer which is not what we want.
+	if (Data.IsFrozen())
+	{
+		// Can't grow a frozen array
+		check(NumElements <= PreviousNumElements);
+	}
+	else if (LocalData || NumElements > 0)
+	{
+		Data = (FScriptContainerElement*)FMemory::Realloc(LocalData, NumElements*NumBytesPerElement, Alignment);
+	}
+}
+
+void FMemoryImageAllocatorBase::WriteMemoryImage(FMemoryImageWriter& Writer, const FTypeLayoutDesc& TypeDesc, int32 NumAllocatedElements, uint32 Alignment) const
+{
+	if (NumAllocatedElements > 0)
+	{
+		const void* RawPtr = GetAllocation();
+		check(RawPtr);
+		FMemoryImageWriter ArrayWriter = Writer.WritePointer(FString::Printf(TEXT("FMemoryImageAllocator<%s>"), TypeDesc.Name));
+		ArrayWriter.AddDependency(TypeDesc);
+		ArrayWriter.WriteAlignment(Alignment);
+		ArrayWriter.WriteObjectArray(RawPtr, TypeDesc, NumAllocatedElements);
+	}
+	else
+	{
+		Writer.WriteMemoryImagePointerSizedBytes(0u);
+	}
+}
+
+void FMemoryImageAllocatorBase::CopyUnfrozen(const FMemoryUnfreezeContent& Context, const FTypeLayoutDesc& TypeDesc, int32 NumAllocatedElements, void* OutDst) const
+{
+	if (NumAllocatedElements > 0)
+	{
+		const void* RawPtr = GetAllocation();
+		FTypeLayoutDesc::FUnfrozenCopyFunc* Func = TypeDesc.UnfrozenCopyFunc;
+		const uint32 ElementSize = TypeDesc.Size;
+
+		for (int32 i = 0; i < NumAllocatedElements; ++i)
+		{
+			Func(Context,
+				(uint8*)RawPtr + ElementSize * i,
+				TypeDesc,
+				(uint8*)OutDst + ElementSize * i);
+		}
+	}
+}
diff --git a/Engine/Source/Runtime/Core/Private/GenericPlatform/GenericPlatformIoDispatcher.cpp b/Engine/Source/Runtime/Core/Private/GenericPlatform/GenericPlatformIoDispatcher.cpp
--- a/Engine/Source/Runtime/Core/Private/GenericPlatform/GenericPlatformIoDispatcher.cpp
+++ b/Engine/Source/Runtime/Core/Private/GenericPlatform/GenericPlatformIoDispatcher.cpp
@@ -5,7 +5,6 @@
 #include "Async/AsyncFileHandle.h"
 #include "HAL/Event.h"
 #include "HAL/PlatformFilemanager.h"
-#include "Misc/ScopeLock.h"
 #include "ProfilingDebugging/CountersTrace.h"
 #include "HAL/RunnableThread.h"
 
diff --git a/Engine/Source/Runtime/Core/Private/HAL/ThreadingBase.cpp b/Engine/Source/Runtime/Core/Private/HAL/ThreadingBase.cpp
--- a/Engine/Source/Runtime/Core/Private/HAL/ThreadingBase.cpp
+++ b/Engine/Source/Runtime/Core/Private/HAL/ThreadingBase.cpp
@@ -68,7 +68,7 @@ CORE_API bool IsInParallelRenderingThread()
 {
 	if (!GRenderingThread || GIsRenderingThreadSuspended.Load(EMemoryOrder::Relaxed))
 	{
-		return FPlatformTLS::GetCurrentThreadId() == GGameThreadId;
+		return true;
 	}
 	else
 	{
diff --git a/Engine/Source/Runtime/Core/Private/Internationalization/ICUInternationalization.h b/Engine/Source/Runtime/Core/Private/Internationalization/ICUInternationalization.h
--- a/Engine/Source/Runtime/Core/Private/Internationalization/ICUInternationalization.h
+++ b/Engine/Source/Runtime/Core/Private/Internationalization/ICUInternationalization.h
@@ -111,6 +111,7 @@ private:
 	struct FICUCachedFileData
 	{
 		FICUCachedFileData(const int64 FileSize);
+		FICUCachedFileData(void* ExistingBuffer);
 		FICUCachedFileData(FICUCachedFileData&& Source);
 		~FICUCachedFileData();
 
diff --git a/Engine/Source/Runtime/Core/Private/Misc/Compression.cpp b/Engine/Source/Runtime/Core/Private/Misc/Compression.cpp
--- a/Engine/Source/Runtime/Core/Private/Misc/Compression.cpp
+++ b/Engine/Source/Runtime/Core/Private/Misc/Compression.cpp
@@ -11,6 +11,7 @@
 #include "Misc/ConfigCacheIni.h"
 #include "Misc/CompressedGrowableBuffer.h"
 #include "Misc/ICompressionFormat.h"
+#include "Compression/lz4hc.h"
 
 #include "Misc/MemoryReadStream.h"
 // #include "TargetPlatformBase.h"
@@ -414,6 +415,11 @@ int32 FCompression::CompressMemoryBound(FName FormatName, int32 UncompressedSize
 		// CompressionBound = deflateBound(&gzipstream, gzipstream.avail_in) + GzipHeaderLength;
 		UE_LOG(LogCompression, Fatal, TEXT("FCompression::CompressMemoryBound - GZip is not supported yet"));
 	}
+	else if (FormatName == NAME_LZ4)
+	{
+		// hardcoded lz4
+		CompressionBound = LZ4_compressBound(UncompressedSize);
+	}
 	else
 	{
 		ICompressionFormat* Format = GetCompressionFormat(FormatName);
@@ -445,6 +451,12 @@ bool FCompression::CompressMemory(FName FormatName, void* CompressedBuffer, int3
 		// hardcoded gzip
 		bCompressSucceeded = appCompressMemoryGZIP(CompressedBuffer, CompressedSize, UncompressedBuffer, UncompressedSize);
 	}
+	else if (FormatName == NAME_LZ4)
+	{
+		// hardcoded lz4
+		CompressedSize = LZ4_compress_HC((const char*)UncompressedBuffer, (char*)CompressedBuffer, UncompressedSize, CompressedSize, LZ4HC_CLEVEL_MAX);
+		bCompressSucceeded = CompressedSize > 0;
+	}
 	else
 	{
 		// let the format module compress it
@@ -516,6 +528,11 @@ bool FCompression::UncompressMemory(FName FormatName, void* UncompressedBuffer,
 	{
 		// @todo buh?
 	}
+	else if (FormatName == NAME_LZ4)
+	{
+		// hardcoded lz4
+		bUncompressSucceeded = LZ4_decompress_safe((const char*)CompressedBuffer, (char*)UncompressedBuffer, CompressedSize, UncompressedSize) > 0;
+	}
 	else
 	{
 		// let the format module compress it
diff --git a/Engine/Source/Runtime/Core/Private/Misc/ConfigCacheIni.cpp b/Engine/Source/Runtime/Core/Private/Misc/ConfigCacheIni.cpp
--- a/Engine/Source/Runtime/Core/Private/Misc/ConfigCacheIni.cpp
+++ b/Engine/Source/Runtime/Core/Private/Misc/ConfigCacheIni.cpp
@@ -3897,7 +3897,9 @@ void FConfigCacheIni::InitializeConfigSystem()
 	FConfigCacheIni::LoadGlobalIniFile(GRuntimeOptionsIni, TEXT("RuntimeOptions"));
 	// Load install bundle config
 	FConfigCacheIni::LoadGlobalIniFile(GInstallBundleIni, TEXT("InstallBundle"));
-	
+	// Load device profiles for current platform
+	FConfigCacheIni::LoadGlobalIniFile(GDeviceProfilesIni, TEXT("DeviceProfiles"));
+
 	// Load user game settings .ini, allowing merging. This also updates the user .ini if necessary.
 #if PLATFORM_PS4
 	FConfigCacheIni::LoadGlobalIniFile(GGameUserSettingsIni, TEXT("GameUserSettings"), nullptr, false, false, true, *GetGameUserSettingsDir());
@@ -3983,7 +3985,7 @@ bool FConfigCacheIni::LoadLocalIniFile(FConfigFile& ConfigFile, const TCHAR* Ini
 	if (bIsBaseIniName)
 	{
 		// If base ini, try to use an existing GConfig file to set the config directories instead of assuming defaults
-		FConfigFile* BaseConfig = GConfig->FindConfigFileWithBaseName(IniName);
+		FConfigFile* BaseConfig = GConfig ? GConfig->FindConfigFileWithBaseName(IniName) : nullptr;
 		if (BaseConfig)
 		{
 			if (BaseConfig->SourceEngineConfigDir.Len())
diff --git a/Engine/Source/Runtime/Core/Private/Misc/CoreGlobals.cpp b/Engine/Source/Runtime/Core/Private/Misc/CoreGlobals.cpp
--- a/Engine/Source/Runtime/Core/Private/Misc/CoreGlobals.cpp
+++ b/Engine/Source/Runtime/Core/Private/Misc/CoreGlobals.cpp
@@ -222,6 +222,7 @@ FString				GGameIni;													/* Game ini filename */
 FString				GGameUserSettingsIni;										/* User Game Settings ini filename */
 FString				GRuntimeOptionsIni;											/* Runtime Options ini filename */
 FString				GInstallBundleIni;											/* Install Bundle ini filename*/
+FString				GDeviceProfilesIni;											/* Runtime DeviceProfiles ini filename - use LoadLocalIni for other platforms' DPs */
 
 float					GNearClippingPlane				= 10.0f;				/* Near clipping plane */
 
diff --git a/Engine/Source/Runtime/Core/Private/Misc/CoreMisc.cpp b/Engine/Source/Runtime/Core/Private/Misc/CoreMisc.cpp
--- a/Engine/Source/Runtime/Core/Private/Misc/CoreMisc.cpp
+++ b/Engine/Source/Runtime/Core/Private/Misc/CoreMisc.cpp
@@ -12,6 +12,7 @@
 #include "Misc/OutputDeviceError.h"
 #include "Misc/ScopeLock.h"
 #include "CoreGlobals.h"
+#include "Templates/RefCounting.h"
 
 /** For FConfigFile in appInit							*/
 #include "Misc/ConfigCacheIni.h"
diff --git a/Engine/Source/Runtime/Core/Private/Misc/DataDrivenPlatformInfoRegistry.cpp b/Engine/Source/Runtime/Core/Private/Misc/DataDrivenPlatformInfoRegistry.cpp
--- a/Engine/Source/Runtime/Core/Private/Misc/DataDrivenPlatformInfoRegistry.cpp
+++ b/Engine/Source/Runtime/Core/Private/Misc/DataDrivenPlatformInfoRegistry.cpp
@@ -4,6 +4,7 @@
 #include "Misc/Paths.h"
 #include "HAL/FileManager.h"
 #include "Misc/FileHelper.h"
+#include "Misc/ConfigCacheIni.h"
 
 
 
@@ -61,6 +62,114 @@ bool FDataDrivenPlatformInfoRegistry::LoadDataDrivenIniFile(int32 Index, FConfig
 	return false;
 }
 
+static void DDPIIniRedirect(FString& StringData)
+{
+	TArray<FString> Tokens;
+	StringData.ParseIntoArray(Tokens, TEXT(":"));
+	if (Tokens.Num() != 5)
+	{
+		StringData = TEXT("");
+		return;
+	}
+
+	// now load a local version of the ini hierarchy
+	FConfigFile LocalIni;
+	FConfigCacheIni::LoadLocalIniFile(LocalIni, *Tokens[1], true, *Tokens[2]);
+
+	// and get the platform's value (if it's not found, return an empty string)
+	FString FoundValue;
+	LocalIni.GetString(*Tokens[3], *Tokens[4], FoundValue);
+	StringData = FoundValue;
+}
+
+static FString DDPITryRedirect(const FConfigFile& IniFile, const TCHAR* Key, bool* OutHadBang=nullptr)
+{
+	FString StringData;
+	if (IniFile.GetString(TEXT("DataDrivenPlatformInfo"), Key, StringData))
+	{
+		if (StringData.StartsWith(TEXT("ini:")) || StringData.StartsWith(TEXT("!ini:")))
+		{
+			// check for !'ing a bool
+			if (OutHadBang != nullptr)
+			{
+				*OutHadBang = StringData[0] == TEXT('!');
+			}
+
+			// replace the string, overwriting it
+			DDPIIniRedirect(StringData);
+		}
+	}
+	return StringData;
+}
+
+static void DDPIGetBool(const FConfigFile& IniFile, const TCHAR* Key, bool& OutBool)
+{
+	bool bHadNot = false;
+	FString StringData = DDPITryRedirect(IniFile, Key, &bHadNot);
+
+	// if we ended up with a string, convert it, otherwise leave it alone
+	if (StringData.Len() > 0)
+	{
+		OutBool = bHadNot ? !StringData.ToBool() : StringData.ToBool();
+	}
+}
+
+static void DDPIGetInt(const FConfigFile& IniFile, const TCHAR* Key, int32& OutInt)
+{
+	FString StringData = DDPITryRedirect(IniFile, Key);
+
+	// if we ended up with a string, convert it, otherwise leave it alone
+	if (StringData.Len() > 0)
+	{
+		OutInt = FCString::Atoi(*StringData);
+	}
+}
+
+static void DDPIGetUInt(const FConfigFile& IniFile, const TCHAR* Key, uint32& OutInt)
+{
+	FString StringData = DDPITryRedirect(IniFile, Key);
+
+	// if we ended up with a string, convert it, otherwise leave it alone
+	if (StringData.Len() > 0)
+	{
+		OutInt = (uint32)FCString::Strtoui64(*StringData, nullptr, 10);
+	}
+}
+
+static void DDPIGetString(const FConfigFile& IniFile, const TCHAR* Key, FString& OutString)
+{
+	FString StringData = DDPITryRedirect(IniFile, Key);
+
+	// if we ended up with a string, convert it, otherwise leave it alone
+	if (StringData.Len() > 0)
+	{
+		OutString = StringData;
+	}
+}
+
+static void DDPIGetStringArray(const FConfigFile& IniFile, const TCHAR* Key, TArray<FString>& OutArray)
+{
+	// we don't support redirecting arrays
+	IniFile.GetArray(TEXT("DataDrivenPlatformInfo"), Key, OutArray);
+}
+
+static void LoadDDPIIniSettings(const FConfigFile& IniFile, FDataDrivenPlatformInfoRegistry::FPlatformInfo& Info)
+{
+	DDPIGetBool(IniFile, TEXT("bIsConfidential"), Info.bIsConfidential);
+	DDPIGetBool(IniFile, TEXT("bRestrictLocalization"), Info.bRestrictLocalization);
+	DDPIGetString(IniFile, TEXT("AudioCompressionSettingsIniSectionName"), Info.AudioCompressionSettingsIniSectionName);
+	DDPIGetStringArray(IniFile, TEXT("AdditionalRestrictedFolders"), Info.AdditionalRestrictedFolders);
+	
+	DDPIGetBool(IniFile, TEXT("Freezing_b32Bit"), Info.Freezing_b32Bit);
+	DDPIGetUInt(IniFile, Info.Freezing_b32Bit ? TEXT("Freezing_MaxFieldAlignment32") : TEXT("Freezing_MaxFieldAlignment64"), Info.Freezing_MaxFieldAlignment);
+	DDPIGetBool(IniFile, TEXT("Freezing_bForce64BitMemoryImagePointers"), Info.Freezing_bForce64BitMemoryImagePointers);
+	DDPIGetBool(IniFile, TEXT("Freezing_bAlignBases"), Info.Freezing_bAlignBases);
+	DDPIGetBool(IniFile, TEXT("Freezing_bWithRayTracing"), Info.Freezing_bWithRayTracing);
+
+	// NOTE: add more settings here!
+}
+
+
 /**
 * Get the global set of data driven platform information
 */
@@ -89,9 +198,7 @@ const TMap<FString, FDataDrivenPlatformInfoRegistry::FPlatformInfo>& FDataDriven
 			{
 				// cache info
 				FDataDrivenPlatformInfoRegistry::FPlatformInfo& Info = DataDrivenPlatforms.Add(PlatformName, FDataDrivenPlatformInfoRegistry::FPlatformInfo());
-				IniFile.GetBool(TEXT("DataDrivenPlatformInfo"), TEXT("bIsConfidential"), Info.bIsConfidential);
-				IniFile.GetBool(TEXT("DataDrivenPlatformInfo"), TEXT("bRestrictLocalization"), Info.bRestrictLocalization);
-				IniFile.GetArray(TEXT("DataDrivenPlatformInfo"), TEXT("AdditionalRestrictedFolders"), Info.AdditionalRestrictedFolders);
+				LoadDDPIIniSettings(IniFile, Info);
 
 				// get the parent to build list later
 				FString IniParent;
diff --git a/Engine/Source/Runtime/Core/Private/Misc/DelayedAutoRegister.cpp b/Engine/Source/Runtime/Core/Private/Misc/DelayedAutoRegister.cpp
new file mode 100644
--- /dev/null
+++ b/Engine/Source/Runtime/Core/Private/Misc/DelayedAutoRegister.cpp
@@ -0,0 +1,38 @@
+// Copyright 1998-2019 Epic Games, Inc. All Rights Reserved.
+
+#include "Misc/DelayedAutoRegister.h"
+#include "Delegates/Delegate.h"
+
+DECLARE_MULTICAST_DELEGATE(FDelayedAutoRegisterDelegate);
+
+static TSet<EDelayedRegisterRunPhase> GPhasesAlreadyRun;
+
+static FDelayedAutoRegisterDelegate& GetDelayedAutoRegisterDelegate(EDelayedRegisterRunPhase Phase)
+{
+	static FDelayedAutoRegisterDelegate Singleton[(uint8)EDelayedRegisterRunPhase::NumPhases];
+	return Singleton[(uint8)Phase];
+}
+
+FDelayedAutoRegisterHelper::FDelayedAutoRegisterHelper(EDelayedRegisterRunPhase Phase, TFunction<void()> RegistrationFunction)
+{
+	// if the phase has already passed, we just run the function immediately
+	if (GPhasesAlreadyRun.Contains(Phase))
+	{
+		RegistrationFunction();
+	}
+	else
+	{
+		GetDelayedAutoRegisterDelegate(Phase).AddLambda(RegistrationFunction);
+	}
+}
+
+void FDelayedAutoRegisterHelper::RunAndClearDelayedAutoRegisterDelegates(EDelayedRegisterRunPhase Phase)
+{
+	checkf(!GPhasesAlreadyRun.Contains(Phase), TEXT("Delayed Startup phase %d has already run - it is not expected to be run again!"), (int32)Phase);
+
+	// run all the delayed functions!
+	GetDelayedAutoRegisterDelegate(Phase).Broadcast();
+	GetDelayedAutoRegisterDelegate(Phase).Clear();
+
+	GPhasesAlreadyRun.Add(Phase);
+}
diff --git a/Engine/Source/Runtime/Core/Private/Misc/SecureHash.cpp b/Engine/Source/Runtime/Core/Private/Misc/SecureHash.cpp
--- a/Engine/Source/Runtime/Core/Private/Misc/SecureHash.cpp
+++ b/Engine/Source/Runtime/Core/Private/Misc/SecureHash.cpp
@@ -16,22 +16,22 @@ DEFINE_LOG_CATEGORY_STATIC(LogSecureHash, Log, All);
 // Constants for MD5 Transform.
 //
 
-enum {S11=7};
-enum {S12=12};
-enum {S13=17};
-enum {S14=22};
-enum {S21=5};
-enum {S22=9};
-enum {S23=14};
-enum {S24=20};
-enum {S31=4};
-enum {S32=11};
-enum {S33=16};
-enum {S34=23};
-enum {S41=6};
-enum {S42=10};
-enum {S43=15};
-enum {S44=21};
+enum {_S11=7};
+enum {_S12=12};
+enum {_S13=17};
+enum {_S14=22};
+enum {_S21=5};
+enum {_S22=9};
+enum {_S23=14};
+enum {_S24=20};
+enum {_S31=4};
+enum {_S32=11};
+enum {_S33=16};
+enum {_S34=23};
+enum {_S41=6};
+enum {_S42=10};
+enum {_S43=15};
+enum {_S44=21};
 
 static uint8 PADDING[64] = {
 	0x80, 0, 0, 0, 0, 0, 0, 0, 0, 0, 0, 0, 0, 0, 0, 0, 0, 0, 0,
@@ -161,76 +161,76 @@ void FMD5::Transform( uint32* state, const uint8* block )
 	Decode( x, block, 64 );
 
 	// Round 1
-	MD5_FF (a, b, c, d, x[ 0], S11, 0xd76aa478); /* 1 */
-	MD5_FF (d, a, b, c, x[ 1], S12, 0xe8c7b756); /* 2 */
-	MD5_FF (c, d, a, b, x[ 2], S13, 0x242070db); /* 3 */
-	MD5_FF (b, c, d, a, x[ 3], S14, 0xc1bdceee); /* 4 */
-	MD5_FF (a, b, c, d, x[ 4], S11, 0xf57c0faf); /* 5 */
-	MD5_FF (d, a, b, c, x[ 5], S12, 0x4787c62a); /* 6 */
-	MD5_FF (c, d, a, b, x[ 6], S13, 0xa8304613); /* 7 */
-	MD5_FF (b, c, d, a, x[ 7], S14, 0xfd469501); /* 8 */
-	MD5_FF (a, b, c, d, x[ 8], S11, 0x698098d8); /* 9 */
-	MD5_FF (d, a, b, c, x[ 9], S12, 0x8b44f7af); /* 10 */
-	MD5_FF (c, d, a, b, x[10], S13, 0xffff5bb1); /* 11 */
-	MD5_FF (b, c, d, a, x[11], S14, 0x895cd7be); /* 12 */
-	MD5_FF (a, b, c, d, x[12], S11, 0x6b901122); /* 13 */
-	MD5_FF (d, a, b, c, x[13], S12, 0xfd987193); /* 14 */
-	MD5_FF (c, d, a, b, x[14], S13, 0xa679438e); /* 15 */
-	MD5_FF (b, c, d, a, x[15], S14, 0x49b40821); /* 16 */
+	MD5_FF (a, b, c, d, x[ 0], _S11, 0xd76aa478); /* 1 */
+	MD5_FF (d, a, b, c, x[ 1], _S12, 0xe8c7b756); /* 2 */
+	MD5_FF (c, d, a, b, x[ 2], _S13, 0x242070db); /* 3 */
+	MD5_FF (b, c, d, a, x[ 3], _S14, 0xc1bdceee); /* 4 */
+	MD5_FF (a, b, c, d, x[ 4], _S11, 0xf57c0faf); /* 5 */
+	MD5_FF (d, a, b, c, x[ 5], _S12, 0x4787c62a); /* 6 */
+	MD5_FF (c, d, a, b, x[ 6], _S13, 0xa8304613); /* 7 */
+	MD5_FF (b, c, d, a, x[ 7], _S14, 0xfd469501); /* 8 */
+	MD5_FF (a, b, c, d, x[ 8], _S11, 0x698098d8); /* 9 */
+	MD5_FF (d, a, b, c, x[ 9], _S12, 0x8b44f7af); /* 10 */
+	MD5_FF (c, d, a, b, x[10], _S13, 0xffff5bb1); /* 11 */
+	MD5_FF (b, c, d, a, x[11], _S14, 0x895cd7be); /* 12 */
+	MD5_FF (a, b, c, d, x[12], _S11, 0x6b901122); /* 13 */
+	MD5_FF (d, a, b, c, x[13], _S12, 0xfd987193); /* 14 */
+	MD5_FF (c, d, a, b, x[14], _S13, 0xa679438e); /* 15 */
+	MD5_FF (b, c, d, a, x[15], _S14, 0x49b40821); /* 16 */
 
 	// Round 2
-	MD5_GG (a, b, c, d, x[ 1], S21, 0xf61e2562); /* 17 */
-	MD5_GG (d, a, b, c, x[ 6], S22, 0xc040b340); /* 18 */
-	MD5_GG (c, d, a, b, x[11], S23, 0x265e5a51); /* 19 */
-	MD5_GG (b, c, d, a, x[ 0], S24, 0xe9b6c7aa); /* 20 */
-	MD5_GG (a, b, c, d, x[ 5], S21, 0xd62f105d); /* 21 */
-	MD5_GG (d, a, b, c, x[10], S22,  0x2441453); /* 22 */
-	MD5_GG (c, d, a, b, x[15], S23, 0xd8a1e681); /* 23 */
-	MD5_GG (b, c, d, a, x[ 4], S24, 0xe7d3fbc8); /* 24 */
-	MD5_GG (a, b, c, d, x[ 9], S21, 0x21e1cde6); /* 25 */
-	MD5_GG (d, a, b, c, x[14], S22, 0xc33707d6); /* 26 */
-	MD5_GG (c, d, a, b, x[ 3], S23, 0xf4d50d87); /* 27 */
-	MD5_GG (b, c, d, a, x[ 8], S24, 0x455a14ed); /* 28 */
-	MD5_GG (a, b, c, d, x[13], S21, 0xa9e3e905); /* 29 */
-	MD5_GG (d, a, b, c, x[ 2], S22, 0xfcefa3f8); /* 30 */
-	MD5_GG (c, d, a, b, x[ 7], S23, 0x676f02d9); /* 31 */
-	MD5_GG (b, c, d, a, x[12], S24, 0x8d2a4c8a); /* 32 */
+	MD5_GG (a, b, c, d, x[ 1], _S21, 0xf61e2562); /* 17 */
+	MD5_GG (d, a, b, c, x[ 6], _S22, 0xc040b340); /* 18 */
+	MD5_GG (c, d, a, b, x[11], _S23, 0x265e5a51); /* 19 */
+	MD5_GG (b, c, d, a, x[ 0], _S24, 0xe9b6c7aa); /* 20 */
+	MD5_GG (a, b, c, d, x[ 5], _S21, 0xd62f105d); /* 21 */
+	MD5_GG (d, a, b, c, x[10], _S22,  0x2441453); /* 22 */
+	MD5_GG (c, d, a, b, x[15], _S23, 0xd8a1e681); /* 23 */
+	MD5_GG (b, c, d, a, x[ 4], _S24, 0xe7d3fbc8); /* 24 */
+	MD5_GG (a, b, c, d, x[ 9], _S21, 0x21e1cde6); /* 25 */
+	MD5_GG (d, a, b, c, x[14], _S22, 0xc33707d6); /* 26 */
+	MD5_GG (c, d, a, b, x[ 3], _S23, 0xf4d50d87); /* 27 */
+	MD5_GG (b, c, d, a, x[ 8], _S24, 0x455a14ed); /* 28 */
+	MD5_GG (a, b, c, d, x[13], _S21, 0xa9e3e905); /* 29 */
+	MD5_GG (d, a, b, c, x[ 2], _S22, 0xfcefa3f8); /* 30 */
+	MD5_GG (c, d, a, b, x[ 7], _S23, 0x676f02d9); /* 31 */
+	MD5_GG (b, c, d, a, x[12], _S24, 0x8d2a4c8a); /* 32 */
 
 	// Round 3
-	MD5_HH (a, b, c, d, x[ 5], S31, 0xfffa3942); /* 33 */
-	MD5_HH (d, a, b, c, x[ 8], S32, 0x8771f681); /* 34 */
-	MD5_HH (c, d, a, b, x[11], S33, 0x6d9d6122); /* 35 */
-	MD5_HH (b, c, d, a, x[14], S34, 0xfde5380c); /* 36 */
-	MD5_HH (a, b, c, d, x[ 1], S31, 0xa4beea44); /* 37 */
-	MD5_HH (d, a, b, c, x[ 4], S32, 0x4bdecfa9); /* 38 */
-	MD5_HH (c, d, a, b, x[ 7], S33, 0xf6bb4b60); /* 39 */
-	MD5_HH (b, c, d, a, x[10], S34, 0xbebfbc70); /* 40 */
-	MD5_HH (a, b, c, d, x[13], S31, 0x289b7ec6); /* 41 */
-	MD5_HH (d, a, b, c, x[ 0], S32, 0xeaa127fa); /* 42 */
-	MD5_HH (c, d, a, b, x[ 3], S33, 0xd4ef3085); /* 43 */
-	MD5_HH (b, c, d, a, x[ 6], S34,  0x4881d05); /* 44 */
-	MD5_HH (a, b, c, d, x[ 9], S31, 0xd9d4d039); /* 45 */
-	MD5_HH (d, a, b, c, x[12], S32, 0xe6db99e5); /* 46 */
-	MD5_HH (c, d, a, b, x[15], S33, 0x1fa27cf8); /* 47 */
-	MD5_HH (b, c, d, a, x[ 2], S34, 0xc4ac5665); /* 48 */
+	MD5_HH (a, b, c, d, x[ 5], _S31, 0xfffa3942); /* 33 */
+	MD5_HH (d, a, b, c, x[ 8], _S32, 0x8771f681); /* 34 */
+	MD5_HH (c, d, a, b, x[11], _S33, 0x6d9d6122); /* 35 */
+	MD5_HH (b, c, d, a, x[14], _S34, 0xfde5380c); /* 36 */
+	MD5_HH (a, b, c, d, x[ 1], _S31, 0xa4beea44); /* 37 */
+	MD5_HH (d, a, b, c, x[ 4], _S32, 0x4bdecfa9); /* 38 */
+	MD5_HH (c, d, a, b, x[ 7], _S33, 0xf6bb4b60); /* 39 */
+	MD5_HH (b, c, d, a, x[10], _S34, 0xbebfbc70); /* 40 */
+	MD5_HH (a, b, c, d, x[13], _S31, 0x289b7ec6); /* 41 */
+	MD5_HH (d, a, b, c, x[ 0], _S32, 0xeaa127fa); /* 42 */
+	MD5_HH (c, d, a, b, x[ 3], _S33, 0xd4ef3085); /* 43 */
+	MD5_HH (b, c, d, a, x[ 6], _S34,  0x4881d05); /* 44 */
+	MD5_HH (a, b, c, d, x[ 9], _S31, 0xd9d4d039); /* 45 */
+	MD5_HH (d, a, b, c, x[12], _S32, 0xe6db99e5); /* 46 */
+	MD5_HH (c, d, a, b, x[15], _S33, 0x1fa27cf8); /* 47 */
+	MD5_HH (b, c, d, a, x[ 2], _S34, 0xc4ac5665); /* 48 */
 
 	// Round 4
-	MD5_II (a, b, c, d, x[ 0], S41, 0xf4292244); /* 49 */
-	MD5_II (d, a, b, c, x[ 7], S42, 0x432aff97); /* 50 */
-	MD5_II (c, d, a, b, x[14], S43, 0xab9423a7); /* 51 */
-	MD5_II (b, c, d, a, x[ 5], S44, 0xfc93a039); /* 52 */
-	MD5_II (a, b, c, d, x[12], S41, 0x655b59c3); /* 53 */
-	MD5_II (d, a, b, c, x[ 3], S42, 0x8f0ccc92); /* 54 */
-	MD5_II (c, d, a, b, x[10], S43, 0xffeff47d); /* 55 */
-	MD5_II (b, c, d, a, x[ 1], S44, 0x85845dd1); /* 56 */
-	MD5_II (a, b, c, d, x[ 8], S41, 0x6fa87e4f); /* 57 */
-	MD5_II (d, a, b, c, x[15], S42, 0xfe2ce6e0); /* 58 */
-	MD5_II (c, d, a, b, x[ 6], S43, 0xa3014314); /* 59 */
-	MD5_II (b, c, d, a, x[13], S44, 0x4e0811a1); /* 60 */
-	MD5_II (a, b, c, d, x[ 4], S41, 0xf7537e82); /* 61 */
-	MD5_II (d, a, b, c, x[11], S42, 0xbd3af235); /* 62 */
-	MD5_II (c, d, a, b, x[ 2], S43, 0x2ad7d2bb); /* 63 */
-	MD5_II (b, c, d, a, x[ 9], S44, 0xeb86d391); /* 64 */
+	MD5_II (a, b, c, d, x[ 0], _S41, 0xf4292244); /* 49 */
+	MD5_II (d, a, b, c, x[ 7], _S42, 0x432aff97); /* 50 */
+	MD5_II (c, d, a, b, x[14], _S43, 0xab9423a7); /* 51 */
+	MD5_II (b, c, d, a, x[ 5], _S44, 0xfc93a039); /* 52 */
+	MD5_II (a, b, c, d, x[12], _S41, 0x655b59c3); /* 53 */
+	MD5_II (d, a, b, c, x[ 3], _S42, 0x8f0ccc92); /* 54 */
+	MD5_II (c, d, a, b, x[10], _S43, 0xffeff47d); /* 55 */
+	MD5_II (b, c, d, a, x[ 1], _S44, 0x85845dd1); /* 56 */
+	MD5_II (a, b, c, d, x[ 8], _S41, 0x6fa87e4f); /* 57 */
+	MD5_II (d, a, b, c, x[15], _S42, 0xfe2ce6e0); /* 58 */
+	MD5_II (c, d, a, b, x[ 6], _S43, 0xa3014314); /* 59 */
+	MD5_II (b, c, d, a, x[13], _S44, 0x4e0811a1); /* 60 */
+	MD5_II (a, b, c, d, x[ 4], _S41, 0xf7537e82); /* 61 */
+	MD5_II (d, a, b, c, x[11], _S42, 0xbd3af235); /* 62 */
+	MD5_II (c, d, a, b, x[ 2], _S43, 0x2ad7d2bb); /* 63 */
+	MD5_II (b, c, d, a, x[ 9], _S44, 0xeb86d391); /* 64 */
 
 	state[0] += a;
 	state[1] += b;
diff --git a/Engine/Source/Runtime/Core/Private/UObject/UnrealNames.cpp b/Engine/Source/Runtime/Core/Private/UObject/UnrealNames.cpp
--- a/Engine/Source/Runtime/Core/Private/UObject/UnrealNames.cpp
+++ b/Engine/Source/Runtime/Core/Private/UObject/UnrealNames.cpp
@@ -24,6 +24,7 @@
 #include "Misc/OutputDeviceRedirector.h"
 #include "HAL/IConsoleManager.h"
 #include "HAL/LowLevelMemTracker.h"
+#include "Serialization/MemoryImage.h"
 #include "Hash/CityHash.h"
 #include "Templates/AlignmentTemplates.h"
 
@@ -3204,4 +3205,9 @@ uint8** FNameDebugVisualizer::GetBlocks()
 	return ((FNamePool*)(NamePoolData))->GetBlocksForDebugVisualizer();
 }
 
+void Freeze::IntrinsicWriteMemoryImage(FMemoryImageWriter& Writer, const FName& Object, const FTypeLayoutDesc&)
+{
+	Writer.WriteFName(Object);
+}
+
 PRAGMA_ENABLE_UNSAFE_TYPECAST_WARNINGS
diff --git a/Engine/Source/Runtime/Core/Public/Async/ParallelFor.h b/Engine/Source/Runtime/Core/Public/Async/ParallelFor.h
--- a/Engine/Source/Runtime/Core/Public/Async/ParallelFor.h
+++ b/Engine/Source/Runtime/Core/Public/Async/ParallelFor.h
@@ -27,7 +27,10 @@ enum class EParallelForFlags
 
 	//Offers better work distribution among threads at the cost of a little bit more synchronization.
 	//This should be used for tasks with highly variable computational time.
-	Unbalanced = 2
+	Unbalanced = 2,
+
+	// if running on the rendering thread, make sure the ProcessThread is called when idle
+	PumpRenderingThread = 4,
 };
 
 ENUM_CLASS_FLAGS(EParallelForFlags)
@@ -35,19 +38,20 @@ ENUM_CLASS_FLAGS(EParallelForFlags)
 namespace ParallelForImpl
 {
 	// struct to hold the working data; this outlives the ParallelFor call; lifetime is controlled by a shared pointer
-	struct FParallelForData
+	template<typename FunctionType>
+	struct TParallelForData
 	{
 		int32 Num;
 		int32 BlockSize;
 		int32 LastBlockExtraNum;
-		TFunctionRef<void(int32)> Body;
+		FunctionType Body;
 		FEvent* Event;
 		FThreadSafeCounter IndexToDo;
 		FThreadSafeCounter NumCompleted;
 		bool bExited;
 		bool bTriggered;
 		bool bSaveLastBlockForMaster;
-		FParallelForData(int32 InTotalNum, int32 InNumThreads, bool bInSaveLastBlockForMaster, TFunctionRef<void(int32)> InBody, EParallelForFlags Flags)
+		TParallelForData(int32 InTotalNum, int32 InNumThreads, bool bInSaveLastBlockForMaster, FunctionType InBody, EParallelForFlags Flags)
 			: Body(InBody)
 			, Event(FPlatformProcess::GetSynchEventFromPool(false))
 			, bExited(false)
@@ -83,22 +87,23 @@ namespace ParallelForImpl
 			LastBlockExtraNum = InTotalNum - Num * BlockSize;
 			check(LastBlockExtraNum >= 0);
 		}
-		~FParallelForData()
+		~TParallelForData()
 		{
 			check(IndexToDo.GetValue() >= Num);
 			check(NumCompleted.GetValue() == Num);
 			check(bExited);
 			FPlatformProcess::ReturnSynchEventToPool(Event);
 		}
-		bool Process(int32 TasksToSpawn, TSharedRef<FParallelForData, ESPMode::ThreadSafe>& Data, bool bMaster);
+		bool Process(int32 TasksToSpawn, TSharedRef<TParallelForData, ESPMode::ThreadSafe>& Data, bool bMaster);
 	};
 
-	class FParallelForTask
+	template<typename FunctionType>
+	class TParallelForTask
 	{
-		TSharedRef<FParallelForData, ESPMode::ThreadSafe> Data;
+		TSharedRef<TParallelForData<FunctionType>, ESPMode::ThreadSafe> Data;
 		int32 TasksToSpawn;
 	public:
-		FParallelForTask(TSharedRef<FParallelForData, ESPMode::ThreadSafe>& InData, int32 InTasksToSpawn = 0)
+		TParallelForTask(TSharedRef<TParallelForData<FunctionType>, ESPMode::ThreadSafe>& InData, int32 InTasksToSpawn = 0)
 			: Data(InData) 
 			, TasksToSpawn(InTasksToSpawn)
 		{
@@ -117,6 +122,7 @@ namespace ParallelForImpl
 		}
 		void DoTask(ENamedThreads::Type CurrentThread, const FGraphEventRef& MyCompletionGraphEvent)
 		{
+			FMemMark Mark(FMemStack::Get());
 			if (Data->Process(TasksToSpawn, Data, false))
 			{
 				checkSlow(!Data->bTriggered);
@@ -126,13 +132,14 @@ namespace ParallelForImpl
 		}
 	};
 
-	inline bool FParallelForData::Process(int32 TasksToSpawn, TSharedRef<FParallelForData, ESPMode::ThreadSafe>& Data, bool bMaster)
+	template<typename FunctionType>
+	inline bool TParallelForData<FunctionType>::Process(int32 TasksToSpawn, TSharedRef<TParallelForData<FunctionType>, ESPMode::ThreadSafe>& Data, bool bMaster)
 	{
 		int32 MaybeTasksLeft = Num - IndexToDo.GetValue();
 		if (TasksToSpawn && MaybeTasksLeft > 0)
 		{
 			TasksToSpawn = FMath::Min<int32>(TasksToSpawn, MaybeTasksLeft);
-			TGraphTask<FParallelForTask>::CreateTask().ConstructAndDispatchWhenReady(Data, TasksToSpawn - 1);		
+			TGraphTask<TParallelForTask<FunctionType>>::CreateTask().ConstructAndDispatchWhenReady(Data, TasksToSpawn - 1);
 		}
 		int32 LocalBlockSize = BlockSize;
 		int32 LocalNum = Num;
@@ -179,7 +186,8 @@ namespace ParallelForImpl
 		return false;
 	}
 
-	inline void ParallelForInternal(int32 Num, TFunctionRef<void(int32)> Body, EParallelForFlags Flags)
+	template<typename FunctionType>
+	inline void ParallelForInternal(int32 Num, FunctionType Body, EParallelForFlags Flags)
 	{
 		SCOPE_CYCLE_COUNTER(STAT_ParallelFor);
 		check(Num >= 0);
@@ -198,13 +206,25 @@ namespace ParallelForImpl
 			}
 			return;
 		}
-		FParallelForData* DataPtr = new FParallelForData(Num, AnyThreadTasks + 1, Num > AnyThreadTasks + 1, Body, Flags);
-		TSharedRef<FParallelForData, ESPMode::ThreadSafe> Data = MakeShareable(DataPtr);
-		TGraphTask<FParallelForTask>::CreateTask().ConstructAndDispatchWhenReady(Data, AnyThreadTasks - 1);		
+
+		const bool bPumpRenderingThread = (Flags & EParallelForFlags::PumpRenderingThread) != EParallelForFlags::None;
+		TParallelForData<FunctionType>* DataPtr = new TParallelForData<FunctionType>(Num, AnyThreadTasks + 1, (Num > AnyThreadTasks + 1) && bPumpRenderingThread, Body, Flags);
+		TSharedRef<TParallelForData<FunctionType>, ESPMode::ThreadSafe> Data = MakeShareable(DataPtr);
+		TGraphTask<TParallelForTask<FunctionType>>::CreateTask().ConstructAndDispatchWhenReady(Data, AnyThreadTasks - 1);
 		// this thread can help too and this is important to prevent deadlock on recursion 
 		if (!Data->Process(0, Data, true))
 		{
-			Data->Event->Wait();
+			if (bPumpRenderingThread && IsInActualRenderingThread())
+			{
+				while (!Data->Event->Wait(1))
+				{
+					FTaskGraphInterface::Get().ProcessThreadUntilIdle(ENamedThreads::GetRenderThread_Local());
+				}
+			}
+			else
+			{
+				Data->Event->Wait();
+			}
 			check(Data->bTriggered);
 		}
 		else
@@ -225,7 +245,8 @@ namespace ParallelForImpl
 		*	@param Flags; Used to customize the behavior of the ParallelFor if needed.
 		*	Notes: Please add stats around to calls to parallel for and within your lambda as appropriate. Do not clog the task graph with long running tasks or tasks that block.
 	**/
-	inline void ParallelForWithPreWorkInternal(int32 Num, TFunctionRef<void(int32)> Body, TFunctionRef<void()> CurrentThreadWorkToDoBeforeHelping, EParallelForFlags Flags = EParallelForFlags::None)
+	template<typename FunctionType>
+	inline void ParallelForWithPreWorkInternal(int32 Num, FunctionType Body, TFunctionRef<void()> CurrentThreadWorkToDoBeforeHelping, EParallelForFlags Flags = EParallelForFlags::None)
 	{
 		SCOPE_CYCLE_COUNTER(STAT_ParallelFor);
 
@@ -246,15 +267,25 @@ namespace ParallelForImpl
 			return;
 		}
 		check(Num);
-		FParallelForData* DataPtr = new FParallelForData(Num, AnyThreadTasks, false, Body, Flags);
-		TSharedRef<FParallelForData, ESPMode::ThreadSafe> Data = MakeShareable(DataPtr);
-		TGraphTask<FParallelForTask>::CreateTask().ConstructAndDispatchWhenReady(Data, AnyThreadTasks - 1);		
+		TParallelForData<FunctionType>* DataPtr = new TParallelForData<FunctionType>(Num, AnyThreadTasks, false, Body, Flags);
+		TSharedRef<TParallelForData<FunctionType>, ESPMode::ThreadSafe> Data = MakeShareable(DataPtr);
+		TGraphTask<TParallelForTask<FunctionType>>::CreateTask().ConstructAndDispatchWhenReady(Data, AnyThreadTasks - 1);
 		// do the prework
 		CurrentThreadWorkToDoBeforeHelping();
 		// this thread can help too and this is important to prevent deadlock on recursion 
 		if (!Data->Process(0, Data, true))
 		{
-			Data->Event->Wait();
+			if (IsInRenderingThread() && (Flags & EParallelForFlags::PumpRenderingThread) != EParallelForFlags::None)
+			{
+				while (!Data->Event->Wait(1))
+				{
+					FTaskGraphInterface::Get().ProcessThreadUntilIdle(ENamedThreads::GetRenderThread_Local());
+				}
+			}
+			else
+			{
+				Data->Event->Wait();
+			}
 			check(Data->bTriggered);
 		}
 		else
@@ -274,11 +305,25 @@ namespace ParallelForImpl
 	*	@param bForceSingleThread; Mostly used for testing, if true, run single threaded instead.
 	*	Notes: Please add stats around to calls to parallel for and within your lambda as appropriate. Do not clog the task graph with long running tasks or tasks that block.
 **/
-inline void ParallelFor(int32 Num, TFunctionRef<void(int32)> Body, bool bForceSingleThread)
+inline void ParallelFor(int32 Num, TFunctionRef<void(int32)> Body, bool bForceSingleThread, bool bPumpRenderingThread=false)
 {
-	ParallelForImpl::ParallelForInternal(Num, Body, bForceSingleThread ? EParallelForFlags::ForceSingleThread : EParallelForFlags::None);
+	ParallelForImpl::ParallelForInternal(Num, Body,
+		(bForceSingleThread ? EParallelForFlags::ForceSingleThread : EParallelForFlags::None) | 
+		(bPumpRenderingThread ? EParallelForFlags::PumpRenderingThread : EParallelForFlags::None));
 }
 
+/**
+	*	General purpose parallel for that uses the taskgraph
+	*	@param Num; number of calls of Body; Body(0), Body(1)....Body(Num - 1)
+	*	@param Body; Function to call from multiple threads
+	*	@param bForceSingleThread; Mostly used for testing, if true, run single threaded instead.
+	*	Notes: Please add stats around to calls to parallel for and within your lambda as appropriate. Do not clog the task graph with long running tasks or tasks that block.
+**/
+template<typename FunctionType>
+inline void ParallelForTemplate(int32 Num, const FunctionType& Body, EParallelForFlags Flags = EParallelForFlags::None)
+{
+	ParallelForImpl::ParallelForInternal(Num, Body, Flags);
+}
 /** 
 	*	General purpose parallel for that uses the taskgraph for unbalanced tasks
 	*	Offers better work distribution among threads at the cost of a little bit more synchronization.
@@ -302,9 +347,11 @@ inline void ParallelFor(int32 Num, TFunctionRef<void(int32)> Body, EParallelForF
 	*	@param bForceSingleThread; Mostly used for testing, if true, run single threaded instead.
 	*	Notes: Please add stats around to calls to parallel for and within your lambda as appropriate. Do not clog the task graph with long running tasks or tasks that block.
 **/
-inline void ParallelForWithPreWork(int32 Num, TFunctionRef<void(int32)> Body, TFunctionRef<void()> CurrentThreadWorkToDoBeforeHelping, bool bForceSingleThread)
+inline void ParallelForWithPreWork(int32 Num, TFunctionRef<void(int32)> Body, TFunctionRef<void()> CurrentThreadWorkToDoBeforeHelping, bool bForceSingleThread, bool bPumpRenderingThread = false)
 {
-	ParallelForImpl::ParallelForWithPreWorkInternal(Num, Body, CurrentThreadWorkToDoBeforeHelping, bForceSingleThread ? EParallelForFlags::ForceSingleThread : EParallelForFlags::None);
+	ParallelForImpl::ParallelForWithPreWorkInternal(Num, Body, CurrentThreadWorkToDoBeforeHelping,
+		(bForceSingleThread ? EParallelForFlags::ForceSingleThread : EParallelForFlags::None) |
+		(bPumpRenderingThread ? EParallelForFlags::PumpRenderingThread : EParallelForFlags::None));
 }
 
 /** 
diff --git a/Engine/Source/Runtime/Core/Public/Async/TaskGraphInterfaces.h b/Engine/Source/Runtime/Core/Public/Async/TaskGraphInterfaces.h
--- a/Engine/Source/Runtime/Core/Public/Async/TaskGraphInterfaces.h
+++ b/Engine/Source/Runtime/Core/Public/Async/TaskGraphInterfaces.h
@@ -210,6 +210,8 @@ namespace ENamedThreads
 
 }
 
+DECLARE_INTRINSIC_TYPE_LAYOUT(ENamedThreads::Type);
+
 extern CORE_API int32 GEnablePowerSavingThreadPriorityReductionCVar;
 
 enum class EPowerSavingEligibility : uint8
@@ -340,7 +342,7 @@ public:
 	 *	Requests that a named thread, which must be this thread, run until idle, then return.
 	 *	@param	CurrentThread; The name of this thread
 	**/
-	virtual void ProcessThreadUntilIdle(ENamedThreads::Type CurrentThread)=0;
+	virtual uint64 ProcessThreadUntilIdle(ENamedThreads::Type CurrentThread)=0;
 
 	/** 
 	 *	Requests that a named thread, which must be this thread, run until an explicit return request is received, then return.
diff --git a/Engine/Source/Runtime/Core/Public/Containers/Array.h b/Engine/Source/Runtime/Core/Public/Containers/Array.h
--- a/Engine/Source/Runtime/Core/Public/Containers/Array.h
+++ b/Engine/Source/Runtime/Core/Public/Containers/Array.h
@@ -11,6 +11,7 @@
 #include "Templates/UnrealTemplate.h"
 #include "Containers/ContainerAllocationPolicies.h"
 #include "Serialization/Archive.h"
+#include "Serialization/MemoryImageWriter.h"
 
 #include "Algo/Heapify.h"
 #include "Algo/HeapSort.h"
@@ -21,6 +22,7 @@
 #include "Templates/Less.h"
 #include "Templates/ChooseClass.h"
 #include "Templates/Sorting.h"
+#include "Templates/AlignmentTemplates.h"
 
 
 #if UE_BUILD_SHIPPING || UE_BUILD_TEST
@@ -29,6 +31,12 @@
 	#define TARRAY_RANGED_FOR_CHECKS 1
 #endif
 
+template <typename T> struct TCanBulkSerialize { enum { Value = false }; };
+template<> struct TCanBulkSerialize<unsigned int> { enum { Value = true }; };
+template<> struct TCanBulkSerialize<unsigned short> { enum { Value = true }; };
+template<> struct TCanBulkSerialize<int> { enum { Value = true }; };
+
+
 /**
  * Generic iterator which can operate on types that expose the following:
  * - A type called ElementType representing the contained type.
@@ -286,6 +294,12 @@ public:
 	typedef InElementType ElementType;
 	typedef InAllocator   Allocator;
 
+	typedef typename TChooseClass<
+		Allocator::NeedsElementType,
+		typename Allocator::template ForElementType<ElementType>,
+		typename Allocator::ForAnyElementType
+	>::Result ElementAllocatorType;
+
 	static_assert(TIsSigned<SizeType>::Value, "TArray only supports signed index types");
 
 	/**
@@ -310,6 +324,9 @@ public:
 		CopyToEmpty(Ptr, Count, 0, 0);
 	}
 
+	template <typename OtherElementType>
+	explicit TArray(const TArrayView<OtherElementType>& Other);
+
 	/**
 	 * Initializer list constructor
 	 */
@@ -402,6 +419,9 @@ public:
 		return *this;
 	}
 
+	template <typename OtherElementType>
+	TArray& operator=(const TArrayView<OtherElementType>& Other);
+
 private:
 #if !PLATFORM_COMPILER_HAS_IF_CONSTEXPR
 	template <
@@ -1119,11 +1139,17 @@ public:
 
 		Ar << SerializeNum;
 
+		if (SerializeNum == 0)
+		{
+			return Ar;
+		}
+
 		check(SerializeNum >= 0);
 
-		if (!Ar.IsError() && SerializeNum >= 0 && ensure(!Ar.IsNetArchive() || SerializeNum <= MaxNetArraySerialize))
+		if (!Ar.IsError() && SerializeNum > 0 && ensure(!Ar.IsNetArchive() || SerializeNum <= MaxNetArraySerialize))
 		{
-			if (sizeof(ElementType) == 1)
+			// if we don't need to perform per-item serialization, just read it in bulk
+			if (sizeof(ElementType) == 1 || TCanBulkSerialize<ElementType>::Value)
 			{
 				A.ArrayNum = SerializeNum;
 
@@ -1133,7 +1159,7 @@ public:
 					A.ResizeForCopy(A.ArrayNum, A.ArrayMax);
 				}
 
-				Ar.Serialize(A.GetData(), A.Num());
+				Ar.Serialize(A.GetData(), A.Num() * sizeof(ElementType));
 			}
 			else if (Ar.IsLoading())
 			{
@@ -2578,22 +2604,73 @@ private:
 
 protected:
 
-	typedef typename TChooseClass<
-		Allocator::NeedsElementType,
-		typename Allocator::template ForElementType<ElementType>,
-		typename Allocator::ForAnyElementType
-		>::Result ElementAllocatorType;
+	template<typename ElementType, typename Allocator>
+	friend class TIndirectArray;
 
 	ElementAllocatorType AllocatorInstance;
 	SizeType             ArrayNum;
 	SizeType             ArrayMax;
 
-	/**
-	 * Implicit heaps
-	 */
+private:
+	template<bool bFreezeMemoryImage, typename Dummy=void>
+	struct TSupportsFreezeMemoryImageHelper
+	{
+		static void WriteMemoryImage(FMemoryImageWriter& Writer, const TArray&)
+		{
+			// Writing non-freezable TArray is only supported for 64-bit target for now
+			// Would need complete layout macros for all allocator types in order to properly write (empty) 32bit versions
+			check(Writer.Is64BitTarget());
+			Writer.WriteBytes(TArray());
+		}
+
+		static void CopyUnfrozen(const FMemoryUnfreezeContent& Context, const TArray&, void* Dst) { new(Dst) TArray(); }
+		static void AppendHash(const FPlatformTypeLayoutParameters& LayoutParams, FSHA1& Hasher) {}
+	};
+
+	template<typename Dummy>
+	struct TSupportsFreezeMemoryImageHelper<true, Dummy>
+	{
+		static void WriteMemoryImage(FMemoryImageWriter& Writer, const TArray& Object)
+		{
+			Object.AllocatorInstance.WriteMemoryImage(Writer, StaticGetTypeLayoutDesc<ElementType>(), Object.ArrayNum);
+			Writer.WriteBytes(Object.ArrayNum);
+			Writer.WriteBytes(Object.ArrayNum);
+		}
+		static void CopyUnfrozen(const FMemoryUnfreezeContent& Context, const TArray& Object, void* Dst)
+		{
+			TArray* DstArray = new(Dst) TArray();
+			DstArray->SetNumZeroed(Object.ArrayNum);
+			Object.AllocatorInstance.CopyUnfrozen(Context, StaticGetTypeLayoutDesc<ElementType>(), Object.ArrayNum, DstArray->GetData());
+		}
+		static void AppendHash(const FPlatformTypeLayoutParameters& LayoutParams, FSHA1& Hasher)
+		{
+			Freeze::AppendHash(StaticGetTypeLayoutDesc<ElementType>(), LayoutParams, Hasher);
+		}
+	};
 
 public:
+	void WriteMemoryImage(FMemoryImageWriter& Writer) const
+	{
+		static const bool bSupportsFreezeMemoryImage = TAllocatorTraits<Allocator>::SupportsFreezeMemoryImage && THasTypeLayout<ElementType>::Value;
+		TSupportsFreezeMemoryImageHelper<bSupportsFreezeMemoryImage>::WriteMemoryImage(Writer, *this);
+	}
+
+	void CopyUnfrozen(const FMemoryUnfreezeContent& Context, void* Dst) const
+	{
+		static const bool bSupportsFreezeMemoryImage = TAllocatorTraits<Allocator>::SupportsFreezeMemoryImage && THasTypeLayout<ElementType>::Value;
+		TSupportsFreezeMemoryImageHelper<bSupportsFreezeMemoryImage>::CopyUnfrozen(Context, *this, Dst);
+	}
+
+	static void AppendHash(const FPlatformTypeLayoutParameters& LayoutParams, FSHA1& Hasher)
+	{
+		static const bool bSupportsFreezeMemoryImage = TAllocatorTraits<Allocator>::SupportsFreezeMemoryImage && THasTypeLayout<ElementType>::Value;
+		TSupportsFreezeMemoryImageHelper<bSupportsFreezeMemoryImage>::AppendHash(LayoutParams, Hasher);
+	}
 
+	/**
+	* Implicit heaps
+	*/
+public:
 	/** 
 	 * Builds an implicit heap from the array.
 	 *
@@ -2867,9 +2944,46 @@ public:
 	{
 		HeapSort(TLess<ElementType>());
 	}
+
+	const ElementAllocatorType& GetAllocatorInstance() const { return AllocatorInstance; }
+	ElementAllocatorType& GetAllocatorInstance() { return AllocatorInstance; }
 };
 
 
+namespace Freeze
+{
+	template<typename T, typename AllocatorType>
+	void IntrinsicWriteMemoryImage(FMemoryImageWriter& Writer, const TArray<T, AllocatorType>& Object, const FTypeLayoutDesc&)
+	{
+		Object.WriteMemoryImage(Writer);
+	}
+
+	template<typename T, typename AllocatorType>
+	void IntrinsicUnfrozenCopy(const FMemoryUnfreezeContent& Context, const TArray<T, AllocatorType>& Object, void* OutDst)
+	{
+		Object.CopyUnfrozen(Context, OutDst);
+	}
+
+	template<typename T, typename AllocatorType>
+	uint32 IntrinsicAppendHash(const TArray<T, AllocatorType>* DummyObject, const FTypeLayoutDesc& TypeDesc, const FPlatformTypeLayoutParameters& LayoutParams, FSHA1& Hasher)
+	{
+		// sizeof(TArray) changes depending on target platform 32bit vs 64bit
+		// For now, calculate the size manually
+		static_assert(sizeof(TArray<T, AllocatorType>) == sizeof(FMemoryImageUPtrInt) + sizeof(int32) + sizeof(int32), "Unexpected TArray size");
+		const uint32 SizeFromFields = LayoutParams.GetMemoryImagePointerSize() + sizeof(int32) + sizeof(int32);
+		return AppendHashForNameAndSize(TypeDesc.Name, SizeFromFields, Hasher);;
+	}
+
+	template<typename T, typename AllocatorType>
+	uint32 IntrinsicGetTargetAlignment(const TArray<T, AllocatorType>* DummyObject, const FTypeLayoutDesc& TypeDesc, const FPlatformTypeLayoutParameters& LayoutParams)
+	{
+		// Assume alignment of array is drive by pointer
+		return FMath::Min(LayoutParams.GetMemoryImagePointerSize(), LayoutParams.MaxFieldAlignment);
+	}
+}
+
+DECLARE_TEMPLATE_INTRINSIC_TYPE_LAYOUT((template <typename T, typename AllocatorType>), (TArray<T, AllocatorType>));
+
 template <typename InElementType, typename Allocator>
 struct TIsZeroConstructType<TArray<InElementType, Allocator>>
 {
@@ -2889,7 +3003,6 @@ struct TIsContiguousContainer<TArray<T, Allocator>>
 	enum { Value = true };
 };
 
-
 /**
  * Traits class which determines whether or not a type is a TArray.
  */
diff --git a/Engine/Source/Runtime/Core/Public/Containers/ArrayView.h b/Engine/Source/Runtime/Core/Public/Containers/ArrayView.h
--- a/Engine/Source/Runtime/Core/Public/Containers/ArrayView.h
+++ b/Engine/Source/Runtime/Core/Public/Containers/ArrayView.h
@@ -575,3 +575,19 @@ TArrayView<const T> MakeArrayView(std::initializer_list<T> List)
 {
 	return TArrayView<const T>(List);
 }
+
+template<typename InElementType, typename InAllocatorType>
+template<typename OtherElementType>
+FORCEINLINE TArray<InElementType, InAllocatorType>::TArray(const TArrayView<OtherElementType>& Other)
+{
+	CopyToEmpty(Other.GetData(), Other.Num(), 0, 0);
+}
+
+template<typename InElementType, typename InAllocatorType>
+template<typename OtherElementType>
+FORCEINLINE TArray<InElementType, InAllocatorType>& TArray<InElementType, InAllocatorType>::operator=(const TArrayView<OtherElementType>& Other)
+{
+	DestructItems(GetData(), ArrayNum);
+	CopyToEmpty(Other.GetData(), Other.Num(), ArrayMax, 0);
+	return *this;
+}
diff --git a/Engine/Source/Runtime/Core/Public/Containers/BitArray.h b/Engine/Source/Runtime/Core/Public/Containers/BitArray.h
--- a/Engine/Source/Runtime/Core/Public/Containers/BitArray.h
+++ b/Engine/Source/Runtime/Core/Public/Containers/BitArray.h
@@ -9,6 +9,7 @@
 #include "Templates/UnrealTemplate.h"
 #include "Containers/ContainerAllocationPolicies.h"
 #include "Serialization/Archive.h"
+#include "Serialization/MemoryImageWriter.h"
 #include "Math/UnrealMathUtility.h"
 
 template<typename Allocator > class TBitArray;
@@ -37,7 +38,8 @@ class TConstSetBitIterator;
 template<typename Allocator = FDefaultBitArrayAllocator,typename OtherAllocator = FDefaultBitArrayAllocator>
 class TConstDualSetBitIterator;
 
-class FScriptBitArray;
+template <typename AllocatorType, typename InDerivedType = void>
+class TScriptBitArray;
 
 
 /**
@@ -175,10 +177,13 @@ public:
 template<typename Allocator /*= FDefaultBitArrayAllocator*/>
 class TBitArray
 {
-	friend class FScriptBitArray;
+	template <typename, typename>
+	friend class TScriptBitArray;
 
 public:
 
+	typedef typename Allocator::template ForElementType<uint32> AllocatorType;
+
 	template<typename>
 	friend class TConstSetBitIterator;
 
@@ -949,8 +954,6 @@ public:
 	}
 
 private:
-	typedef typename Allocator::template ForElementType<uint32> AllocatorType;
-
 	AllocatorType AllocatorInstance;
 	int32         NumBits;
 	int32         MaxBits;
@@ -968,8 +971,45 @@ private:
 			FMemory::Memzero((uint32*)AllocatorInstance.GetAllocation() + PreviousNumDWORDs,(MaxDWORDs - PreviousNumDWORDs) * sizeof(uint32));
 		}
 	}
+
+	template<bool bFreezeMemoryImage, typename Dummy=void>
+	struct TSupportsFreezeMemoryImageHelper
+	{
+		static void WriteMemoryImage(FMemoryImageWriter& Writer, const TBitArray&) { Writer.WriteBytes(TBitArray()); }
+	};
+
+	template<typename Dummy>
+	struct TSupportsFreezeMemoryImageHelper<true, Dummy>
+	{
+		static void WriteMemoryImage(FMemoryImageWriter& Writer, const TBitArray& Object)
+		{
+			const int32 NumDWORDs = FMath::DivideAndRoundUp(Object.NumBits, NumBitsPerDWORD);
+			Object.AllocatorInstance.WriteMemoryImage(Writer, StaticGetTypeLayoutDesc<uint32>(), NumDWORDs);
+			Writer.WriteBytes(Object.NumBits);
+			Writer.WriteBytes(Object.NumBits);
+		}
+	};
+
+public:
+	void WriteMemoryImage(FMemoryImageWriter& Writer) const
+	{
+		static const bool bSupportsFreezeMemoryImage = TAllocatorTraits<Allocator>::SupportsFreezeMemoryImage;
+		checkf(!Writer.GetTargetLayoutParams().b32Bit, TEXT("TBitArray does not currently support freezing for 32bits"));
+		TSupportsFreezeMemoryImageHelper<bSupportsFreezeMemoryImage>::WriteMemoryImage(Writer, *this);
+	}
 };
 
+namespace Freeze
+{
+	template<typename Allocator>
+	void IntrinsicWriteMemoryImage(FMemoryImageWriter& Writer, const TBitArray<Allocator>& Object, const FTypeLayoutDesc&)
+	{
+		Object.WriteMemoryImage(Writer);
+	}
+}
+
+DECLARE_TEMPLATE_INTRINSIC_TYPE_LAYOUT(template<typename Allocator>, TBitArray<Allocator>);
+
 template<typename Allocator>
 FORCEINLINE uint32 GetTypeHash(const TBitArray<Allocator>& BitArray)
 {
@@ -1214,15 +1254,18 @@ private:
 
 // Untyped bit array type for accessing TBitArray data, like FScriptArray for TArray.
 // Must have the same memory representation as a TBitArray.
-class FScriptBitArray
+template <typename Allocator, typename InDerivedType>
+class TScriptBitArray
 {
+	using DerivedType = typename TChooseClass<TIsVoidType<InDerivedType>::Value, TScriptBitArray, InDerivedType>::Result;
+
 public:
 	/**
 	 * Minimal initialization constructor.
 	 * @param Value - The value to initial the bits to.
 	 * @param InNumBits - The initial number of bits in the array.
 	 */
-	FScriptBitArray()
+	TScriptBitArray()
 		: NumBits(0)
 		, MaxBits(0)
 	{
@@ -1245,7 +1288,7 @@ public:
 		return FConstBitReference(GetData()[Index / NumBitsPerDWORD], 1 << (Index & (NumBitsPerDWORD - 1)));
 	}
 
-	void MoveAssign(FScriptBitArray& Other)
+	void MoveAssign(DerivedType& Other)
 	{
 		checkSlow(this != &Other);
 		Empty(0);
@@ -1280,7 +1323,7 @@ public:
 	}
 
 private:
-	typedef FDefaultBitArrayAllocator::ForElementType<uint32> AllocatorType;
+	typedef typename Allocator::template ForElementType<uint32> AllocatorType;
 
 	AllocatorType AllocatorInstance;
 	int32         NumBits;
@@ -1289,22 +1332,22 @@ private:
 	// This function isn't intended to be called, just to be compiled to validate the correctness of the type.
 	static void CheckConstraints()
 	{
-		typedef FScriptBitArray ScriptType;
+		typedef TScriptBitArray ScriptType;
 		typedef TBitArray<>     RealType;
 
 		// Check that the class footprint is the same
-		static_assert(sizeof (ScriptType) == sizeof (RealType), "FScriptBitArray's size doesn't match TBitArray");
-		static_assert(alignof(ScriptType) == alignof(RealType), "FScriptBitArray's alignment doesn't match TBitArray");
+		static_assert(sizeof (ScriptType) == sizeof (RealType), "TScriptBitArray's size doesn't match TBitArray");
+		static_assert(alignof(ScriptType) == alignof(RealType), "TScriptBitArray's alignment doesn't match TBitArray");
 
 		// Check member sizes
-		static_assert(sizeof(DeclVal<ScriptType>().AllocatorInstance) == sizeof(DeclVal<RealType>().AllocatorInstance), "FScriptBitArray's AllocatorInstance member size does not match TBitArray's");
-		static_assert(sizeof(DeclVal<ScriptType>().NumBits)           == sizeof(DeclVal<RealType>().NumBits),           "FScriptBitArray's NumBits member size does not match TBitArray's");
-		static_assert(sizeof(DeclVal<ScriptType>().MaxBits)           == sizeof(DeclVal<RealType>().MaxBits),           "FScriptBitArray's MaxBits member size does not match TBitArray's");
+		static_assert(sizeof(DeclVal<ScriptType>().AllocatorInstance) == sizeof(DeclVal<RealType>().AllocatorInstance), "TScriptBitArray's AllocatorInstance member size does not match TBitArray's");
+		static_assert(sizeof(DeclVal<ScriptType>().NumBits)           == sizeof(DeclVal<RealType>().NumBits),           "TScriptBitArray's NumBits member size does not match TBitArray's");
+		static_assert(sizeof(DeclVal<ScriptType>().MaxBits)           == sizeof(DeclVal<RealType>().MaxBits),           "TScriptBitArray's MaxBits member size does not match TBitArray's");
 
 		// Check member offsets
-		static_assert(STRUCT_OFFSET(ScriptType, AllocatorInstance) == STRUCT_OFFSET(RealType, AllocatorInstance), "FScriptBitArray's AllocatorInstance member offset does not match TBitArray's");
-		static_assert(STRUCT_OFFSET(ScriptType, NumBits)           == STRUCT_OFFSET(RealType, NumBits),           "FScriptBitArray's NumBits member offset does not match TBitArray's");
-		static_assert(STRUCT_OFFSET(ScriptType, MaxBits)           == STRUCT_OFFSET(RealType, MaxBits),           "FScriptBitArray's MaxBits member offset does not match TBitArray's");
+		static_assert(STRUCT_OFFSET(ScriptType, AllocatorInstance) == STRUCT_OFFSET(RealType, AllocatorInstance), "TScriptBitArray's AllocatorInstance member offset does not match TBitArray's");
+		static_assert(STRUCT_OFFSET(ScriptType, NumBits)           == STRUCT_OFFSET(RealType, NumBits),           "TScriptBitArray's NumBits member offset does not match TBitArray's");
+		static_assert(STRUCT_OFFSET(ScriptType, MaxBits)           == STRUCT_OFFSET(RealType, MaxBits),           "TScriptBitArray's MaxBits member offset does not match TBitArray's");
 	}
 
 	FORCEINLINE uint32* GetData()
@@ -1355,12 +1398,20 @@ private:
 public:
 	// These should really be private, because they shouldn't be called, but there's a bunch of code
 	// that needs to be fixed first.
-	FScriptBitArray(const FScriptBitArray&) { check(false); }
-	void operator=(const FScriptBitArray&) { check(false); }
+	TScriptBitArray(const TScriptBitArray&) { check(false); }
+	void operator=(const TScriptBitArray&) { check(false); }
 };
 
-template <>
-struct TIsZeroConstructType<FScriptBitArray>
+template <typename AllocatorType, typename InDerivedType>
+struct TIsZeroConstructType<TScriptBitArray<AllocatorType, InDerivedType>>
 {
 	enum { Value = true };
 };
+
+class FScriptBitArray : public TScriptBitArray<FDefaultBitArrayAllocator, FScriptBitArray>
+{
+	using Super = TScriptBitArray<FDefaultBitArrayAllocator, FScriptBitArray>;
+
+public:
+	using Super::Super;
+};
diff --git a/Engine/Source/Runtime/Core/Public/Containers/ContainerAllocationPolicies.h b/Engine/Source/Runtime/Core/Public/Containers/ContainerAllocationPolicies.h
--- a/Engine/Source/Runtime/Core/Public/Containers/ContainerAllocationPolicies.h
+++ b/Engine/Source/Runtime/Core/Public/Containers/ContainerAllocationPolicies.h
@@ -9,6 +9,7 @@
 #include "HAL/PlatformMath.h"
 #include "Templates/MemoryOps.h"
 #include "Math/NumericLimits.h"
+#include "Templates/IsPolymorphic.h"
 
 class FDefaultBitArrayAllocator;
 
@@ -125,6 +126,7 @@ struct TAllocatorTraitsBase
 {
 	enum { SupportsMove    = false };
 	enum { IsZeroConstruct = false };
+	enum { SupportsFreezeMemoryImage = false };
 };
 
 template <typename AllocatorType>
@@ -480,7 +482,7 @@ public:
 		{
 			return !!Data;
 		}
-		
+
 		SizeType GetInitialCapacity() const
 		{
 			return 0;
@@ -898,7 +900,7 @@ public:
 		{
 			return false;
 		}
-		
+
 		SizeType GetInitialCapacity() const
 		{
 			return NumInlineElements;
@@ -978,6 +980,8 @@ public:
 	typedef TFixedAllocator<InlineBitArrayDWORDs> BitArrayAllocator;
 };
 
+
+
 //
 // Set allocation definitions.
 //
@@ -1013,6 +1017,22 @@ public:
 	typedef InHashAllocator        HashAllocator;
 };
 
+template<
+	typename InSparseArrayAllocator,
+	typename InHashAllocator,
+	uint32   AverageNumberOfElementsPerHashBucket,
+	uint32   BaseNumberOfHashBuckets,
+	uint32   MinNumberOfHashedElements
+>
+struct TAllocatorTraits<TSetAllocator<InSparseArrayAllocator, InHashAllocator, AverageNumberOfElementsPerHashBucket, BaseNumberOfHashBuckets, MinNumberOfHashedElements>> :
+	TAllocatorTraitsBase<TSetAllocator<InSparseArrayAllocator, InHashAllocator, AverageNumberOfElementsPerHashBucket, BaseNumberOfHashBuckets, MinNumberOfHashedElements>>
+{
+	enum
+	{
+		SupportsFreezeMemoryImage = TAllocatorTraits<InSparseArrayAllocator>::SupportsFreezeMemoryImage && TAllocatorTraits<InHashAllocator>::SupportsFreezeMemoryImage,
+	};
+};
+
 /** An inline set allocator that allows sizing of the inline allocations for a set number of elements. */
 template<
 	uint32   NumInlineElements,
@@ -1109,4 +1129,13 @@ template <> struct TAllocatorTraits<FDefaultSetAllocator>         : TAllocatorTr
 template <> struct TAllocatorTraits<FDefaultBitArrayAllocator>    : TAllocatorTraits<typename FDefaultBitArrayAllocator   ::Typedef> {};
 template <> struct TAllocatorTraits<FDefaultSparseArrayAllocator> : TAllocatorTraits<typename FDefaultSparseArrayAllocator::Typedef> {};
 
+template <typename InElementAllocator, typename InBitArrayAllocator>
+struct TAllocatorTraits<TSparseArrayAllocator<InElementAllocator, InBitArrayAllocator>> : TAllocatorTraitsBase<TSparseArrayAllocator<InElementAllocator, InBitArrayAllocator>>
+{
+	enum
+	{
+		SupportsFreezeMemoryImage = TAllocatorTraits<InElementAllocator>::SupportsFreezeMemoryImage && TAllocatorTraits<InBitArrayAllocator>::SupportsFreezeMemoryImage,
+	};
+};
+
 template <uint8 FromIndexSize, uint8 ToIndexSize> struct TCanMoveBetweenAllocators<TSizedDefaultAllocator<FromIndexSize>, TSizedDefaultAllocator<ToIndexSize>> : TCanMoveBetweenAllocators<typename TSizedDefaultAllocator<FromIndexSize>::Typedef, typename TSizedDefaultAllocator<ToIndexSize>::Typedef> {};
diff --git a/Engine/Source/Runtime/Core/Public/Containers/ContainersFwd.h b/Engine/Source/Runtime/Core/Public/Containers/ContainersFwd.h
--- a/Engine/Source/Runtime/Core/Public/Containers/ContainersFwd.h
+++ b/Engine/Source/Runtime/Core/Public/Containers/ContainersFwd.h
@@ -12,6 +12,8 @@ class FString;
 
 template<typename T, typename Allocator = FDefaultAllocator> class TArray;
 template<typename T> using TArray64 = TArray<T, FDefaultAllocator64>;
+template<typename T> class TArrayView;
+template<typename T> using TConstArrayView = TArrayView<const T>;
 template<typename T> class TTransArray;
 template<typename KeyType, typename ValueType, bool bInAllowDuplicateKeys> struct TDefaultMapHashableKeyFuncs;
 template<typename KeyType, typename ValueType, typename SetAllocator = FDefaultSetAllocator, typename KeyFuncs = TDefaultMapHashableKeyFuncs<KeyType, ValueType, false> > class TMap;
diff --git a/Engine/Source/Runtime/Core/Public/Containers/DynamicRHIResourceArray.h b/Engine/Source/Runtime/Core/Public/Containers/DynamicRHIResourceArray.h
--- a/Engine/Source/Runtime/Core/Public/Containers/DynamicRHIResourceArray.h
+++ b/Engine/Source/Runtime/Core/Public/Containers/DynamicRHIResourceArray.h
@@ -8,6 +8,7 @@
 #include "Containers/Array.h"
 #include "CoreGlobals.h"
 #include "Containers/ResourceArray.h"
+#include "Serialization/MemoryImage.h"
 
 /** alignment for supported resource types */
 enum EResourceAlignment
@@ -23,13 +24,15 @@ enum EResourceAlignment
  *
  * @param Alignment - memory alignment to use for the allocation
  */
-template< typename ElementType, uint32 Alignment = DEFAULT_ALIGNMENT >
+template< typename ElementType, uint32 Alignment = DEFAULT_ALIGNMENT>
 class TResourceArray
 	:	public FResourceArrayInterface
-	,	public TArray<ElementType, TAlignedHeapAllocator<Alignment> >
+	,	public TArray<ElementType, TMemoryImageAllocator<Alignment>>
 {
+	using ParentArrayType = TArray<ElementType, TMemoryImageAllocator<Alignment>>;
+	DECLARE_INLINE_TYPE_LAYOUT_EXPLICIT_BASES(TResourceArray, Virtual, FResourceArrayInterface, ParentArrayType);
 public:
-	typedef TArray<ElementType, TAlignedHeapAllocator<Alignment> > Super;
+	using Super = ParentArrayType;
 
 	/** 
 	* Constructor 
@@ -138,11 +141,11 @@ private:
 	* true if this array needs to be accessed by the CPU.  
 	* If no CPU access is needed then the resource is freed once its RHI resource has been created
 	*/
-	bool bNeedsCPUAccess;
+	LAYOUT_FIELD(bool, bNeedsCPUAccess);
 };
 
-template< typename ElementType, uint32 Alignment >
-struct TContainerTraits<TResourceArray<ElementType, Alignment> > : public TContainerTraitsBase<TResourceArray<ElementType, Alignment> >
+template< typename ElementType, uint32 Alignment>
+struct TContainerTraits<TResourceArray<ElementType, Alignment>> : public TContainerTraitsBase<TResourceArray<ElementType, Alignment>>
 {
 	enum { MoveWillEmptyContainer = TContainerTraits<typename TResourceArray<ElementType, Alignment>::Super>::MoveWillEmptyContainer };
 };
diff --git a/Engine/Source/Runtime/Core/Public/Containers/HashTable.h b/Engine/Source/Runtime/Core/Public/Containers/HashTable.h
--- a/Engine/Source/Runtime/Core/Public/Containers/HashTable.h
+++ b/Engine/Source/Runtime/Core/Public/Containers/HashTable.h
@@ -183,6 +183,8 @@ public:
 					FHashTable( uint32 InHashSize = 1024, uint32 InIndexSize = 0 );
 					~FHashTable();
 
+	void			Initialize(uint32 InHashSize = 1024, uint32 InIndexSize = 0);
+
 	void			Clear();
 	void			Free();
 	CORE_API void	Resize( uint32 NewIndexSize );
@@ -213,23 +215,37 @@ protected:
 
 
 FORCEINLINE FHashTable::FHashTable( uint32 InHashSize, uint32 InIndexSize )
-	: HashSize( InHashSize )
+	: HashSize( 0 )
 	, HashMask( 0 )
-	, IndexSize( InIndexSize )
+	, IndexSize( 0 )
 	, Hash( EmptyHash )
 	, NextIndex( NULL )
 {
-	check( HashSize <= 0x10000 );
-	check( FMath::IsPowerOfTwo( HashSize ) );
-	
-	if( IndexSize )
+	if (InHashSize > 0u)
+	{
+		Initialize(InHashSize, InIndexSize);
+	}
+}
+
+FORCEINLINE void FHashTable::Initialize(uint32 InHashSize, uint32 InIndexSize)
+{
+	check(HashSize == 0u);
+	check(IndexSize == 0u);
+
+	HashSize = InHashSize;
+	IndexSize = InIndexSize;
+
+	check(HashSize <= 0x10000);
+	check(FMath::IsPowerOfTwo(HashSize));
+
+	if (IndexSize)
 	{
 		HashMask = (uint16)(HashSize - 1);
-		
-		Hash = new uint32[ HashSize ];
-		NextIndex = new uint32[ IndexSize ];
 
-		FMemory::Memset( Hash, 0xff, HashSize * 4 );
+		Hash = new uint32[HashSize];
+		NextIndex = new uint32[IndexSize];
+
+		FMemory::Memset(Hash, 0xff, HashSize * 4);
 	}
 }
 
@@ -272,6 +288,7 @@ FORCEINLINE uint32 FHashTable::First( uint16 Key ) const
 FORCEINLINE uint32 FHashTable::Next( uint32 Index ) const
 {
 	checkSlow( Index < IndexSize );
+	checkSlow( NextIndex[Index] != Index ); // check for corrupt tables
 	return NextIndex[ Index ];
 }
 
