fix for verifying the data received from the web call is valid

diff --git a/Engine/Source/Runtime/ApplicationCore/Private/IOS/IOSAppDelegate.cpp b/Engine/Source/Runtime/ApplicationCore/Private/IOS/IOSAppDelegate.cpp
--- a/Engine/Source/Runtime/ApplicationCore/Private/IOS/IOSAppDelegate.cpp
+++ b/Engine/Source/Runtime/ApplicationCore/Private/IOS/IOSAppDelegate.cpp
@@ -315,9 +315,9 @@ static IOSAppDelegate* CachedDelegate = nil;
 	FAppEntry::Init();
 
 	// check for update on app store if cvar is enabled
-	dispatch_async(dispatch_get_main_queue(), ^{
+/*	dispatch_async(dispatch_get_main_queue(), ^{
 		[[IOSAppDelegate GetDelegate] DoUpdateCheck];
-	});
+	});*/
 	
 	// now that GConfig has been loaded, load the EnabledAudioFeatures from ini
 	TArray<FString> EnabledAudioFeatures;
@@ -1311,16 +1311,24 @@ static FAutoConsoleVariableRef CVarGEnableThermalsReport(
 		
 		// kick off an NSURLSession to read the data
 		NSURLSession* Session = [NSURLSession sharedSession];
-		NSURLSessionDataTask* SessionTask = [Session dataTaskWithRequest: [NSURLRequest requestWithURL: StoreURL] completionHandler:^(NSData* _Nullable data, NSURLResponse* _Nullable response, NSError* _Nullable error) {
-			NSDictionary* StoreDictionary = [NSJSONSerialization JSONObjectWithData: data options: 0 error: nil];
+		NSURLSessionDataTask* SessionTask = [Session dataTaskWithRequest: [NSURLRequest requestWithURL: StoreURL cachePolicy: NSURLRequestReloadIgnoringLocalCacheData timeoutInterval: Session.configuration.timeoutIntervalForRequest] completionHandler:^(NSData* _Nullable data, NSURLResponse* _Nullable response, NSError* _Nullable error) {
 			
-			if ([StoreDictionary[@"resultCount"] integerValue] == 1)
+			if (error == nil && data != nil)
 			{
-				// get the store version
-				NSString* StoreVersion = StoreDictionary[@"results"][0][@"version"];
-				if ([StoreVersion compare: CurrentVersion options: NSNumericSearch] == NSOrderedDescending)
+				NSDictionary* StoreDictionary = [NSJSONSerialization JSONObjectWithData: data options: 0 error: nil];
+			
+				if ([StoreDictionary[@"resultCount"] integerValue] == 1)
 				{
-					self.bUpdateAvailable = true;
+					// get the store version
+					NSString* StoreVersion = StoreDictionary[@"results"][0][@"version"];
+					if ([StoreVersion compare: CurrentVersion options: NSNumericSearch] == NSOrderedDescending)
+					{
+						self.bUpdateAvailable = true;
+					}
+					else
+					{
+						self.bUpdateAvailable = false;
+					}
 				}
 				else
 				{
