HTML5 - bit of code cleanup

#rnx

diff --git a/Engine/Source/Runtime/ApplicationCore/Private/HTML5/HTML5Application.cpp b/Engine/Source/Runtime/ApplicationCore/Private/HTML5/HTML5Application.cpp
--- a/Engine/Source/Runtime/ApplicationCore/Private/HTML5/HTML5Application.cpp
+++ b/Engine/Source/Runtime/ApplicationCore/Private/HTML5/HTML5Application.cpp
@@ -453,58 +453,8 @@ EM_BOOL FHTML5Application::OnKeyEvent(int eventType, const EmscriptenKeyboardEve
 		CharCode = Character = '`';
 	}
 
-	/* BUG: Workaround a crash - when pressing the (') character, HTML5 builds crash on assertion failure for some reason. Error message:
-		Expression '(GMaxRHIFeatureLevel >= ERHIFeatureLevel::SM5) || !PendingState.bFramebufferSetupInvalid' failed in Engine/Source/Runtime/OpenGLDrv/Private/OpenGLRenderTarget.cpp:1036!
-		  at _FOpenGLDynamicRHI::BindPendingFramebuffer(FOpenGLContextState&)
-		  at _FOpenGLDynamicRHI::RHIEndDrawIndexedPrimitiveUP()
-		  at _non-virtual thunk to FOpenGLDynamicRHI::RHIEndDrawIndexedPrimitiveUP()
-		  at _DrawIndexedPrimitiveUP(FRHICommandList&, unsigned int, unsigned int, unsigned int, unsigned int, void const*, unsigned int, void const*, unsigned int)
-		  at _FBatchedElements::Draw(FRHICommandList&, FDrawingPolicyRenderState const&, ERHIFeatureLevel::Type, bool, FSceneView const&, bool, float, TRefCountPtr<FRHITexture2D>, EBlendModeFilter::Type) const
-		  at _FCanvasBatchedElementRenderItem::Render_GameThread(FCanvas const*)::__8::operator()(FRHICommandList&) const
-		  at _TEnqueueUniqueRenderCommandType<FCanvasBatchedElementRenderItem::Render_GameThread(FCanvas const*)::BatchedDrawCommandName, FCanvasBatchedElementRenderItem::Render_GameThread(FCanvas const*)::__8>::DoTask(ENamedThreads::Type, TRefCountPtr<FGraphEvent> const&)
-		  at _FCanvasBatchedElementRenderItem::Render_GameThread(FCanvas const*)
-		  at _FCanvas::Flush_GameThread(bool)
-		  at _UGameViewportClient::Draw(FViewport*, FCanvas*)
-		  at _non-virtual thunk to UGameViewportClient::Draw(FViewport*, FCanvas*)
-		  at _FViewport::Draw(bool)
-		  at _UGameEngine::RedrawViewports(bool)
-		  at _UGameEngine::Tick(float, bool)
-		  at _FEngineLoop::Tick()
-		  at _HTML5_Tick()
-		  at dynCall_v
-		  at browserIterationFunc
-		  at runIter
-		  at Browser_mainLoop_runner
-		Note: Currently this prevents inputting the ' character to console. TODO: Figure out why this crashes. */
-	if (!strcmp(keyEvent->key, "'")) return 0;
-
-	/* BUG: Workaround a crash - when pressing F5, HTML5 builds crash on assertion failure:
-		DebugBreak() called!
-	      at __Z19ValidateBoundShader12TRefCountPtrI23FOpenGLBoundShaderStateEP15FRHIPixelShader
-	      at __ZN17FOpenGLDynamicRHI25RHISetShaderUniformBufferEP15FRHIPixelShaderjP17FRHIUniformBuffer
-	      at __ZThn4_N17FOpenGLDynamicRHI25RHISetShaderUniformBufferEP15FRHIPixelShaderjP17FRHIUniformBuffer
-	      at __Z25SetUniformBufferParameterIP15FRHIPixelShader39FMobileDirectionalLightShaderParameters15FRHICommandListEvRT1_T_RK29TShaderUniformBufferParameterIT0_ERK17TUniformBufferRefIS8_E
-	      at __ZNK28TMobileBasePassDrawingPolicyI22FUniformLightMapPolicyLi0EE18SetMeshRenderStateER15FRHICommandListRK9FViewInfoPK20FPrimitiveSceneProxyRK10FMeshBatchiRK25FDrawingPolicyRenderStateRKNS1_15ElementDataTypeEN18FMeshDrawingPolicy15ContextDataTypeE
-	      at __ZN36FDrawMobileBasePassDynamicMeshAction7ProcessILi0EEEvR15FRHICommandListRK30FProcessBasePassMeshParametersRK22FUniformLightMapPolicyRKPK20FLightCacheInterface
-	      at __Z25ProcessMobileBasePassMeshI36FDrawMobileBasePassDynamicMeshActionLi0EEvR15FRHICommandListRK30FProcessBasePassMeshParametersOT_
-	      at __ZN41FMobileBasePassOpaqueDrawingPolicyFactory20DrawDynamicMeshTemplILi0EEEvR15FRHICommandListRK9FViewInfoNS_11ContextTypeERK25FDrawingPolicyRenderStateRK10FMeshBatchPK9FMaterialPK20FPrimitiveSceneProxy11FHitProxyId
-	      at __ZN41FMobileBasePassOpaqueDrawingPolicyFactory15DrawDynamicMeshER15FRHICommandListRK9FViewInfoNS_11ContextTypeERK10FMeshBatchbRK25FDrawingPolicyRenderStatePK20FPrimitiveSceneProxy11FHitProxyId
-	      at __ZN20FMobileSceneRenderer20RenderMobileBasePassER24FRHICommandListImmediate10TArrayViewIPK9FViewInfoE
-	      at __ZN20FMobileSceneRenderer6RenderER24FRHICommandListImmediate
-	      at __ZL29RenderViewFamily_RenderThreadR24FRHICommandListImmediateP14FSceneRenderer
-	      at __ZZN15FRendererModule24BeginRenderingViewFamilyEP7FCanvasP16FSceneViewFamilyEN27EURCMacro_FDrawSceneCommand6DoTaskEN13ENamedThreads4TypeERK12TRefCountPtrI11FGraphEventE
-	      at __ZN15FRendererModule24BeginRenderingViewFamilyEP7FCanvasP16FSceneViewFamily
-	      at __ZN19UGameViewportClient4DrawEP9FViewportP7FCanvas
-	      at __ZThn32_N19UGameViewportClient4DrawEP9FViewportP7FCanvas
-	      at __ZN9FViewport4DrawEb
-	      at __ZN11UGameEngine15RedrawViewportsEb
-	      at __ZN11UGameEngine4TickEfb
-	      at __ZN11FEngineLoop4TickEv
-	      at __Z10HTML5_Tickv
-	      at dynCall_v
-	      at browserIterationFunc
-	      at runIter
-	      at Browser_mainLoop_runner */
+	// HACK: trap F5 key presses (i.e. do not process this key)
+	// browser will reload page -- which can lead to unknown conditions and crashes...
 	if (domPhysicalKeyCode == DOM_PK_F5) return 0;
 
 	switch(eventType)
@@ -705,43 +655,43 @@ void FHTML5Application::PollGameDeviceState( const float TimeDelta )
 		// Tick Input Interface.
 		switch (Event.type)
 		{
-				case SDL_WINDOWEVENT:
-				{
-					SDL_WindowEvent windowEvent = Event.window;
+			case SDL_WINDOWEVENT:
+			{
+				SDL_WindowEvent windowEvent = Event.window;
 
-					switch (windowEvent.event)
+				switch (windowEvent.event)
+				{
+				case SDL_WINDOWEVENT_ENTER:
+					{
+						UE_LOG(LogHTML5Application, Verbose, TEXT("WindowEnter"));
+					//	MessageHandler->OnWindowActivationChanged(ApplicationWindow, EWindowActivation::Activate);
+						WarmUpTicks = 0;
+					}
+					break;
+				case SDL_WINDOWEVENT_LEAVE:
+					{
+						UE_LOG(LogHTML5Application, Verbose, TEXT("WindowLeave"));
+					 //	MessageHandler->OnWindowActivationChanged(ApplicationWindow, EWindowActivation::Deactivate);
+					}
+					break;
+				case SDL_WINDOWEVENT_FOCUS_GAINED:
+					{
+						UE_LOG(LogHTML5Application, Verbose, TEXT("WindowFocusGained"));
+						MessageHandler->OnCursorSet();
+						MessageHandler->OnWindowActivationChanged(ApplicationWindow, EWindowActivation::Activate);
+								WarmUpTicks = 0;
+					}
+					break;
+				case SDL_WINDOWEVENT_FOCUS_LOST:
 					{
-					case SDL_WINDOWEVENT_ENTER:
-						{
-							UE_LOG(LogHTML5Application, Verbose, TEXT("WindowEnter"));
-						//	MessageHandler->OnWindowActivationChanged(ApplicationWindow, EWindowActivation::Activate);
-							WarmUpTicks = 0;
-						}
-						break;
-					case SDL_WINDOWEVENT_LEAVE:
-						{
-							UE_LOG(LogHTML5Application, Verbose, TEXT("WindowLeave"));
-						 //	MessageHandler->OnWindowActivationChanged(ApplicationWindow, EWindowActivation::Deactivate);
-						}
-						break;
-					case SDL_WINDOWEVENT_FOCUS_GAINED:
-						{
-							UE_LOG(LogHTML5Application, Verbose, TEXT("WindowFocusGained"));
-							MessageHandler->OnCursorSet();
-							MessageHandler->OnWindowActivationChanged(ApplicationWindow, EWindowActivation::Activate);
-									WarmUpTicks = 0;
-						}
-						break;
-					case SDL_WINDOWEVENT_FOCUS_LOST:
-						{
-							UE_LOG(LogHTML5Application, Verbose, TEXT("WindowFocusLost"));
-							MessageHandler->OnWindowActivationChanged(ApplicationWindow, EWindowActivation::Deactivate);
-						}
-						break;
-					default:
-						break;
+						UE_LOG(LogHTML5Application, Verbose, TEXT("WindowFocusLost"));
+						MessageHandler->OnWindowActivationChanged(ApplicationWindow, EWindowActivation::Deactivate);
 					}
+					break;
+				default:
+					break;
 				}
+			}
 			default:
 			{
 				InputInterface->Tick( TimeDelta,Event, ApplicationWindow);
@@ -754,11 +704,9 @@ void FHTML5Application::PollGameDeviceState( const float TimeDelta )
 
 /*
 
-
 	if ( WarmUpTicks >= 0)
 		WarmUpTicks ++;
 
-
 	if ( WarmUpTicks == MaxWarmUpTicks  )
 	{
 		// browsers don't allow locking and hiding to work independently. use warmup ticks after the application has settled
