Apply PIE Advanced options Safe Zone to mobile PIE standalone process using command line overrides

diff --git a/Engine/Source/Editor/UnrealEd/Private/PlayLevel.cpp b/Engine/Source/Editor/UnrealEd/Private/PlayLevel.cpp
--- a/Engine/Source/Editor/UnrealEd/Private/PlayLevel.cpp
+++ b/Engine/Source/Editor/UnrealEd/Private/PlayLevel.cpp
@@ -1469,7 +1469,24 @@ void UEditorEngine::PlayStandaloneLocalPc(FString MapNameOverride, FIntPoint* Wi
 	{
 		AdditionalParameters += TEXT(" -ForceRes");
 	}
-	
+
+	// In order for the mobile previewer to adjust its safe zone according to the device profile specified in the editor play settings,
+	// we need to pass the PIESafeZoneOverride's values as command line variables to the new process that we are about to launch.
+	FMargin PIESafeZoneOverride = PlayInSettings->PIESafeZoneOverride;
+	FString PIESafeZonesStr = FString::Printf(TEXT(" -SafeZonePaddingLeft=%f -SafeZonePaddingRight=%f -SafeZonePaddingTop=%f -SafeZonePaddingBottom=%f"),
+		PIESafeZoneOverride.Left,
+		PIESafeZoneOverride.Right,
+		PIESafeZoneOverride.Top,
+		PIESafeZoneOverride.Bottom
+	);
+	AdditionalParameters.Append(PIESafeZonesStr);
+
+	if (!PIESafeZoneOverride.GetDesiredSize().IsZero())
+	{
+		// Send -DrawUnSafeZones so that the red "unsafe" zones show up in the mobile previewer.
+		AdditionalParameters.Append(TEXT(" -DrawUnSafeZones"));
+	}
+   	
 	// Check if centered
 	FString Params;
 	if (PlayInSettings->CenterNewWindow)
diff --git a/Engine/Source/Runtime/ApplicationCore/Private/GenericPlatform/GenericApplication.cpp b/Engine/Source/Runtime/ApplicationCore/Private/GenericPlatform/GenericApplication.cpp
--- a/Engine/Source/Runtime/ApplicationCore/Private/GenericPlatform/GenericApplication.cpp
+++ b/Engine/Source/Runtime/ApplicationCore/Private/GenericPlatform/GenericApplication.cpp
@@ -3,6 +3,7 @@
 #include "GenericPlatform/GenericApplication.h"
 #include "HAL/IConsoleManager.h"
 #include "Misc/CoreDelegates.h"
+#include "Misc/CommandLine.h"
 #include "IInputDevice.h"
 
 const FGamepadKeyNames::Type FGamepadKeyNames::Invalid(NAME_None);
@@ -197,11 +198,22 @@ float FDisplayMetrics::GetDebugActionSafeZoneRatio()
 
 void FDisplayMetrics::ApplyDefaultSafeZones()
 {
-	const float SafeZoneRatio = GetDebugTitleSafeZoneRatio();
-	if (SafeZoneRatio < 1.0f)
+	// Allow safe zones to be overridden by the command line. Used by mobile PIE.
+	TitleSafePaddingSize = FVector4(0.0f, 0.0f, 0.0f, 0.0f);
+	bool bSetByCommandLine;
+	bSetByCommandLine = FParse::Value(FCommandLine::Get(), TEXT("SafeZonePaddingLeft="),   TitleSafePaddingSize.X);
+	bSetByCommandLine = FParse::Value(FCommandLine::Get(), TEXT("SafeZonePaddingRight="),  TitleSafePaddingSize.Y) || bSetByCommandLine;
+	bSetByCommandLine = FParse::Value(FCommandLine::Get(), TEXT("SafeZonePaddingTop="),    TitleSafePaddingSize.Z) || bSetByCommandLine;
+	bSetByCommandLine = FParse::Value(FCommandLine::Get(), TEXT("SafeZonePaddingBottom="), TitleSafePaddingSize.W) || bSetByCommandLine;
+
+	if (!bSetByCommandLine)
 	{
-		const float HalfUnsafeRatio = (1.0f - SafeZoneRatio) * 0.5f;
-		TitleSafePaddingSize = FVector4(PrimaryDisplayWidth * HalfUnsafeRatio, PrimaryDisplayHeight * HalfUnsafeRatio, PrimaryDisplayWidth * HalfUnsafeRatio, PrimaryDisplayHeight * HalfUnsafeRatio);
+		const float SafeZoneRatio = GetDebugTitleSafeZoneRatio();
+		if (SafeZoneRatio < 1.0f)
+		{
+			const float HalfUnsafeRatio = (1.0f - SafeZoneRatio) * 0.5f;
+			TitleSafePaddingSize = FVector4(PrimaryDisplayWidth * HalfUnsafeRatio, PrimaryDisplayHeight * HalfUnsafeRatio, PrimaryDisplayWidth * HalfUnsafeRatio, PrimaryDisplayHeight * HalfUnsafeRatio);
+		}
 	}
 
 	const float ActionSafeZoneRatio = GetDebugActionSafeZoneRatio();
diff --git a/Engine/Source/Runtime/Engine/Private/GameViewportClient.cpp b/Engine/Source/Runtime/Engine/Private/GameViewportClient.cpp
--- a/Engine/Source/Runtime/Engine/Private/GameViewportClient.cpp
+++ b/Engine/Source/Runtime/Engine/Private/GameViewportClient.cpp
@@ -2430,6 +2430,9 @@ void UGameViewportClient::DrawTitleSafeArea( UCanvas* Canvas )
 	FCanvasTileItem TileItem(FVector2D::ZeroVector, GWhiteTexture, UnsafeZoneColor);
 	TileItem.BlendMode = SE_BLEND_Translucent;
 
+	// Command line override used by mobile PIE.
+	static bool bDrawUnSafeZones = FParse::Param(FCommandLine::Get(), TEXT("DrawUnSafeZones"));
+
 	// CalculateSafeZoneValues() can be slow, so we only want to run it if we have boundaries to draw
 	if (FDisplayMetrics::GetDebugTitleSafeZoneRatio() < 1.f)
 	{
@@ -2455,7 +2458,7 @@ void UGameViewportClient::DrawTitleSafeArea( UCanvas* Canvas )
 		TileItem.Size = FVector2D(SafeZone.Right, HeightOfSides);
 		Canvas->DrawItem(TileItem);
 	}
-	else if (!FSlateApplication::Get().GetCustomSafeZone().GetDesiredSize().IsZero())
+	else if (!FSlateApplication::Get().GetCustomSafeZone().GetDesiredSize().IsZero() || bDrawUnSafeZones)
 	{
 		ULevelEditorPlaySettings* PlaySettings = GetMutableDefault<ULevelEditorPlaySettings>();
 		PlaySettings->CalculateCustomUnsafeZones(PlaySettings->CustomUnsafeZoneStarts, PlaySettings->CustomUnsafeZoneDimensions, PlaySettings->DeviceToEmulate, FVector2D(Width, Height));
