Storing current controller input before using it to prevent issues caused by the current controller being updated from another thread after we use it but before we store it.

diff --git a/Engine/Source/Runtime/ApplicationCore/Private/IOS/IOSInputInterface.cpp b/Engine/Source/Runtime/ApplicationCore/Private/IOS/IOSInputInterface.cpp
--- a/Engine/Source/Runtime/ApplicationCore/Private/IOS/IOSInputInterface.cpp
+++ b/Engine/Source/Runtime/ApplicationCore/Private/IOS/IOSInputInterface.cpp
@@ -385,9 +385,19 @@ void FIOSInputInterface::SendControllerEvents()
         if(!Controllers[i].bIsGamepadConnected) continue;
         
         GCController* Cont = Controllers[i].Controller;
-		GCExtendedGamepad* ExtendedGamepad = Cont.extendedGamepad;
+        GCExtendedGamepadSnapshot* ExtendedGamepad = nullptr;
+        static bool bSupportsGamepadCapture = [Cont respondsToSelector:@selector(capture)];
+        if (bSupportsGamepadCapture)
+        {
+            ExtendedGamepad = (GCExtendedGamepadSnapshot*)[Cont capture].extendedGamepad;
+        }
+        else
+        {
+            ExtendedGamepad = [Cont.extendedGamepad saveSnapshot];
+        }
+        
 #if PLATFORM_TVOS
-		GCMicroGamepad* MicroGamepad = Cont.microGamepad;
+        GCMicroGamepad* MicroGamepad = [Cont.microGamepad saveSnapshot];
 #endif
 		GCMotion* Motion = Cont.motion;
 
@@ -515,15 +525,7 @@ if ((Previous##Gamepad != nil && Gamepad.GCAxis.value != Previous##Gamepad.GCAxi
             }
 
             [Controller.PreviousExtendedGamepad release];
-            static bool bSupportsGamepadCapture = [Cont respondsToSelector:@selector(capture)];
-            if (bSupportsGamepadCapture)
-            {
-                Controller.PreviousExtendedGamepad = (GCExtendedGamepadSnapshot*)[ExtendedGamepad.controller capture].extendedGamepad;
-            }
-            else
-            {
-                Controller.PreviousExtendedGamepad = [ExtendedGamepad saveSnapshot];
-            }
+            Controller.PreviousExtendedGamepad = ExtendedGamepad;
             [Controller.PreviousExtendedGamepad retain];
 		}
 #if PLATFORM_TVOS
@@ -555,7 +557,7 @@ if ((Previous##Gamepad != nil && Gamepad.GCAxis.value != Previous##Gamepad.GCAxi
 		
 			
 			[Controller.PreviousMicroGamepad release];
-			Controller.PreviousMicroGamepad = [MicroGamepad saveSnapshot];
+			Controller.PreviousMicroGamepad = MicroGamepad;
 			[Controller.PreviousMicroGamepad retain];
         }
         
