remove WINDOWS_USE_FEATURE_APPLICATION and WINDOWS_USE_FEATURE_WINDOW

diff --git a/Engine/Source/Runtime/ApplicationCore/Private/Windows/WindowsApplication.cpp b/Engine/Source/Runtime/ApplicationCore/Private/Windows/WindowsApplication.cpp
--- a/Engine/Source/Runtime/ApplicationCore/Private/Windows/WindowsApplication.cpp
+++ b/Engine/Source/Runtime/ApplicationCore/Private/Windows/WindowsApplication.cpp
@@ -2,8 +2,6 @@
 
 #include "Windows/WindowsApplication.h"
 
-#if WINDOWS_USE_FEATURE_APPLICATION
-
 #include "Containers/StringConv.h"
 #include "CoreGlobals.h"
 #include "Internationalization/Text.h"
@@ -2823,4 +2821,3 @@ TSharedRef<FTaskbarList> FTaskbarList::Create()
 
 #include "Windows/HideWindowsPlatformTypes.h"
 
-#endif
diff --git a/Engine/Source/Runtime/ApplicationCore/Private/Windows/WindowsPlatformApplicationMisc.cpp b/Engine/Source/Runtime/ApplicationCore/Private/Windows/WindowsPlatformApplicationMisc.cpp
--- a/Engine/Source/Runtime/ApplicationCore/Private/Windows/WindowsPlatformApplicationMisc.cpp
+++ b/Engine/Source/Runtime/ApplicationCore/Private/Windows/WindowsPlatformApplicationMisc.cpp
@@ -72,7 +72,6 @@ class FFeedbackContext* FWindowsPlatformApplicationMisc::GetFeedbackContext()
 
 GenericApplication* FWindowsPlatformApplicationMisc::CreateApplication()
 {
-#if WINDOWS_USE_FEATURE_APPLICATION
 	HICON AppIconHandle = LoadIcon( hInstance, MAKEINTRESOURCE( GetAppIcon() ) );
 	if( AppIconHandle == NULL )
 	{
@@ -80,9 +79,6 @@ GenericApplication* FWindowsPlatformApplicationMisc::CreateApplication()
 	}
 
 	return FWindowsApplication::CreateWindowsApplication( hInstance, AppIconHandle );
-#else
-	return nullptr;
-#endif
 }
 
 void FWindowsPlatformApplicationMisc::RequestMinimize()
diff --git a/Engine/Source/Runtime/ApplicationCore/Private/Windows/WindowsWindow.cpp b/Engine/Source/Runtime/ApplicationCore/Private/Windows/WindowsWindow.cpp
--- a/Engine/Source/Runtime/ApplicationCore/Private/Windows/WindowsWindow.cpp
+++ b/Engine/Source/Runtime/ApplicationCore/Private/Windows/WindowsWindow.cpp
@@ -2,8 +2,6 @@
 
 #include "Windows/WindowsWindow.h"
 
-#if WINDOWS_USE_FEATURE_WINDOW
-
 #include "Math/UnrealMathUtility.h"
 #include "HAL/UnrealMemory.h"
 #include "Containers/UnrealString.h"
@@ -1182,5 +1180,3 @@ HRESULT STDCALL FWindowsWindow::Drop( __RPC__in_opt IDataObject *DataObjectPoint
 		return 0;
 	}
 }
-
-#endif
diff --git a/Engine/Source/Runtime/ApplicationCore/Public/Windows/WindowsApplication.h b/Engine/Source/Runtime/ApplicationCore/Public/Windows/WindowsApplication.h
--- a/Engine/Source/Runtime/ApplicationCore/Public/Windows/WindowsApplication.h
+++ b/Engine/Source/Runtime/ApplicationCore/Public/Windows/WindowsApplication.h
@@ -4,8 +4,6 @@
 
 #include "CoreTypes.h"
 
-#if WINDOWS_USE_FEATURE_APPLICATION
-
 #include "Math/Color.h"
 #include "HAL/IConsoleManager.h"
 #include "GenericPlatform/GenericApplication.h"
@@ -404,11 +402,11 @@ protected:
 	/** Processes deferred drag and drop operations. */
 	void ProcessDeferredDragDropOperation(const FDeferredWindowsDragDropOperation& Op);
 
-private:
-
 	/** Hidden constructor. */
 	FWindowsApplication( const HINSTANCE HInstance, const HICON IconHandle );
 
+private:
+
 	/** Registers the Windows class for windows and assigns the application instance and icon */
 	static bool RegisterClass( const HINSTANCE HInstance, const HICON HIcon );
 
@@ -532,4 +530,3 @@ private:
 
 PRAGMA_ENABLE_DEPRECATION_WARNINGS
 
-#endif
diff --git a/Engine/Source/Runtime/ApplicationCore/Public/Windows/WindowsWindow.h b/Engine/Source/Runtime/ApplicationCore/Public/Windows/WindowsWindow.h
--- a/Engine/Source/Runtime/ApplicationCore/Public/Windows/WindowsWindow.h
+++ b/Engine/Source/Runtime/ApplicationCore/Public/Windows/WindowsWindow.h
@@ -4,8 +4,6 @@
 
 #include "CoreTypes.h"
 
-#if WINDOWS_USE_FEATURE_WINDOW
-
 #include "GenericPlatform/GenericWindow.h"
 #include "Templates/SharedPointer.h"
 #include "Windows/WindowsHWrapper.h"
@@ -211,5 +209,3 @@ private:
 	/** when true the window is responsible for its own size adjustments in response to a DPI change */
 	bool bHandleManualDPIChanges = false;
 };
-
-#endif
\ No newline at end of file
diff --git a/Engine/Source/Runtime/Core/Public/Windows/WIndowsPlatform.h b/Engine/Source/Runtime/Core/Public/Windows/WIndowsPlatform.h
--- a/Engine/Source/Runtime/Core/Public/Windows/WIndowsPlatform.h
+++ b/Engine/Source/Runtime/Core/Public/Windows/WIndowsPlatform.h
@@ -87,8 +87,6 @@ typedef FWindowsPlatformTypes FPlatformTypes;
 #define WINDOWS_USE_FEATURE_APPLICATIONMISC_CLASS			1
 #define WINDOWS_USE_FEATURE_PLATFORMPROCESS_CLASS			1
 #define WINDOWS_USE_FEATURE_PLATFORMMISC_CLASS				1
-#define WINDOWS_USE_FEATURE_APPLICATION						1
-#define	WINDOWS_USE_FEATURE_WINDOW							1
 #define WINDOWS_USE_FEATURE_LAUNCH							1
 #define WINDOWS_USE_FEATURE_DYNAMIC_RHI						1
 
