Fix non-unity
#rnx

diff --git a/Engine/Source/Runtime/ApplicationCore/Private/Windows/WindowsApplication.cpp b/Engine/Source/Runtime/ApplicationCore/Private/Windows/WindowsApplication.cpp
--- a/Engine/Source/Runtime/ApplicationCore/Private/Windows/WindowsApplication.cpp
+++ b/Engine/Source/Runtime/ApplicationCore/Private/Windows/WindowsApplication.cpp
@@ -18,6 +18,7 @@
 #include "HAL/ThreadHeartBeat.h"
 #include "Templates/UniquePtr.h"
 #include "Windows/WindowsPlatformApplicationMisc.h"
+#include "Stats/Stats.h"
 
 #if WITH_EDITOR
 #include "Modules/ModuleManager.h"
