Fix Xcode 9 compile issue with new type NSWritingDirection. Using deprecated UITextWritingDirection to maintain compatibility for now.
#rnx

diff --git a/Engine/Source/Runtime/ApplicationCore/Private/IOS/IOSView.cpp b/Engine/Source/Runtime/ApplicationCore/Private/IOS/IOSView.cpp
--- a/Engine/Source/Runtime/ApplicationCore/Private/IOS/IOSView.cpp
+++ b/Engine/Source/Runtime/ApplicationCore/Private/IOS/IOSView.cpp
@@ -886,8 +886,10 @@ id<MTLDevice> GMetalDevice = nil;
 	return nil;
 }
 
+// UE 4.25 note: type UITextWritingDirection is deprecated and the new type NSWritingDirection used from Xcode 11 but we need to 
+// continue to use UITextWritingDirection for now in order to maintain compatibility with Xcode 9 and 10.
 
-- (NSWritingDirection)baseWritingDirectionForPosition:(UITextPosition *)position inDirection:(UITextStorageDirection)direction
+- (UITextWritingDirection)baseWritingDirectionForPosition:(UITextPosition *)position inDirection:(UITextStorageDirection)direction
 {
 	REPORT_EVENT;
 	// assume left to right for now
@@ -895,7 +897,7 @@ id<MTLDevice> GMetalDevice = nil;
 }
 
 
-- (void)setBaseWritingDirection:(NSWritingDirection)writingDirection forRange:(UITextRange *)range
+- (void)setBaseWritingDirection:(UITextWritingDirection)writingDirection forRange:(UITextRange *)range
 {
 	// @todo keyboard: This is called
 }
