Windows - When pumping messages if we get an SC_CLOSE command on windows during the slate threaded loading screen we now ignore it to prevent closing the app during a point that the application isn't safe to be shutdown without crashing most of the time.

#localization none
#tests PIE/ Editor -game

diff --git a/Engine/Source/Runtime/ApplicationCore/Private/Windows/WindowsApplication.cpp b/Engine/Source/Runtime/ApplicationCore/Private/Windows/WindowsApplication.cpp
--- a/Engine/Source/Runtime/ApplicationCore/Private/Windows/WindowsApplication.cpp
+++ b/Engine/Source/Runtime/ApplicationCore/Private/Windows/WindowsApplication.cpp
@@ -1501,8 +1501,8 @@ int32 FWindowsApplication::ProcessMessage( HWND hwnd, uint32 msg, WPARAM wParam,
 					break;
 				case SC_CLOSE:
 					{
-						// do not allow Alt-f4 during slow tasks.  This causes entry into the shutdown sequence at abonrmal times which causes crashes.
-						if (!GIsSlowTask)
+						// Do not allow Alt-f4 during slow tasks, or slate's loading thread.  This causes entry into the shutdown sequence at abnormal times which causes crashes.
+						if (!GIsSlowTask && GSlateLoadingThreadId == 0)
 						{
 							DeferMessage(CurrentNativeEventWindowPtr, hwnd, WM_CLOSE, 0, 0);
 						}
