Fixed conversion issue from fileSystemRepresentation -> TCHAR

fileSystemRepresentation is actually a UTF-8 (int8*) buffer so needs encoding via UTF8_TO_TCHAR

diff --git a/Engine/Source/Runtime/ApplicationCore/Private/Mac/MacApplication.cpp b/Engine/Source/Runtime/ApplicationCore/Private/Mac/MacApplication.cpp
--- a/Engine/Source/Runtime/ApplicationCore/Private/Mac/MacApplication.cpp
+++ b/Engine/Source/Runtime/ApplicationCore/Private/Mac/MacApplication.cpp
@@ -779,7 +779,7 @@ void FMacApplication::ProcessEvent(const FDeferredMacEvent& Event)
 				for (int32 Index = 0; Index < [Files count]; Index++)
 				{
 					NSString* FilePath = [Files objectAtIndex: Index];
-					const FString ListElement = FString([FilePath fileSystemRepresentation]);
+					const FString ListElement = UTF8_TO_TCHAR([FilePath fileSystemRepresentation]);
 					FileList.Add(ListElement);
 				}
 
@@ -795,7 +795,7 @@ void FMacApplication::ProcessEvent(const FDeferredMacEvent& Event)
 				for (int32 Index = 0; Index < [Files count]; Index++)
 				{
 					NSString* FilePath = [Files objectAtIndex: Index];
-					const FString ListElement = FString([FilePath fileSystemRepresentation]);
+					const FString ListElement = UTF8_TO_TCHAR([FilePath fileSystemRepresentation]);
 					FileList.Add(ListElement);
 				}
 
