Move log to verbose as it will now spam due to DPI being checked per tick for autosized widgets

diff --git a/Engine/Source/Runtime/ApplicationCore/Private/Linux/LinuxPlatformApplicationMisc.cpp b/Engine/Source/Runtime/ApplicationCore/Private/Linux/LinuxPlatformApplicationMisc.cpp
--- a/Engine/Source/Runtime/ApplicationCore/Private/Linux/LinuxPlatformApplicationMisc.cpp
+++ b/Engine/Source/Runtime/ApplicationCore/Private/Linux/LinuxPlatformApplicationMisc.cpp
@@ -499,7 +499,7 @@ float FLinuxPlatformApplicationMisc::GetDPIScaleFactorAtPoint(float X, float Y)
 				if (SDL_GetDisplayDPI(Idx, nullptr, &HorzDPI, &VertDPI) == 0)
 				{
 					float Scale = LinuxPlatformApplicationMisc::QuantizeScale((HorzDPI + VertDPI) / 192.0f);	// average between two scales (divided by 96.0f)
-					UE_LOG(LogLinux, Log, TEXT("Scale at X=%f, Y=%f: %f (monitor=#%d, HDPI=%f (horz scale: %f), VDPI=%f (vert scale: %f))"), X, Y, Scale, Idx, HorzDPI, HorzDPI / 96.0f, VertDPI, VertDPI / 96.0f);
+					UE_LOG(LogLinux, Verbose, TEXT("Scale at X=%f, Y=%f: %f (monitor=#%d, HDPI=%f (horz scale: %f), VDPI=%f (vert scale: %f))"), X, Y, Scale, Idx, HorzDPI, HorzDPI / 96.0f, VertDPI, VertDPI / 96.0f);
 					return Scale;
 				}
 				else
