Added a cvar and some hacks to work around raw input not working over remote desktop causing the cursor to get stuck at the edge of the screen
Disabled by default until more testing is done

diff --git a/Engine/Source/Runtime/ApplicationCore/Private/Windows/WindowsApplication.cpp b/Engine/Source/Runtime/ApplicationCore/Private/Windows/WindowsApplication.cpp
--- a/Engine/Source/Runtime/ApplicationCore/Private/Windows/WindowsApplication.cpp
+++ b/Engine/Source/Runtime/ApplicationCore/Private/Windows/WindowsApplication.cpp
@@ -71,6 +71,14 @@ FAutoConsoleVariableRef	CVarPreventDuplicateMouseEventsForTouch(
 	TEXT("Hack to get around multiple mouse events being triggered for touch events on Windows 7 and lower.  Enabling this will prevent pen tablets from working on windows 7 since until we switch to the windows 8 sdk (and can use WM_POINTER* events) we cannot detect the difference")
 );
 
+
+static int32 EnableRawInputSimulationOverRDP = false;
+FAutoConsoleVariableRef	CVarEnableRawInputSimulationOverRDP(
+	TEXT("Slate.EnableRawInputSimulationOverRDP"),
+	EnableRawInputSimulationOverRDP,
+	TEXT("")
+);
+
 const FIntPoint FWindowsApplication::MinimizedWindowPosition(-32000,-32000);
 
 FWindowsApplication* WindowsApplication = nullptr;
@@ -427,36 +435,58 @@ void* FWindowsApplication::GetCapture( void ) const
 
 void FWindowsApplication::SetHighPrecisionMouseMode( const bool Enable, const TSharedPtr< FGenericWindow >& InWindow )
 {
-	HWND hwnd = NULL;
-	DWORD flags = RIDEV_REMOVE;
-	bUsingHighPrecisionMouseInput = Enable;
-
-	if ( Enable )
+	if (EnableRawInputSimulationOverRDP && FPlatformMisc::IsRemoteSession())
 	{
-		flags = 0;
+		if(Enable)
+		{
+			bUsingHighPrecisionMouseInput = true;
+			bSimulatingHighPrecisionMouseInputForRDP = true;
+
+			POINT CursorPos;
+			BOOL bGotPoint = ::GetCursorPos(&CursorPos);
 
-		if ( InWindow.IsValid() )
+			CachedPreHighPrecisionMousePosForRDP = FIntPoint(CursorPos.x, CursorPos.y);
+		}
+		else
 		{
-			hwnd = (HWND)InWindow->GetOSWindowHandle(); 
+			CachedPreHighPrecisionMousePosForRDP = FIntPoint(INT_MAX, INT_MAX);
+			bSimulatingHighPrecisionMouseInputForRDP = false;
 		}
 	}
 
-	// NOTE: Currently has to be created every time due to conflicts with Direct8 Input used by the wx unrealed
-	RAWINPUTDEVICE RawInputDevice;
+	{
+		
+		HWND hwnd = NULL;
+		DWORD flags = RIDEV_REMOVE;
+		bUsingHighPrecisionMouseInput = Enable;
+
+		if (Enable)
+		{
+			flags = 0;
+
+			if (InWindow.IsValid())
+			{
+				hwnd = (HWND)InWindow->GetOSWindowHandle();
+			}
+		}
+
+		// NOTE: Currently has to be created every time due to conflicts with Direct8 Input used by the wx unrealed
+		RAWINPUTDEVICE RawInputDevice;
 
-	//The HID standard for mouse
-	const uint16 StandardMouse = 0x02;
+		//The HID standard for mouse
+		const uint16 StandardMouse = 0x02;
 
-	RawInputDevice.usUsagePage = 0x01; 
-	RawInputDevice.usUsage = StandardMouse;
-	RawInputDevice.dwFlags = flags;
+		RawInputDevice.usUsagePage = 0x01;
+		RawInputDevice.usUsage = StandardMouse;
+		RawInputDevice.dwFlags = flags;
 
-	// Process input for just the window that requested it.  NOTE: If we pass NULL here events are routed to the window with keyboard focus
-	// which is not always known at the HWND level with Slate
-	RawInputDevice.hwndTarget = hwnd; 
+		// Process input for just the window that requested it.  NOTE: If we pass NULL here events are routed to the window with keyboard focus
+		// which is not always known at the HWND level with Slate
+		RawInputDevice.hwndTarget = hwnd;
 
-	// Register the raw input device
-	::RegisterRawInputDevices( &RawInputDevice, 1, sizeof( RAWINPUTDEVICE ) );
+		// Register the raw input device
+		::RegisterRawInputDevices(&RawInputDevice, 1, sizeof(RAWINPUTDEVICE));
+	}
 }
 
 FPlatformRect FWindowsApplication::GetWorkArea( const FPlatformRect& CurrentWindow ) const
@@ -996,9 +1026,16 @@ int32 FWindowsApplication::ProcessMessage( HWND hwnd, uint32 msg, WPARAM wParam,
 						const bool IsAbsoluteInput = (Raw->data.mouse.usFlags & MOUSE_MOVE_ABSOLUTE) == MOUSE_MOVE_ABSOLUTE;
 						if( IsAbsoluteInput )
 						{
-							// Since the raw input is coming in as absolute it is likely the user is using a tablet
-							// or perhaps is interacting through a virtual desktop
-							DeferMessage( CurrentNativeEventWindowPtr, hwnd, msg, wParam, lParam, 0, 0, MOUSE_MOVE_ABSOLUTE );
+							if (bSimulatingHighPrecisionMouseInputForRDP)
+							{
+								DeferMessage(CurrentNativeEventWindowPtr, hwnd, WM_MOUSEMOVE, wParam, lParam);
+							}
+							else
+							{
+								// Since the raw input is coming in as absolute it is likely the user is using a tablet
+								// or perhaps is interacting through a virtual desktop
+								DeferMessage(CurrentNativeEventWindowPtr, hwnd, msg, wParam, lParam, 0, 0, MOUSE_MOVE_ABSOLUTE);
+							}
 							return 1;
 						}
 
@@ -1999,6 +2036,17 @@ int32 FWindowsApplication::ProcessDeferredMessage( const FDeferredWindowsMessage
 		case WM_MOUSEMOVE:
 			{
 				BOOL Result = false;
+				if (bSimulatingHighPrecisionMouseInputForRDP)
+				{
+					POINT CursorPoint;
+					::GetCursorPos(&CursorPoint);
+					int32 DeltaX = CursorPoint.x - CachedPreHighPrecisionMousePosForRDP.X;
+					int32 DeltaY = CursorPoint.y - CachedPreHighPrecisionMousePosForRDP.Y;
+
+					MessageHandler->OnRawMouseMove(DeltaX, DeltaY);
+
+					Result = true;
+				}
 				if( !bUsingHighPrecisionMouseInput )
 				{
 					Result = MessageHandler->OnMouseMove();
@@ -2561,6 +2609,21 @@ void FWindowsApplication::AddExternalInputDevice(TSharedPtr<IInputDevice> InputD
 	}
 }
 
+void FWindowsApplication::FinishedInputThisFrame()
+{
+	if (bSimulatingHighPrecisionMouseInputForRDP)
+	{
+		MessageHandler->SetCursorPos(FVector2D(CachedPreHighPrecisionMousePosForRDP));
+		//::SetCursorPos(CachedPreHighPrecisionMousePosForRDP.X, CachedPreHighPrecisionMousePosForRDP.Y);
+
+		// SetCursorPos sends a windows message, we dont want to process it as real input so get rid of it
+		// Todo needed?  It is theoretically possible this could remove valid WM_MOUSEMOVE messages that came right after we finished 
+		// processing input
+		MSG msg;
+		while (::PeekMessage(&msg, NULL, WM_MOUSEMOVE, WM_MOUSEMOVE, PM_REMOVE)) {};
+	}
+}
+
 TSharedPtr<FTaskbarList> FWindowsApplication::GetTaskbarList()
 {
 	return TaskbarList;
diff --git a/Engine/Source/Runtime/ApplicationCore/Private/Windows/WindowsCursor.cpp b/Engine/Source/Runtime/ApplicationCore/Private/Windows/WindowsCursor.cpp
--- a/Engine/Source/Runtime/ApplicationCore/Private/Windows/WindowsCursor.cpp
+++ b/Engine/Source/Runtime/ApplicationCore/Private/Windows/WindowsCursor.cpp
@@ -29,6 +29,13 @@ FWindowsCursor::FWindowsCursor()
 		switch( CurCursorIndex )
 		{
 		case EMouseCursor::None:
+			if (FPlatformMisc::IsRemoteSession())
+			{
+				// during remote sessions we rely on constantly resetting the mouse cursor position for infinite mouse deltas while dragging (WM_INPUT doesnt work over rdp)
+				// this hack here is to avoid having a hidden cursor which prevents SetCursorPos from working.  This is just a completely transparent cursor yet "visible" cursor to work around that
+				CursorHandle = LoadCursorFromFile((LPCTSTR) * (FString(FPlatformProcess::BaseDir()) / FString::Printf(TEXT("%sEditor/Slate/Cursor/invisible.cur"), *FPaths::EngineContentDir())));
+			}
+			break;
 		case EMouseCursor::Custom:
 			// The mouse cursor will not be visible when None is used
 			break;
diff --git a/Engine/Source/Runtime/ApplicationCore/Public/GenericPlatform/GenericApplication.h b/Engine/Source/Runtime/ApplicationCore/Public/GenericPlatform/GenericApplication.h
--- a/Engine/Source/Runtime/ApplicationCore/Public/GenericPlatform/GenericApplication.h
+++ b/Engine/Source/Runtime/ApplicationCore/Public/GenericPlatform/GenericApplication.h
@@ -551,6 +551,7 @@ public:
 
 	virtual bool IsAllowedToRender() const { return true; }
 
+	virtual void FinishedInputThisFrame() {}
 public:
 
 	const TSharedPtr< ICursor > Cursor;
diff --git a/Engine/Source/Runtime/ApplicationCore/Public/GenericPlatform/GenericApplicationMessageHandler.h b/Engine/Source/Runtime/ApplicationCore/Public/GenericPlatform/GenericApplicationMessageHandler.h
--- a/Engine/Source/Runtime/ApplicationCore/Public/GenericPlatform/GenericApplicationMessageHandler.h
+++ b/Engine/Source/Runtime/ApplicationCore/Public/GenericPlatform/GenericApplicationMessageHandler.h
@@ -470,4 +470,9 @@ public:
 	{
 		return true;
 	}
+
+	virtual void SetCursorPos(const FVector2D& MouseCoordinate)
+	{
+
+	}
 };
diff --git a/Engine/Source/Runtime/ApplicationCore/Public/Windows/WindowsApplication.h b/Engine/Source/Runtime/ApplicationCore/Public/Windows/WindowsApplication.h
--- a/Engine/Source/Runtime/ApplicationCore/Public/Windows/WindowsApplication.h
+++ b/Engine/Source/Runtime/ApplicationCore/Public/Windows/WindowsApplication.h
@@ -378,6 +378,7 @@ public:
 
 	virtual void AddExternalInputDevice(TSharedPtr<class IInputDevice> InputDevice);
 
+	virtual void FinishedInputThisFrame() override;
 public:
 
 	// IInputInterface overrides
@@ -525,6 +526,11 @@ private:
 	/** Maps touch indexes to windows touch IDs. */
 	TArray<TOptional<int32>> TouchIDs;
 #endif
+
+	bool bSimulatingHighPrecisionMouseInputForRDP = false;
+
+	FIntPoint CachedPreHighPrecisionMousePosForRDP = FIntPoint::ZeroValue;
+
 };
 
 
diff --git a/Engine/Source/Runtime/Slate/Private/Framework/Application/SlateApplication.cpp b/Engine/Source/Runtime/Slate/Private/Framework/Application/SlateApplication.cpp
--- a/Engine/Source/Runtime/Slate/Private/Framework/Application/SlateApplication.cpp
+++ b/Engine/Source/Runtime/Slate/Private/Framework/Application/SlateApplication.cpp
@@ -1271,6 +1271,8 @@ void FSlateApplication::FinishedInputThisFrame()
 {
 	const float DeltaTime = GetDeltaTime();
 
+	PlatformApplication->FinishedInputThisFrame();
+	
 	// Any preprocessors are given a chance to process accumulated values (or do whatever other tick things they want)
 	// after we've finished processing all of the input for the frame
 	if (PlatformApplication->Cursor.IsValid())
