Turn on the remote desktop raw input workaround by default

diff --git a/Engine/Source/Runtime/ApplicationCore/Private/Windows/WindowsApplication.cpp b/Engine/Source/Runtime/ApplicationCore/Private/Windows/WindowsApplication.cpp
--- a/Engine/Source/Runtime/ApplicationCore/Private/Windows/WindowsApplication.cpp
+++ b/Engine/Source/Runtime/ApplicationCore/Private/Windows/WindowsApplication.cpp
@@ -72,7 +72,7 @@ FAutoConsoleVariableRef	CVarPreventDuplicateMouseEventsForTouch(
 );
 
 
-static int32 EnableRawInputSimulationOverRDP = false;
+static int32 EnableRawInputSimulationOverRDP = true;
 FAutoConsoleVariableRef	CVarEnableRawInputSimulationOverRDP(
 	TEXT("Slate.EnableRawInputSimulationOverRDP"),
 	EnableRawInputSimulationOverRDP,
