Further adjustments to simulated raw input over RDP

diff --git a/Engine/Source/Runtime/ApplicationCore/Private/Windows/WindowsApplication.cpp b/Engine/Source/Runtime/ApplicationCore/Private/Windows/WindowsApplication.cpp
--- a/Engine/Source/Runtime/ApplicationCore/Private/Windows/WindowsApplication.cpp
+++ b/Engine/Source/Runtime/ApplicationCore/Private/Windows/WindowsApplication.cpp
@@ -444,8 +444,9 @@ void FWindowsApplication::SetHighPrecisionMouseMode( const bool Enable, const TS
 
 			POINT CursorPos;
 			BOOL bGotPoint = ::GetCursorPos(&CursorPos);
-
+		
 			CachedPreHighPrecisionMousePosForRDP = FIntPoint(CursorPos.x, CursorPos.y);
+			LastCursorPoint = CachedPreHighPrecisionMousePosForRDP;
 		}
 		else
 		{
@@ -2036,18 +2037,22 @@ int32 FWindowsApplication::ProcessDeferredMessage( const FDeferredWindowsMessage
 		case WM_MOUSEMOVE:
 			{
 				BOOL Result = false;
-				if (bSimulatingHighPrecisionMouseInputForRDP)
+				if (bUsingHighPrecisionMouseInput && bSimulatingHighPrecisionMouseInputForRDP)
 				{
 					POINT CursorPoint;
 					::GetCursorPos(&CursorPoint);
-					int32 DeltaX = CursorPoint.x - CachedPreHighPrecisionMousePosForRDP.X;
-					int32 DeltaY = CursorPoint.y - CachedPreHighPrecisionMousePosForRDP.Y;
+
+					int32 DeltaX = CursorPoint.x - LastCursorPoint.X;
+					int32 DeltaY = CursorPoint.y - LastCursorPoint.Y;
+					
+					LastCursorPoint.X = CursorPoint.x;
+					LastCursorPoint.Y = CursorPoint.y;
 
 					MessageHandler->OnRawMouseMove(DeltaX, DeltaY);
 
 					Result = true;
 				}
-				if( !bUsingHighPrecisionMouseInput )
+				else if( !bUsingHighPrecisionMouseInput )
 				{
 					Result = MessageHandler->OnMouseMove();
 				}
@@ -2614,6 +2619,7 @@ void FWindowsApplication::FinishedInputThisFrame()
 	if (bSimulatingHighPrecisionMouseInputForRDP)
 	{
 		MessageHandler->SetCursorPos(FVector2D(CachedPreHighPrecisionMousePosForRDP));
+		LastCursorPoint = CachedPreHighPrecisionMousePosForRDP;
 		//::SetCursorPos(CachedPreHighPrecisionMousePosForRDP.X, CachedPreHighPrecisionMousePosForRDP.Y);
 
 		// SetCursorPos sends a windows message, we dont want to process it as real input so get rid of it
diff --git a/Engine/Source/Runtime/ApplicationCore/Public/Windows/WindowsApplication.h b/Engine/Source/Runtime/ApplicationCore/Public/Windows/WindowsApplication.h
--- a/Engine/Source/Runtime/ApplicationCore/Public/Windows/WindowsApplication.h
+++ b/Engine/Source/Runtime/ApplicationCore/Public/Windows/WindowsApplication.h
@@ -530,7 +530,7 @@ private:
 	bool bSimulatingHighPrecisionMouseInputForRDP = false;
 
 	FIntPoint CachedPreHighPrecisionMousePosForRDP = FIntPoint::ZeroValue;
-
+	FIntPoint LastCursorPoint = FIntPoint::ZeroValue;
 };
 
 
