Fix for some large delta cases with RDP simulated high DPI.

diff --git a/Engine/Source/Runtime/ApplicationCore/Private/Windows/WindowsApplication.cpp b/Engine/Source/Runtime/ApplicationCore/Private/Windows/WindowsApplication.cpp
--- a/Engine/Source/Runtime/ApplicationCore/Private/Windows/WindowsApplication.cpp
+++ b/Engine/Source/Runtime/ApplicationCore/Private/Windows/WindowsApplication.cpp
@@ -110,6 +110,10 @@ FWindowsApplication::FWindowsApplication( const HINSTANCE HInstance, const HICON
 #if WITH_ACCESSIBILITY
 	, UIAManager(new FWindowsUIAManager(*this))
 #endif
+	, bSimulatingHighPrecisionMouseInputForRDP(false)
+	, CachedPreHighPrecisionMousePosForRDP(FIntPoint::ZeroValue)
+	, LastCursorPoint(FIntPoint::ZeroValue)
+	, ClipCursorRect()
 
 {
 	FMemory::Memzero(ModifierKeyState, EModifierKey::Count);
@@ -444,6 +448,10 @@ void FWindowsApplication::SetHighPrecisionMouseMode( const bool Enable, const TS
 
 			POINT CursorPos;
 			BOOL bGotPoint = ::GetCursorPos(&CursorPos);
+
+			::GetClipCursor(&ClipCursorRect);
+
+			//UE_LOG(LogWindowsDesktop, Log, TEXT("Entering High Precision to Top: %d Bottom: %d Left: %d Right: %d"), ClipCursorRect.top, ClipCursorRect.bottom, ClipCursorRect.left, ClipCursorRect.right);
 		
 			CachedPreHighPrecisionMousePosForRDP = FIntPoint(CursorPos.x, CursorPos.y);
 			LastCursorPoint = CachedPreHighPrecisionMousePosForRDP;
@@ -1028,8 +1036,83 @@ int32 FWindowsApplication::ProcessMessage( HWND hwnd, uint32 msg, WPARAM wParam,
 						if( IsAbsoluteInput )
 						{
 							if (bSimulatingHighPrecisionMouseInputForRDP)
-							{
-								DeferMessage(CurrentNativeEventWindowPtr, hwnd, WM_MOUSEMOVE, wParam, lParam);
+							{	
+								// Get the new cursor position
+								POINT CursorPoint;
+								int32 Top = 0;
+								int32 Left = 0;
+								int32 Width = 0;
+								int32 Height = 0;
+								const bool IsVirtualScreen = (Raw->data.mouse.usFlags & MOUSE_VIRTUAL_DESKTOP) == MOUSE_VIRTUAL_DESKTOP;
+								if (IsVirtualScreen)
+								{
+									// This is used to make Remote Desktop sessons work
+									Width = GetSystemMetrics(SM_CXVIRTUALSCREEN);
+									Height = GetSystemMetrics(SM_CYVIRTUALSCREEN);
+
+									CursorPoint.x = static_cast<int>((float(Raw->data.mouse.lLastX) / 65535.0f) * Width);
+									CursorPoint.y = static_cast<int>((float(Raw->data.mouse.lLastY) / 65535.0f) * Height);
+								}
+								else
+								{
+									Top = ClipCursorRect.top;
+									Left = ClipCursorRect.left;
+									Height = ClipCursorRect.bottom - ClipCursorRect.top;
+									Width = ClipCursorRect.right - ClipCursorRect.left;
+									::GetCursorPos(&CursorPoint);
+								}
+
+								// Calculate the cursor delta from the last position
+								// We do this prior to and wrapping for continuous input as that would mean an incorrect delta we'd just ignore
+								const int32 DeltaX = CursorPoint.x - LastCursorPoint.X;
+								const int32 DeltaY = CursorPoint.y - LastCursorPoint.Y;
+
+								// Wrap and set cursor position in necessary
+								const int32 TopEdge = Top + int32(0.1f * float(Height));
+								const int32 BottomEdge = Top + int32(0.9f * float(Height));
+								const int32 LeftEdge = Left + int32(0.1f * float(Width));
+								const int32 RightEdge = Left + int32(0.9f * float(Width));
+
+								bool bSet = false;
+								if (CursorPoint.y < TopEdge)			{ CursorPoint.y = BottomEdge;	bSet = true; }
+								else if (CursorPoint.y > BottomEdge)	{ CursorPoint.y = TopEdge;		bSet = true; }
+
+								if (CursorPoint.x < LeftEdge)			{ CursorPoint.x = RightEdge;	bSet = true; }
+								else if (CursorPoint.x > RightEdge)		{ CursorPoint.x = LeftEdge;		bSet = true; }
+
+								if (bSet)
+								{
+									//UE_LOG(LogWindowsDesktop, Log, TEXT("Wrapping Cursor to X: %d Y: %d"), CursorPoint.x, CursorPoint.y);
+									MessageHandler->SetCursorPos(FVector2D(LastCursorPoint));
+									LastCursorPoint.X = CursorPoint.x;
+									LastCursorPoint.Y = CursorPoint.y;
+									return 1;
+								}
+
+								const int32 DeltaWidthMax = (int32)((float)(ClipCursorRect.right - ClipCursorRect.left) * 0.4f);
+								const int32 DeltaHeightMax = (int32)((float)(ClipCursorRect.bottom - ClipCursorRect.top) * 0.4f);
+								/*
+								if (DeltaX != 0 && DeltaY != 0)
+								{
+									if (FMath::Abs(DeltaX) < DeltaWidthMax && FMath::Abs(DeltaY) < DeltaHeightMax)
+									{
+										UE_LOG(LogWindowsDesktop, Log, TEXT("Accept Delta X: %d Y: %d  ---- Last X: %d Y: %d ---- Cur X: %d Y: %d ---- Width: %d Height: %d"), DeltaX, DeltaY, LastCursorPoint.X, LastCursorPoint.Y, CursorPoint.x, CursorPoint.y, (int32)Width, (int32)Height);
+									}
+									else
+									{
+										UE_LOG(LogWindowsDesktop, Log, TEXT("IGNORE Delta X: %d Y: %d  ---- Last X: %d Y: %d ---- Cur X: %d Y: %d ---- Width: %d Height: %d"), DeltaX, DeltaY, LastCursorPoint.X, LastCursorPoint.Y, CursorPoint.x, CursorPoint.y, (int32)Width, (int32)Height);
+									}
+								}*/
+
+								LastCursorPoint.X = CursorPoint.x;
+								LastCursorPoint.Y = CursorPoint.y;
+
+								// Send a delta assuming it's not zero or beyond our max delta 
+								if (DeltaX != 0 && DeltaY != 0 && FMath::Abs(DeltaX) < DeltaWidthMax && FMath::Abs(DeltaY) < DeltaHeightMax)
+								{
+									DeferMessage(CurrentNativeEventWindowPtr, hwnd, msg, wParam, lParam, DeltaX, DeltaY, MOUSE_MOVE_RELATIVE);
+								}
+								return 1;
 							}
 							else
 							{
@@ -2020,7 +2103,7 @@ int32 FWindowsApplication::ProcessDeferredMessage( const FDeferredWindowsMessage
 			{
 				if( DeferredMessage.RawInputFlags == MOUSE_MOVE_RELATIVE )
 				{
-					MessageHandler->OnRawMouseMove( DeferredMessage.X, DeferredMessage.Y );
+					MessageHandler->OnRawMouseMove(DeferredMessage.X, DeferredMessage.Y);
 				}
 				else
 				{
@@ -2037,22 +2120,7 @@ int32 FWindowsApplication::ProcessDeferredMessage( const FDeferredWindowsMessage
 		case WM_MOUSEMOVE:
 			{
 				BOOL Result = false;
-				if (bUsingHighPrecisionMouseInput && bSimulatingHighPrecisionMouseInputForRDP)
-				{
-					POINT CursorPoint;
-					::GetCursorPos(&CursorPoint);
-
-					int32 DeltaX = CursorPoint.x - LastCursorPoint.X;
-					int32 DeltaY = CursorPoint.y - LastCursorPoint.Y;
-					
-					LastCursorPoint.X = CursorPoint.x;
-					LastCursorPoint.Y = CursorPoint.y;
-
-					MessageHandler->OnRawMouseMove(DeltaX, DeltaY);
-
-					Result = true;
-				}
-				else if( !bUsingHighPrecisionMouseInput )
+				if (!bUsingHighPrecisionMouseInput)
 				{
 					Result = MessageHandler->OnMouseMove();
 				}
@@ -2616,18 +2684,7 @@ void FWindowsApplication::AddExternalInputDevice(TSharedPtr<IInputDevice> InputD
 
 void FWindowsApplication::FinishedInputThisFrame()
 {
-	if (bSimulatingHighPrecisionMouseInputForRDP)
-	{
-		MessageHandler->SetCursorPos(FVector2D(CachedPreHighPrecisionMousePosForRDP));
-		LastCursorPoint = CachedPreHighPrecisionMousePosForRDP;
-		//::SetCursorPos(CachedPreHighPrecisionMousePosForRDP.X, CachedPreHighPrecisionMousePosForRDP.Y);
 
-		// SetCursorPos sends a windows message, we dont want to process it as real input so get rid of it
-		// Todo needed?  It is theoretically possible this could remove valid WM_MOUSEMOVE messages that came right after we finished 
-		// processing input
-		MSG msg;
-		while (::PeekMessage(&msg, NULL, WM_MOUSEMOVE, WM_MOUSEMOVE, PM_REMOVE)) {};
-	}
 }
 
 TSharedPtr<FTaskbarList> FWindowsApplication::GetTaskbarList()
diff --git a/Engine/Source/Runtime/ApplicationCore/Public/Windows/WindowsApplication.h b/Engine/Source/Runtime/ApplicationCore/Public/Windows/WindowsApplication.h
--- a/Engine/Source/Runtime/ApplicationCore/Public/Windows/WindowsApplication.h
+++ b/Engine/Source/Runtime/ApplicationCore/Public/Windows/WindowsApplication.h
@@ -527,10 +527,11 @@ private:
 	TArray<TOptional<int32>> TouchIDs;
 #endif
 
-	bool bSimulatingHighPrecisionMouseInputForRDP = false;
+	bool bSimulatingHighPrecisionMouseInputForRDP;
 
-	FIntPoint CachedPreHighPrecisionMousePosForRDP = FIntPoint::ZeroValue;
-	FIntPoint LastCursorPoint = FIntPoint::ZeroValue;
+	FIntPoint CachedPreHighPrecisionMousePosForRDP;
+	FIntPoint LastCursorPoint;
+	RECT ClipCursorRect;
 };
 
 
