Add ability to force on raw input simulation over RDP for testing

#rnx

diff --git a/Engine/Source/Runtime/ApplicationCore/Private/Windows/WindowsApplication.cpp b/Engine/Source/Runtime/ApplicationCore/Private/Windows/WindowsApplication.cpp
--- a/Engine/Source/Runtime/ApplicationCore/Private/Windows/WindowsApplication.cpp
+++ b/Engine/Source/Runtime/ApplicationCore/Private/Windows/WindowsApplication.cpp
@@ -72,6 +72,7 @@ FAutoConsoleVariableRef	CVarPreventDuplicateMouseEventsForTouch(
 );
 
 
+#if !UE_BUILD_SHIPPING
 static int32 EnableRawInputSimulationOverRDP = true;
 FAutoConsoleVariableRef	CVarEnableRawInputSimulationOverRDP(
 	TEXT("Slate.EnableRawInputSimulationOverRDP"),
@@ -79,10 +80,26 @@ FAutoConsoleVariableRef	CVarEnableRawInputSimulationOverRDP(
 	TEXT("")
 );
 
+static int32 ForceRawInputSimulation = false;
+FAutoConsoleVariableRef	CVarForceRawInputSimulation(
+	TEXT("Slate.ForceRawInputSimulation"),
+	ForceRawInputSimulation,
+	TEXT("")
+);
+#else
+static int32 ForceRawInputSimulation = false;
+static int32 EnableRawInputSimulationOverRDP = false;
+#endif
+
 const FIntPoint FWindowsApplication::MinimizedWindowPosition(-32000,-32000);
 
 FWindowsApplication* WindowsApplication = nullptr;
 
+static bool ShouldSimulateRawInput()
+{
+	return ForceRawInputSimulation || (EnableRawInputSimulationOverRDP && FPlatformMisc::IsRemoteSession());
+}
+
 FWindowsApplication* FWindowsApplication::CreateWindowsApplication( const HINSTANCE InstanceHandle, const HICON IconHandle )
 {
 	WindowsApplication = new FWindowsApplication( InstanceHandle, IconHandle );
@@ -439,7 +456,7 @@ void* FWindowsApplication::GetCapture( void ) const
 
 void FWindowsApplication::SetHighPrecisionMouseMode( const bool Enable, const TSharedPtr< FGenericWindow >& InWindow )
 {
-	if (EnableRawInputSimulationOverRDP && FPlatformMisc::IsRemoteSession())
+	if (ShouldSimulateRawInput())
 	{
 		if(Enable)
 		{
