For a period after wrapping i track both spaces and accept old data as deltas. This means no user input is thown away any longer.

diff --git a/Engine/Source/Runtime/ApplicationCore/Private/Windows/WindowsApplication.cpp b/Engine/Source/Runtime/ApplicationCore/Private/Windows/WindowsApplication.cpp
--- a/Engine/Source/Runtime/ApplicationCore/Private/Windows/WindowsApplication.cpp
+++ b/Engine/Source/Runtime/ApplicationCore/Private/Windows/WindowsApplication.cpp
@@ -130,6 +130,8 @@ FWindowsApplication::FWindowsApplication( const HINSTANCE HInstance, const HICON
 	, bSimulatingHighPrecisionMouseInputForRDP(false)
 	, CachedPreHighPrecisionMousePosForRDP(FIntPoint::ZeroValue)
 	, LastCursorPoint(FIntPoint::ZeroValue)
+	, LastCursorPointPreWrap(FIntPoint::ZeroValue)
+	, NumPreWrapMsgsToRespect(0)
 	, ClipCursorRect()
 
 {
@@ -472,6 +474,9 @@ void FWindowsApplication::SetHighPrecisionMouseMode( const bool Enable, const TS
 		
 			CachedPreHighPrecisionMousePosForRDP = FIntPoint(CursorPos.x, CursorPos.y);
 			LastCursorPoint = CachedPreHighPrecisionMousePosForRDP;
+			
+			LastCursorPointPreWrap = FIntPoint::ZeroValue;
+			NumPreWrapMsgsToRespect = 0;
 		}
 		else
 		{
@@ -1050,7 +1055,7 @@ int32 FWindowsApplication::ProcessMessage( HWND hwnd, uint32 msg, WPARAM wParam,
 					if (Raw->header.dwType == RIM_TYPEMOUSE) 
 					{
 						const bool IsAbsoluteInput = (Raw->data.mouse.usFlags & MOUSE_MOVE_ABSOLUTE) == MOUSE_MOVE_ABSOLUTE;
-						if( IsAbsoluteInput )
+						if (IsAbsoluteInput)
 						{
 							if (bSimulatingHighPrecisionMouseInputForRDP)
 							{	
@@ -1074,55 +1079,92 @@ int32 FWindowsApplication::ProcessMessage( HWND hwnd, uint32 msg, WPARAM wParam,
 								{
 									Top = ClipCursorRect.top;
 									Left = ClipCursorRect.left;
-									Height = ClipCursorRect.bottom - ClipCursorRect.top;
 									Width = ClipCursorRect.right - ClipCursorRect.left;
+									Height = ClipCursorRect.bottom - ClipCursorRect.top;
 									::GetCursorPos(&CursorPoint);
 								}
 
+								const int32 DeltaWidthMax = (int32)((float)Width * 0.4f);
+								const int32 DeltaHeightMax = (int32)((float)Height * 0.4f);
+
+								const POINT CursorPointNoWrap = CursorPoint;
+
+								const bool bCanAcceptPreWrapMsg = NumPreWrapMsgsToRespect > 0;
+								NumPreWrapMsgsToRespect--;
+
 								// Calculate the cursor delta from the last position
 								// We do this prior to and wrapping for continuous input as that would mean an incorrect delta we'd just ignore
-								const int32 DeltaX = CursorPoint.x - LastCursorPoint.X;
-								const int32 DeltaY = CursorPoint.y - LastCursorPoint.Y;
+								int32 DeltaX = CursorPoint.x - LastCursorPoint.X;
+								int32 DeltaY = CursorPoint.y - LastCursorPoint.Y;
 
-								// Wrap and set cursor position in necessary
-								const int32 TopEdge = Top + int32(0.1f * float(Height));
-								const int32 BottomEdge = Top + int32(0.9f * float(Height));
-								const int32 LeftEdge = Left + int32(0.1f * float(Width));
-								const int32 RightEdge = Left + int32(0.9f * float(Width));
+								bool bAcceptingPreWrapDelta = false;
+								if (bCanAcceptPreWrapMsg)
+								{
+									const int32 DeltaXPreWrap = CursorPoint.x - LastCursorPointPreWrap.X;
+									const int32 DeltaYPreWrap = CursorPoint.y - LastCursorPointPreWrap.Y;
 
-								bool bSet = false;
-								if (CursorPoint.y < TopEdge)			{ CursorPoint.y = BottomEdge;	bSet = true; }
-								else if (CursorPoint.y > BottomEdge)	{ CursorPoint.y = TopEdge;		bSet = true; }
+									const int32 DeltaLen = (DeltaX * DeltaX) + (DeltaY * DeltaY);
+									const int32 DeltaPreWrapLen = (DeltaXPreWrap * DeltaXPreWrap) + (DeltaYPreWrap * DeltaYPreWrap);
 
-								if (CursorPoint.x < LeftEdge)			{ CursorPoint.x = RightEdge;	bSet = true; }
-								else if (CursorPoint.x > RightEdge)		{ CursorPoint.x = LeftEdge;		bSet = true; }
+									if (DeltaPreWrapLen < DeltaLen)
+									{
+										bAcceptingPreWrapDelta = true;
 
-								if (bSet)
-								{
-									//UE_LOG(LogWindowsDesktop, Log, TEXT("Wrapping Cursor to X: %d Y: %d"), CursorPoint.x, CursorPoint.y);
-									MessageHandler->SetCursorPos(FVector2D(LastCursorPoint));
-									LastCursorPoint.X = CursorPoint.x;
-									LastCursorPoint.Y = CursorPoint.y;
-									return 1;
+										DeltaX = DeltaXPreWrap;
+										DeltaY = DeltaYPreWrap;
+										LastCursorPointPreWrap.X = CursorPoint.x;
+										LastCursorPointPreWrap.Y = CursorPoint.y;
+
+										//UE_LOG(LogWindowsDesktop, Log, TEXT("Accept PreWrap Delta X: %d Y: %d  ---- Last X: %d Y: %d ---- Cur X: %d Y: %d ---- Width: %d Height: %d"), DeltaX, DeltaY, LastCursorPointPreWrap.X, LastCursorPointPreWrap.Y, CursorPoint.x, CursorPoint.y, (int32)Width, (int32)Height);
+									}
 								}
 
-								const int32 DeltaWidthMax = (int32)((float)(ClipCursorRect.right - ClipCursorRect.left) * 0.4f);
-								const int32 DeltaHeightMax = (int32)((float)(ClipCursorRect.bottom - ClipCursorRect.top) * 0.4f);
-								/*
-								if (DeltaX != 0 && DeltaY != 0)
+								if (!bAcceptingPreWrapDelta)
 								{
-									if (FMath::Abs(DeltaX) < DeltaWidthMax && FMath::Abs(DeltaY) < DeltaHeightMax)
+									// Wrap and set cursor position in necessary
+									const int32 WrapLeeway = 50; // We add some leeway on the wrap so that if the user is doing small movements hear the border we don't wrap back and fourth constantly
+									const int32 TopEdge = Top + int32(0.1f * float(Height));
+									const int32 BottomEdge = Top + int32(0.9f * float(Height));
+									const int32 LeftEdge = Left + int32(0.1f * float(Width));
+									const int32 RightEdge = Left + int32(0.9f * float(Width));
+
+									bool bSet = false;
+									if (CursorPoint.y < TopEdge) { CursorPoint.y = BottomEdge - WrapLeeway;	bSet = true; }
+									else if (CursorPoint.y > BottomEdge) { CursorPoint.y = TopEdge + WrapLeeway;		bSet = true; }
+
+									if (CursorPoint.x < LeftEdge) { CursorPoint.x = RightEdge - WrapLeeway;	bSet = true; }
+									else if (CursorPoint.x > RightEdge) { CursorPoint.x = LeftEdge + WrapLeeway;	bSet = true; }
+
+									if (bSet)
 									{
-										UE_LOG(LogWindowsDesktop, Log, TEXT("Accept Delta X: %d Y: %d  ---- Last X: %d Y: %d ---- Cur X: %d Y: %d ---- Width: %d Height: %d"), DeltaX, DeltaY, LastCursorPoint.X, LastCursorPoint.Y, CursorPoint.x, CursorPoint.y, (int32)Width, (int32)Height);
+										//UE_LOG(LogWindowsDesktop, Log, TEXT("Wrapping Cursor to X: %d Y: %d"), CursorPoint.x, CursorPoint.y);
+
+										MessageHandler->SetCursorPos(FVector2D(CursorPoint.x,CursorPoint.y));
+										LastCursorPoint.X = CursorPoint.x;
+										LastCursorPoint.Y = CursorPoint.y;
+
+										NumPreWrapMsgsToRespect = 10;
+										LastCursorPointPreWrap.X = CursorPointNoWrap.x;
+										LastCursorPointPreWrap.Y = CursorPointNoWrap.y;
 									}
-									else
+
+									
+									/*
+									if (DeltaX != 0 && DeltaY != 0)
 									{
-										UE_LOG(LogWindowsDesktop, Log, TEXT("IGNORE Delta X: %d Y: %d  ---- Last X: %d Y: %d ---- Cur X: %d Y: %d ---- Width: %d Height: %d"), DeltaX, DeltaY, LastCursorPoint.X, LastCursorPoint.Y, CursorPoint.x, CursorPoint.y, (int32)Width, (int32)Height);
-									}
-								}*/
+										if (FMath::Abs(DeltaX) < DeltaWidthMax && FMath::Abs(DeltaY) < DeltaHeightMax)
+										{
+											UE_LOG(LogWindowsDesktop, Log, TEXT("Accept Delta X: %d Y: %d  ---- Last X: %d Y: %d ---- Cur X: %d Y: %d ---- Width: %d Height: %d"), DeltaX, DeltaY, LastCursorPoint.X, LastCursorPoint.Y, CursorPoint.x, CursorPoint.y, (int32)Width, (int32)Height);
+										}
+										else
+										{
+											UE_LOG(LogWindowsDesktop, Log, TEXT("IGNORE Delta X: %d Y: %d  ---- Last X: %d Y: %d ---- Cur X: %d Y: %d ---- Width: %d Height: %d"), DeltaX, DeltaY, LastCursorPoint.X, LastCursorPoint.Y, CursorPoint.x, CursorPoint.y, (int32)Width, (int32)Height);
+										}
+									}*/
 
-								LastCursorPoint.X = CursorPoint.x;
-								LastCursorPoint.Y = CursorPoint.y;
+									LastCursorPoint.X = CursorPoint.x;
+									LastCursorPoint.Y = CursorPoint.y;
+								}
 
 								// Send a delta assuming it's not zero or beyond our max delta 
 								if (DeltaX != 0 && DeltaY != 0 && FMath::Abs(DeltaX) < DeltaWidthMax && FMath::Abs(DeltaY) < DeltaHeightMax)
diff --git a/Engine/Source/Runtime/ApplicationCore/Public/Windows/WindowsApplication.h b/Engine/Source/Runtime/ApplicationCore/Public/Windows/WindowsApplication.h
--- a/Engine/Source/Runtime/ApplicationCore/Public/Windows/WindowsApplication.h
+++ b/Engine/Source/Runtime/ApplicationCore/Public/Windows/WindowsApplication.h
@@ -531,6 +531,8 @@ private:
 
 	FIntPoint CachedPreHighPrecisionMousePosForRDP;
 	FIntPoint LastCursorPoint;
+	FIntPoint LastCursorPointPreWrap;
+	int32 NumPreWrapMsgsToRespect;
 	RECT ClipCursorRect;
 };
 
