Fix for monitor arrangements that don't align with the virtual desktop space created by RDP.

diff --git a/Engine/Source/Runtime/ApplicationCore/Private/Windows/WindowsApplication.cpp b/Engine/Source/Runtime/ApplicationCore/Private/Windows/WindowsApplication.cpp
--- a/Engine/Source/Runtime/ApplicationCore/Private/Windows/WindowsApplication.cpp
+++ b/Engine/Source/Runtime/ApplicationCore/Private/Windows/WindowsApplication.cpp
@@ -473,7 +473,7 @@ void FWindowsApplication::SetHighPrecisionMouseMode( const bool Enable, const TS
 			//UE_LOG(LogWindowsDesktop, Log, TEXT("Entering High Precision to Top: %d Bottom: %d Left: %d Right: %d"), ClipCursorRect.top, ClipCursorRect.bottom, ClipCursorRect.left, ClipCursorRect.right);
 		
 			CachedPreHighPrecisionMousePosForRDP = FIntPoint(CursorPos.x, CursorPos.y);
-			LastCursorPoint = CachedPreHighPrecisionMousePosForRDP;
+			LastCursorPoint = FIntPoint(CursorPos.x - ClipCursorRect.left, CursorPos.y - ClipCursorRect.top);
 			
 			LastCursorPointPreWrap = FIntPoint::ZeroValue;
 			NumPreWrapMsgsToRespect = 0;
@@ -1059,30 +1059,17 @@ int32 FWindowsApplication::ProcessMessage( HWND hwnd, uint32 msg, WPARAM wParam,
 						{
 							if (bSimulatingHighPrecisionMouseInputForRDP)
 							{	
+								const bool IsVirtualScreen = (Raw->data.mouse.usFlags & MOUSE_VIRTUAL_DESKTOP) == MOUSE_VIRTUAL_DESKTOP;
+
 								// Get the new cursor position
 								POINT CursorPoint;
-								int32 Top = 0;
-								int32 Left = 0;
-								int32 Width = 0;
-								int32 Height = 0;
-								const bool IsVirtualScreen = (Raw->data.mouse.usFlags & MOUSE_VIRTUAL_DESKTOP) == MOUSE_VIRTUAL_DESKTOP;
-								if (IsVirtualScreen)
-								{
-									// This is used to make Remote Desktop sessons work
-									Width = GetSystemMetrics(SM_CXVIRTUALSCREEN);
-									Height = GetSystemMetrics(SM_CYVIRTUALSCREEN);
+								const int32 Top = 0;
+								const int32 Left = 0;
+								const int32 Width = GetSystemMetrics(IsVirtualScreen ? SM_CXVIRTUALSCREEN : SM_CXSCREEN);
+								const int32 Height = GetSystemMetrics(IsVirtualScreen ? SM_CYVIRTUALSCREEN : SM_CYSCREEN);
 
-									CursorPoint.x = static_cast<int>((float(Raw->data.mouse.lLastX) / 65535.0f) * Width);
-									CursorPoint.y = static_cast<int>((float(Raw->data.mouse.lLastY) / 65535.0f) * Height);
-								}
-								else
-								{
-									Top = ClipCursorRect.top;
-									Left = ClipCursorRect.left;
-									Width = ClipCursorRect.right - ClipCursorRect.left;
-									Height = ClipCursorRect.bottom - ClipCursorRect.top;
-									::GetCursorPos(&CursorPoint);
-								}
+								CursorPoint.x = static_cast<int>((float(Raw->data.mouse.lLastX) / 65535.0f) * Width);
+								CursorPoint.y = static_cast<int>((float(Raw->data.mouse.lLastY) / 65535.0f) * Height);
 
 								const int32 DeltaWidthMax = (int32)((float)Width * 0.4f);
 								const int32 DeltaHeightMax = (int32)((float)Height * 0.4f);
@@ -1139,7 +1126,7 @@ int32 FWindowsApplication::ProcessMessage( HWND hwnd, uint32 msg, WPARAM wParam,
 									{
 										//UE_LOG(LogWindowsDesktop, Log, TEXT("Wrapping Cursor to X: %d Y: %d"), CursorPoint.x, CursorPoint.y);
 
-										MessageHandler->SetCursorPos(FVector2D(CursorPoint.x,CursorPoint.y));
+										MessageHandler->SetCursorPos(FVector2D(CursorPoint.x + ClipCursorRect.left,CursorPoint.y + ClipCursorRect.top));
 										LastCursorPoint.X = CursorPoint.x;
 										LastCursorPoint.Y = CursorPoint.y;
 
