Added raw input simulation over remote desktop on windows since remote desktop uses absolute mouse position in mouse move events.  This prevents the cursor from getting stuck over RDP

diff --git a/Engine/Source/Runtime/ApplicationCore/Private/Windows/WindowsApplication.cpp b/Engine/Source/Runtime/ApplicationCore/Private/Windows/WindowsApplication.cpp
--- a/Engine/Source/Runtime/ApplicationCore/Private/Windows/WindowsApplication.cpp
+++ b/Engine/Source/Runtime/ApplicationCore/Private/Windows/WindowsApplication.cpp
@@ -70,10 +70,35 @@ FAutoConsoleVariableRef	CVarPreventDuplicateMouseEventsForTouch(
 	TEXT("Hack to get around multiple mouse events being triggered for touch events on Windows 7 and lower.  Enabling this will prevent pen tablets from working on windows 7 since until we switch to the windows 8 sdk (and can use WM_POINTER* events) we cannot detect the difference")
 );
 
+
+#if !UE_BUILD_SHIPPING
+static int32 EnableRawInputSimulationOverRDP = true;
+FAutoConsoleVariableRef	CVarEnableRawInputSimulationOverRDP(
+	TEXT("Slate.EnableRawInputSimulationOverRDP"),
+	EnableRawInputSimulationOverRDP,
+	TEXT("")
+);
+
+static int32 ForceRawInputSimulation = false;
+FAutoConsoleVariableRef	CVarForceRawInputSimulation(
+	TEXT("Slate.ForceRawInputSimulation"),
+	ForceRawInputSimulation,
+	TEXT("")
+);
+#else
+static int32 ForceRawInputSimulation = false;
+static int32 EnableRawInputSimulationOverRDP = false;
+#endif
+
 const FIntPoint FWindowsApplication::MinimizedWindowPosition(-32000,-32000);
 
 FWindowsApplication* WindowsApplication = nullptr;
 
+static bool ShouldSimulateRawInput()
+{
+	return ForceRawInputSimulation || (EnableRawInputSimulationOverRDP && FPlatformMisc::IsRemoteSession());
+}
+
 FWindowsApplication* FWindowsApplication::CreateWindowsApplication( const HINSTANCE InstanceHandle, const HICON IconHandle )
 {
 	WindowsApplication = new FWindowsApplication( InstanceHandle, IconHandle );
@@ -101,6 +126,12 @@ FWindowsApplication::FWindowsApplication( const HINSTANCE HInstance, const HICON
 #if WITH_ACCESSIBILITY
 	, UIAManager(new FWindowsUIAManager(*this))
 #endif
+	, bSimulatingHighPrecisionMouseInputForRDP(false)
+	, CachedPreHighPrecisionMousePosForRDP(FIntPoint::ZeroValue)
+	, LastCursorPoint(FIntPoint::ZeroValue)
+	, LastCursorPointPreWrap(FIntPoint::ZeroValue)
+	, NumPreWrapMsgsToRespect(0)
+	, ClipCursorRect()
 
 {
 	FMemory::Memzero(ModifierKeyState, EModifierKey::Count);
@@ -424,36 +455,66 @@ void* FWindowsApplication::GetCapture( void ) const
 
 void FWindowsApplication::SetHighPrecisionMouseMode( const bool Enable, const TSharedPtr< FGenericWindow >& InWindow )
 {
-	HWND hwnd = NULL;
-	DWORD flags = RIDEV_REMOVE;
-	bUsingHighPrecisionMouseInput = Enable;
-
-	if ( Enable )
+	if (ShouldSimulateRawInput())
 	{
-		flags = 0;
+		if(Enable)
+		{
+			bUsingHighPrecisionMouseInput = true;
+			bSimulatingHighPrecisionMouseInputForRDP = true;
 
-		if ( InWindow.IsValid() )
+			POINT CursorPos;
+			BOOL bGotPoint = ::GetCursorPos(&CursorPos);
+
+			::GetClipCursor(&ClipCursorRect);
+
+			//UE_LOG(LogWindowsDesktop, Log, TEXT("Entering High Precision to Top: %d Bottom: %d Left: %d Right: %d"), ClipCursorRect.top, ClipCursorRect.bottom, ClipCursorRect.left, ClipCursorRect.right);
+		
+			CachedPreHighPrecisionMousePosForRDP = FIntPoint(CursorPos.x, CursorPos.y);
+			LastCursorPoint = CachedPreHighPrecisionMousePosForRDP;
+			
+			LastCursorPointPreWrap = FIntPoint::ZeroValue;
+			NumPreWrapMsgsToRespect = 0;
+		}
+		else
 		{
-			hwnd = (HWND)InWindow->GetOSWindowHandle(); 
+			CachedPreHighPrecisionMousePosForRDP = FIntPoint(INT_MAX, INT_MAX);
+			bSimulatingHighPrecisionMouseInputForRDP = false;
 		}
 	}
 
-	// NOTE: Currently has to be created every time due to conflicts with Direct8 Input used by the wx unrealed
-	RAWINPUTDEVICE RawInputDevice;
+	{
+		
+		HWND hwnd = NULL;
+		DWORD flags = RIDEV_REMOVE;
+		bUsingHighPrecisionMouseInput = Enable;
 
-	//The HID standard for mouse
-	const uint16 StandardMouse = 0x02;
+		if (Enable)
+		{
+			flags = 0;
 
-	RawInputDevice.usUsagePage = 0x01; 
-	RawInputDevice.usUsage = StandardMouse;
-	RawInputDevice.dwFlags = flags;
+			if (InWindow.IsValid())
+			{
+				hwnd = (HWND)InWindow->GetOSWindowHandle();
+			}
+		}
 
-	// Process input for just the window that requested it.  NOTE: If we pass NULL here events are routed to the window with keyboard focus
-	// which is not always known at the HWND level with Slate
-	RawInputDevice.hwndTarget = hwnd; 
+		// NOTE: Currently has to be created every time due to conflicts with Direct8 Input used by the wx unrealed
+		RAWINPUTDEVICE RawInputDevice;
 
-	// Register the raw input device
-	::RegisterRawInputDevices( &RawInputDevice, 1, sizeof( RAWINPUTDEVICE ) );
+		//The HID standard for mouse
+		const uint16 StandardMouse = 0x02;
+
+		RawInputDevice.usUsagePage = 0x01;
+		RawInputDevice.usUsage = StandardMouse;
+		RawInputDevice.dwFlags = flags;
+
+		// Process input for just the window that requested it.  NOTE: If we pass NULL here events are routed to the window with keyboard focus
+		// which is not always known at the HWND level with Slate
+		RawInputDevice.hwndTarget = hwnd;
+
+		// Register the raw input device
+		::RegisterRawInputDevices(&RawInputDevice, 1, sizeof(RAWINPUTDEVICE));
+	}
 }
 
 FPlatformRect FWindowsApplication::GetWorkArea( const FPlatformRect& CurrentWindow ) const
@@ -991,11 +1052,130 @@ int32 FWindowsApplication::ProcessMessage( HWND hwnd, uint32 msg, WPARAM wParam,
 					if (Raw->header.dwType == RIM_TYPEMOUSE) 
 					{
 						const bool IsAbsoluteInput = (Raw->data.mouse.usFlags & MOUSE_MOVE_ABSOLUTE) == MOUSE_MOVE_ABSOLUTE;
-						if( IsAbsoluteInput )
+						if (IsAbsoluteInput)
 						{
-							// Since the raw input is coming in as absolute it is likely the user is using a tablet
-							// or perhaps is interacting through a virtual desktop
-							DeferMessage( CurrentNativeEventWindowPtr, hwnd, msg, wParam, lParam, 0, 0, MOUSE_MOVE_ABSOLUTE );
+							if (bSimulatingHighPrecisionMouseInputForRDP)
+							{	
+								// Get the new cursor position
+								POINT CursorPoint;
+								int32 Top = 0;
+								int32 Left = 0;
+								int32 Width = 0;
+								int32 Height = 0;
+								const bool IsVirtualScreen = (Raw->data.mouse.usFlags & MOUSE_VIRTUAL_DESKTOP) == MOUSE_VIRTUAL_DESKTOP;
+								if (IsVirtualScreen)
+								{
+									// This is used to make Remote Desktop sessons work
+									Width = GetSystemMetrics(SM_CXVIRTUALSCREEN);
+									Height = GetSystemMetrics(SM_CYVIRTUALSCREEN);
+
+									CursorPoint.x = static_cast<int>((float(Raw->data.mouse.lLastX) / 65535.0f) * Width);
+									CursorPoint.y = static_cast<int>((float(Raw->data.mouse.lLastY) / 65535.0f) * Height);
+								}
+								else
+								{
+									Top = ClipCursorRect.top;
+									Left = ClipCursorRect.left;
+									Width = ClipCursorRect.right - ClipCursorRect.left;
+									Height = ClipCursorRect.bottom - ClipCursorRect.top;
+									::GetCursorPos(&CursorPoint);
+								}
+
+								const int32 DeltaWidthMax = (int32)((float)Width * 0.4f);
+								const int32 DeltaHeightMax = (int32)((float)Height * 0.4f);
+
+								const POINT CursorPointNoWrap = CursorPoint;
+
+								const bool bCanAcceptPreWrapMsg = NumPreWrapMsgsToRespect > 0;
+								NumPreWrapMsgsToRespect--;
+
+								// Calculate the cursor delta from the last position
+								// We do this prior to and wrapping for continuous input as that would mean an incorrect delta we'd just ignore
+								int32 DeltaX = CursorPoint.x - LastCursorPoint.X;
+								int32 DeltaY = CursorPoint.y - LastCursorPoint.Y;
+
+								bool bAcceptingPreWrapDelta = false;
+								if (bCanAcceptPreWrapMsg)
+								{
+									const int32 DeltaXPreWrap = CursorPoint.x - LastCursorPointPreWrap.X;
+									const int32 DeltaYPreWrap = CursorPoint.y - LastCursorPointPreWrap.Y;
+
+									const int32 DeltaLen = (DeltaX * DeltaX) + (DeltaY * DeltaY);
+									const int32 DeltaPreWrapLen = (DeltaXPreWrap * DeltaXPreWrap) + (DeltaYPreWrap * DeltaYPreWrap);
+
+									if (DeltaPreWrapLen < DeltaLen)
+									{
+										bAcceptingPreWrapDelta = true;
+
+										DeltaX = DeltaXPreWrap;
+										DeltaY = DeltaYPreWrap;
+										LastCursorPointPreWrap.X = CursorPoint.x;
+										LastCursorPointPreWrap.Y = CursorPoint.y;
+
+										//UE_LOG(LogWindowsDesktop, Log, TEXT("Accept PreWrap Delta X: %d Y: %d  ---- Last X: %d Y: %d ---- Cur X: %d Y: %d ---- Width: %d Height: %d"), DeltaX, DeltaY, LastCursorPointPreWrap.X, LastCursorPointPreWrap.Y, CursorPoint.x, CursorPoint.y, (int32)Width, (int32)Height);
+									}
+								}
+
+								if (!bAcceptingPreWrapDelta)
+								{
+									// Wrap and set cursor position in necessary
+									const int32 WrapLeeway = 50; // We add some leeway on the wrap so that if the user is doing small movements hear the border we don't wrap back and fourth constantly
+									const int32 TopEdge = Top + int32(0.1f * float(Height));
+									const int32 BottomEdge = Top + int32(0.9f * float(Height));
+									const int32 LeftEdge = Left + int32(0.1f * float(Width));
+									const int32 RightEdge = Left + int32(0.9f * float(Width));
+
+									bool bSet = false;
+									if (CursorPoint.y < TopEdge) { CursorPoint.y = BottomEdge - WrapLeeway;	bSet = true; }
+									else if (CursorPoint.y > BottomEdge) { CursorPoint.y = TopEdge + WrapLeeway;		bSet = true; }
+
+									if (CursorPoint.x < LeftEdge) { CursorPoint.x = RightEdge - WrapLeeway;	bSet = true; }
+									else if (CursorPoint.x > RightEdge) { CursorPoint.x = LeftEdge + WrapLeeway;	bSet = true; }
+
+									if (bSet)
+									{
+										//UE_LOG(LogWindowsDesktop, Log, TEXT("Wrapping Cursor to X: %d Y: %d"), CursorPoint.x, CursorPoint.y);
+
+										MessageHandler->SetCursorPos(FVector2D(CursorPoint.x,CursorPoint.y));
+										LastCursorPoint.X = CursorPoint.x;
+										LastCursorPoint.Y = CursorPoint.y;
+
+										NumPreWrapMsgsToRespect = 10;
+										LastCursorPointPreWrap.X = CursorPointNoWrap.x;
+										LastCursorPointPreWrap.Y = CursorPointNoWrap.y;
+									}
+
+									
+									/*
+									if (DeltaX != 0 && DeltaY != 0)
+									{
+										if (FMath::Abs(DeltaX) < DeltaWidthMax && FMath::Abs(DeltaY) < DeltaHeightMax)
+										{
+											UE_LOG(LogWindowsDesktop, Log, TEXT("Accept Delta X: %d Y: %d  ---- Last X: %d Y: %d ---- Cur X: %d Y: %d ---- Width: %d Height: %d"), DeltaX, DeltaY, LastCursorPoint.X, LastCursorPoint.Y, CursorPoint.x, CursorPoint.y, (int32)Width, (int32)Height);
+										}
+										else
+										{
+											UE_LOG(LogWindowsDesktop, Log, TEXT("IGNORE Delta X: %d Y: %d  ---- Last X: %d Y: %d ---- Cur X: %d Y: %d ---- Width: %d Height: %d"), DeltaX, DeltaY, LastCursorPoint.X, LastCursorPoint.Y, CursorPoint.x, CursorPoint.y, (int32)Width, (int32)Height);
+										}
+									}*/
+
+									LastCursorPoint.X = CursorPoint.x;
+									LastCursorPoint.Y = CursorPoint.y;
+								}
+
+								// Send a delta assuming it's not zero or beyond our max delta 
+								if (DeltaX != 0 && DeltaY != 0 && FMath::Abs(DeltaX) < DeltaWidthMax && FMath::Abs(DeltaY) < DeltaHeightMax)
+								{
+									DeferMessage(CurrentNativeEventWindowPtr, hwnd, msg, wParam, lParam, DeltaX, DeltaY, MOUSE_MOVE_RELATIVE);
+								}
+								return 1;
+							}
+							else
+							{
+								// Since the raw input is coming in as absolute it is likely the user is using a tablet
+								// or perhaps is interacting through a virtual desktop
+								DeferMessage(CurrentNativeEventWindowPtr, hwnd, msg, wParam, lParam, 0, 0, MOUSE_MOVE_ABSOLUTE);
+							}
 							return 1;
 						}
 
@@ -1979,7 +2159,7 @@ int32 FWindowsApplication::ProcessDeferredMessage( const FDeferredWindowsMessage
 			{
 				if( DeferredMessage.RawInputFlags == MOUSE_MOVE_RELATIVE )
 				{
-					MessageHandler->OnRawMouseMove( DeferredMessage.X, DeferredMessage.Y );
+					MessageHandler->OnRawMouseMove(DeferredMessage.X, DeferredMessage.Y);
 				}
 				else
 				{
@@ -1996,7 +2176,7 @@ int32 FWindowsApplication::ProcessDeferredMessage( const FDeferredWindowsMessage
 		case WM_MOUSEMOVE:
 			{
 				BOOL Result = false;
-				if( !bUsingHighPrecisionMouseInput )
+				if (!bUsingHighPrecisionMouseInput)
 				{
 					Result = MessageHandler->OnMouseMove();
 				}
diff --git a/Engine/Source/Runtime/ApplicationCore/Private/Windows/WindowsCursor.cpp b/Engine/Source/Runtime/ApplicationCore/Private/Windows/WindowsCursor.cpp
--- a/Engine/Source/Runtime/ApplicationCore/Private/Windows/WindowsCursor.cpp
+++ b/Engine/Source/Runtime/ApplicationCore/Private/Windows/WindowsCursor.cpp
@@ -29,6 +29,13 @@ FWindowsCursor::FWindowsCursor()
 		switch( CurCursorIndex )
 		{
 		case EMouseCursor::None:
+			if (FPlatformMisc::IsRemoteSession())
+			{
+				// during remote sessions we rely on constantly resetting the mouse cursor position for infinite mouse deltas while dragging (WM_INPUT doesnt work over rdp)
+				// this hack here is to avoid having a hidden cursor which prevents SetCursorPos from working.  This is just a completely transparent cursor yet "visible" cursor to work around that
+				CursorHandle = LoadCursorFromFile((LPCTSTR) * (FString(FPlatformProcess::BaseDir()) / FString::Printf(TEXT("%sEditor/Slate/Cursor/invisible.cur"), *FPaths::EngineContentDir())));
+			}
+			break;
 		case EMouseCursor::Custom:
 			// The mouse cursor will not be visible when None is used
 			break;
diff --git a/Engine/Source/Runtime/ApplicationCore/Public/GenericPlatform/GenericApplicationMessageHandler.h b/Engine/Source/Runtime/ApplicationCore/Public/GenericPlatform/GenericApplicationMessageHandler.h
--- a/Engine/Source/Runtime/ApplicationCore/Public/GenericPlatform/GenericApplicationMessageHandler.h
+++ b/Engine/Source/Runtime/ApplicationCore/Public/GenericPlatform/GenericApplicationMessageHandler.h
@@ -470,4 +470,9 @@ public:
 	{
 		return true;
 	}
+	
+	virtual void SetCursorPos(const FVector2D& MouseCoordinate)
+	{
+
+	}
 };
diff --git a/Engine/Source/Runtime/ApplicationCore/Public/Windows/WindowsApplication.h b/Engine/Source/Runtime/ApplicationCore/Public/Windows/WindowsApplication.h
--- a/Engine/Source/Runtime/ApplicationCore/Public/Windows/WindowsApplication.h
+++ b/Engine/Source/Runtime/ApplicationCore/Public/Windows/WindowsApplication.h
@@ -525,6 +525,14 @@ private:
 	/** Maps touch indexes to windows touch IDs. */
 	TArray<TOptional<int32>> TouchIDs;
 #endif
+
+	bool bSimulatingHighPrecisionMouseInputForRDP;
+
+	FIntPoint CachedPreHighPrecisionMousePosForRDP;
+	FIntPoint LastCursorPoint;
+	FIntPoint LastCursorPointPreWrap;
+	int32 NumPreWrapMsgsToRespect;
+	RECT ClipCursorRect;
 };
 
 
