Fix for improper clip rect behavior when in RDP.

diff --git a/Engine/Source/Runtime/ApplicationCore/Private/Windows/WindowsApplication.cpp b/Engine/Source/Runtime/ApplicationCore/Private/Windows/WindowsApplication.cpp
--- a/Engine/Source/Runtime/ApplicationCore/Private/Windows/WindowsApplication.cpp
+++ b/Engine/Source/Runtime/ApplicationCore/Private/Windows/WindowsApplication.cpp
@@ -473,7 +473,8 @@ void FWindowsApplication::SetHighPrecisionMouseMode( const bool Enable, const TS
 			//UE_LOG(LogWindowsDesktop, Log, TEXT("Entering High Precision to Top: %d Bottom: %d Left: %d Right: %d"), ClipCursorRect.top, ClipCursorRect.bottom, ClipCursorRect.left, ClipCursorRect.right);
 		
 			CachedPreHighPrecisionMousePosForRDP = FIntPoint(CursorPos.x, CursorPos.y);
-			LastCursorPoint = FIntPoint(CursorPos.x - ClipCursorRect.left, CursorPos.y - ClipCursorRect.top);
+			
+			LastCursorPoint = FIntPoint(CursorPos.x, CursorPos.y);
 			
 			LastCursorPointPreWrap = FIntPoint::ZeroValue;
 			NumPreWrapMsgsToRespect = 0;
@@ -1063,16 +1064,19 @@ int32 FWindowsApplication::ProcessMessage( HWND hwnd, uint32 msg, WPARAM wParam,
 
 								// Get the new cursor position
 								POINT CursorPoint;
-								const int32 Top = 0;
-								const int32 Left = 0;
+								const int32 Top = IsVirtualScreen ? GetSystemMetrics(SM_XVIRTUALSCREEN) : 0;
+								const int32 Left = IsVirtualScreen ? GetSystemMetrics(SM_YVIRTUALSCREEN) : 0;
 								const int32 Width = GetSystemMetrics(IsVirtualScreen ? SM_CXVIRTUALSCREEN : SM_CXSCREEN);
 								const int32 Height = GetSystemMetrics(IsVirtualScreen ? SM_CYVIRTUALSCREEN : SM_CYSCREEN);
 
-								CursorPoint.x = static_cast<int>((float(Raw->data.mouse.lLastX) / 65535.0f) * Width);
-								CursorPoint.y = static_cast<int>((float(Raw->data.mouse.lLastY) / 65535.0f) * Height);
+								CursorPoint.x = static_cast<int>((float(Raw->data.mouse.lLastX) / 65535.0f) * Width) - Left;
+								CursorPoint.y = static_cast<int>((float(Raw->data.mouse.lLastY) / 65535.0f) * Height) - Top;
+
+								const int32 ClipWidth = ClipCursorRect.right - ClipCursorRect.left;
+								const int32 ClipHeight = ClipCursorRect.bottom - ClipCursorRect.top;
 
-								const int32 DeltaWidthMax = (int32)((float)Width * 0.4f);
-								const int32 DeltaHeightMax = (int32)((float)Height * 0.4f);
+								const int32 DeltaWidthMax = (int32)((float)ClipWidth * 0.4f);
+								const int32 DeltaHeightMax = (int32)((float)ClipHeight * 0.4f);
 
 								const POINT CursorPointNoWrap = CursorPoint;
 
@@ -1110,23 +1114,24 @@ int32 FWindowsApplication::ProcessMessage( HWND hwnd, uint32 msg, WPARAM wParam,
 								{
 									// Wrap and set cursor position in necessary
 									const int32 WrapLeeway = 50; // We add some leeway on the wrap so that if the user is doing small movements hear the border we don't wrap back and fourth constantly
-									const int32 TopEdge = Top + int32(0.1f * float(Height));
-									const int32 BottomEdge = Top + int32(0.9f * float(Height));
-									const int32 LeftEdge = Left + int32(0.1f * float(Width));
-									const int32 RightEdge = Left + int32(0.9f * float(Width));
+									const int32 TopEdge = ClipCursorRect.top + int32(0.1f * float(ClipHeight));
+									const int32 BottomEdge = ClipCursorRect.top + int32(0.9f * float(ClipHeight));
+									const int32 LeftEdge = ClipCursorRect.left + int32(0.1f * float(ClipWidth));
+									const int32 RightEdge = ClipCursorRect.left + int32(0.9f * float(ClipWidth));
+
 
 									bool bSet = false;
 									if (CursorPoint.y < TopEdge) { CursorPoint.y = BottomEdge - WrapLeeway;	bSet = true; }
-									else if (CursorPoint.y > BottomEdge) { CursorPoint.y = TopEdge + WrapLeeway;		bSet = true; }
+									else if (CursorPoint.y > BottomEdge) { CursorPoint.y = TopEdge + WrapLeeway; bSet = true; }
 
 									if (CursorPoint.x < LeftEdge) { CursorPoint.x = RightEdge - WrapLeeway;	bSet = true; }
-									else if (CursorPoint.x > RightEdge) { CursorPoint.x = LeftEdge + WrapLeeway;	bSet = true; }
+									else if (CursorPoint.x > RightEdge) { CursorPoint.x = LeftEdge + WrapLeeway; bSet = true; }
 
 									if (bSet)
 									{
 										//UE_LOG(LogWindowsDesktop, Log, TEXT("Wrapping Cursor to X: %d Y: %d"), CursorPoint.x, CursorPoint.y);
 
-										MessageHandler->SetCursorPos(FVector2D(CursorPoint.x + ClipCursorRect.left,CursorPoint.y + ClipCursorRect.top));
+										MessageHandler->SetCursorPos(FVector2D(CursorPoint.x,CursorPoint.y));
 										LastCursorPoint.X = CursorPoint.x;
 										LastCursorPoint.Y = CursorPoint.y;
 
