Fix for a few minor issues exposed by non 0, 0 origin of dektop space.

diff --git a/Engine/Source/Runtime/ApplicationCore/Private/Windows/WindowsApplication.cpp b/Engine/Source/Runtime/ApplicationCore/Private/Windows/WindowsApplication.cpp
--- a/Engine/Source/Runtime/ApplicationCore/Private/Windows/WindowsApplication.cpp
+++ b/Engine/Source/Runtime/ApplicationCore/Private/Windows/WindowsApplication.cpp
@@ -467,7 +467,7 @@ void FWindowsApplication::SetHighPrecisionMouseMode( const bool Enable, const TS
 
 			::GetClipCursor(&ClipCursorRect);
 
-			//UE_LOG(LogWindowsDesktop, Log, TEXT("Entering High Precision to Top: %d Bottom: %d Left: %d Right: %d"), ClipCursorRect.top, ClipCursorRect.bottom, ClipCursorRect.left, ClipCursorRect.right);
+			//UE_LOG(LogWindowsDesktop, Log, TEXT("Entering High Precision to Top: %d Bottom: %d Left: %d Right: %d ---- Starting x: %d y: %d"), ClipCursorRect.top, ClipCursorRect.bottom, ClipCursorRect.left, ClipCursorRect.right, CursorPos.x, CursorPos.y);
 		
 			CachedPreHighPrecisionMousePosForRDP = FIntPoint(CursorPos.x, CursorPos.y);
 			
@@ -1061,13 +1061,13 @@ int32 FWindowsApplication::ProcessMessage( HWND hwnd, uint32 msg, WPARAM wParam,
 
 								// Get the new cursor position
 								POINT CursorPoint;
-								const int32 Top = IsVirtualScreen ? GetSystemMetrics(SM_XVIRTUALSCREEN) : 0;
-								const int32 Left = IsVirtualScreen ? GetSystemMetrics(SM_YVIRTUALSCREEN) : 0;
+								const int32 Left = IsVirtualScreen ? GetSystemMetrics(SM_XVIRTUALSCREEN) : 0;
+								const int32 Top = IsVirtualScreen ? GetSystemMetrics(SM_YVIRTUALSCREEN) : 0;
 								const int32 Width = GetSystemMetrics(IsVirtualScreen ? SM_CXVIRTUALSCREEN : SM_CXSCREEN);
 								const int32 Height = GetSystemMetrics(IsVirtualScreen ? SM_CYVIRTUALSCREEN : SM_CYSCREEN);
 
-								CursorPoint.x = static_cast<int>((float(Raw->data.mouse.lLastX) / 65535.0f) * Width) - Left;
-								CursorPoint.y = static_cast<int>((float(Raw->data.mouse.lLastY) / 65535.0f) * Height) - Top;
+								CursorPoint.x = static_cast<int>((float(Raw->data.mouse.lLastX) / 65535.0f) * Width) + Left;
+								CursorPoint.y = static_cast<int>((float(Raw->data.mouse.lLastY) / 65535.0f) * Height) + Top;
 
 								const int32 ClipWidth = ClipCursorRect.right - ClipCursorRect.left;
 								const int32 ClipHeight = ClipCursorRect.bottom - ClipCursorRect.top;
@@ -1116,7 +1116,6 @@ int32 FWindowsApplication::ProcessMessage( HWND hwnd, uint32 msg, WPARAM wParam,
 									const int32 LeftEdge = ClipCursorRect.left + int32(0.1f * float(ClipWidth));
 									const int32 RightEdge = ClipCursorRect.left + int32(0.9f * float(ClipWidth));
 
-
 									bool bSet = false;
 									if (CursorPoint.y < TopEdge) { CursorPoint.y = BottomEdge - WrapLeeway;	bSet = true; }
 									else if (CursorPoint.y > BottomEdge) { CursorPoint.y = TopEdge + WrapLeeway; bSet = true; }
@@ -1126,7 +1125,7 @@ int32 FWindowsApplication::ProcessMessage( HWND hwnd, uint32 msg, WPARAM wParam,
 
 									if (bSet)
 									{
-										//UE_LOG(LogWindowsDesktop, Log, TEXT("Wrapping Cursor to X: %d Y: %d"), CursorPoint.x, CursorPoint.y);
+										//UE_LOG(LogWindowsDesktop, Log, TEXT("Wrapping Cursor to X: %d Y: %d ---- TopEdge: %d BottomEdge: %d LeftEdge: %d RightEdge: %d "), CursorPoint.x, CursorPoint.y, TopEdge, BottomEdge, LeftEdge, RightEdge);
 
 										MessageHandler->SetCursorPos(FVector2D(CursorPoint.x,CursorPoint.y));
 										LastCursorPoint.X = CursorPoint.x;
@@ -1136,10 +1135,9 @@ int32 FWindowsApplication::ProcessMessage( HWND hwnd, uint32 msg, WPARAM wParam,
 										LastCursorPointPreWrap.X = CursorPointNoWrap.x;
 										LastCursorPointPreWrap.Y = CursorPointNoWrap.y;
 									}
-
 									
 									/*
-									if (DeltaX != 0 && DeltaY != 0)
+									if (DeltaX != 0 || DeltaY != 0)
 									{
 										if (FMath::Abs(DeltaX) < DeltaWidthMax && FMath::Abs(DeltaY) < DeltaHeightMax)
 										{
@@ -1156,7 +1154,7 @@ int32 FWindowsApplication::ProcessMessage( HWND hwnd, uint32 msg, WPARAM wParam,
 								}
 
 								// Send a delta assuming it's not zero or beyond our max delta 
-								if (DeltaX != 0 && DeltaY != 0 && FMath::Abs(DeltaX) < DeltaWidthMax && FMath::Abs(DeltaY) < DeltaHeightMax)
+								if ((DeltaX != 0 || DeltaY != 0) && FMath::Abs(DeltaX) < DeltaWidthMax && FMath::Abs(DeltaY) < DeltaHeightMax)
 								{
 									DeferMessage(CurrentNativeEventWindowPtr, hwnd, msg, wParam, lParam, DeltaX, DeltaY, MOUSE_MOVE_RELATIVE);
 								}
