Fix non-unity compile
#rnx

diff --git a/Engine/Source/Runtime/ApplicationCore/Private/Android/AndroidPlatformFramePacer.cpp b/Engine/Source/Runtime/ApplicationCore/Private/Android/AndroidPlatformFramePacer.cpp
--- a/Engine/Source/Runtime/ApplicationCore/Private/Android/AndroidPlatformFramePacer.cpp
+++ b/Engine/Source/Runtime/ApplicationCore/Private/Android/AndroidPlatformFramePacer.cpp
@@ -1,7 +1,6 @@
 // Copyright Epic Games, Inc. All Rights Reserved.
 
 #include "Android/AndroidPlatformFramePacer.h"
-#include "HAL/IConsoleManager.h"
 #include "Misc/ConfigCacheIni.h"
 #include "Misc/Parse.h"
 #include "Misc/CommandLine.h"
diff --git a/Engine/Source/Runtime/ApplicationCore/Public/Android/AndroidPlatformFramePacer.h b/Engine/Source/Runtime/ApplicationCore/Public/Android/AndroidPlatformFramePacer.h
--- a/Engine/Source/Runtime/ApplicationCore/Public/Android/AndroidPlatformFramePacer.h
+++ b/Engine/Source/Runtime/ApplicationCore/Public/Android/AndroidPlatformFramePacer.h
@@ -6,6 +6,7 @@
 
 #pragma once
 #include "GenericPlatform/GenericPlatformFramePacer.h"
+#include "HAL/IConsoleManager.h"
 
 /**
  * Android implementation of FGenericPlatformRHIFramePacer
