Storyboard logic post launch

diff --git a/Engine/Source/Runtime/ApplicationCore/Private/IOS/IOSAppDelegate.cpp b/Engine/Source/Runtime/ApplicationCore/Private/IOS/IOSAppDelegate.cpp
--- a/Engine/Source/Runtime/ApplicationCore/Private/IOS/IOSAppDelegate.cpp
+++ b/Engine/Source/Runtime/ApplicationCore/Private/IOS/IOSAppDelegate.cpp
@@ -947,197 +947,23 @@ static FAutoConsoleVariableRef CVarGEnableThermalsReport(
 	self.PeakMemoryTimer = [NSTimer scheduledTimerWithTimeInterval:0.1f target:self selector:@selector(RecordPeakMemory) userInfo:nil repeats:YES];
 
 #if !BUILD_EMBEDDED_APP
-	// create the main landscape window object
-	CGRect MainFrame = [[UIScreen mainScreen] bounds];
-	self.Window = [[UIWindow alloc] initWithFrame:MainFrame];
-	self.Window.screen = [UIScreen mainScreen];
-    
     // get the native scale
     const float NativeScale = [[UIScreen mainScreen] scale];
-    
-	//Make this the primary window, and show it.
-	[self.Window makeKeyAndVisible];
-
-	FAppEntry::PreInit(self, application);
-    
-    // add the default image as a subview
-    NSMutableString* path = [[NSMutableString alloc]init];
-    [path setString: [[NSBundle mainBundle] resourcePath]];
-    UIImageOrientation orient = UIImageOrientationUp;
-    NSMutableString* ImageString = [[NSMutableString alloc]init];
-	NSMutableString* PngString = [[NSMutableString alloc] init];
-    [ImageString appendString:@"Default"];
-
-
-	FPlatformMisc::EIOSDevice Device = FPlatformMisc::GetIOSDeviceType();
 
-	// iphone6 has specially named files, this seems to be needed for every iphone since, so let's see if we can find a better way to do this which isn't device specific
-    if (Device == FPlatformMisc::IOS_IPhone6 || Device == FPlatformMisc::IOS_IPhone6S || Device == FPlatformMisc::IOS_IPhone7 || Device == FPlatformMisc::IOS_IPhone8)
-	{
-		[ImageString appendString:@"-IPhone6"];
-		if (!self.bDeviceInPortraitMode)
-		{
-			[ImageString appendString : @"-Landscape"];
-		}
-	}
-    else if (Device == FPlatformMisc::IOS_IPhone6Plus || Device == FPlatformMisc::IOS_IPhone6SPlus || Device == FPlatformMisc::IOS_IPhone7Plus || Device == FPlatformMisc::IOS_IPhone8Plus)
-	{
-		[ImageString appendString : @"-IPhone6Plus"];
-		if (!self.bDeviceInPortraitMode)
-		{
-			[ImageString appendString : @"-Landscape"];
-		}
-		else
-		{
-			[ImageString appendString : @"-Portrait"];
-		}
-	}
-	else if (Device == FPlatformMisc::IOS_IPhoneX || Device == FPlatformMisc::IOS_IPhoneXS || Device == FPlatformMisc::IOS_IPhone11Pro)
-	{
-		[ImageString appendString : @"-IPhoneXS"];
-		if (!self.bDeviceInPortraitMode)
-		{
-			[ImageString appendString : @"-Landscape"];
-		}
-		else
-		{
-			[ImageString appendString : @"-Portrait"];
-		}
-	}
-	else if (Device == FPlatformMisc::IOS_IPhoneXSMax || Device == FPlatformMisc::IOS_IPhone11ProMax)
+    self.Window = [[UIWindow alloc] initWithFrame:UIScreen.mainScreen.bounds];
+    UIStoryboard *storyboard = [UIStoryboard storyboardWithName:@"LaunchScreen" bundle:nil];
+    if (storyboard != nil)
     {
-        [ImageString appendString : @"-IPhoneXSMax"];
-        if (!self.bDeviceInPortraitMode)
-        {
-            [ImageString appendString : @"-Landscape"];
-        }
-        else
-        {
-            [ImageString appendString : @"-Portrait"];
-        }
-    }
-	else if (Device == FPlatformMisc::IOS_IPhoneXR || Device == FPlatformMisc::IOS_IPhone11)
-    {
-        [ImageString appendString : @"-IPhoneXR"];
-        if (!self.bDeviceInPortraitMode)
-        {
-            [ImageString appendString : @"-Landscape"];
-        }
-        else
-        {
-            [ImageString appendString : @"-Portrait"];
-        }
+        UIViewController *viewController = [storyboard instantiateViewControllerWithIdentifier:@"LaunchScreen"];
+        self.Window.rootViewController = viewController;
     }
-	else if (Device == FPlatformMisc::IOS_AppleTV)
-	{
-		// @todo tvos: Make an AppleTV one?
-		// use IPhone6 image for now
-		[ImageString appendString : @"-IPhone6Plus-Landscape"];
-	}
-	else if (Device == FPlatformMisc::IOS_IPadPro_129 || Device == FPlatformMisc::IOS_IPadPro2_129 || Device == FPlatformMisc::IOS_IPadPro3_129)
-	{
-		if (!self.bDeviceInPortraitMode)
-		{
-			[ImageString appendString : @"-Landscape-1336"];
-		}
-		else
-		{
-			[ImageString appendString : @"-Portrait-1336"];
-		}
-        
-        if (NativeScale > 1.0f)
-        {
-            [ImageString appendString:@"@2x"];
-        }
-	}
-	else if (Device == FPlatformMisc::IOS_IPadPro_105)
-	{
-		if (!self.bDeviceInPortraitMode)
-		{
-			[ImageString appendString : @"-Landscape-1112"];
-		}
-		else
-		{
-			[ImageString appendString : @"-Portrait-1112"];
-		}
+    [self.Window makeKeyAndVisible];
 
-		if (NativeScale > 1.0f)
-		{
-			[ImageString appendString : @"@2x"];
-		}
-	}
-	else if (Device == FPlatformMisc::IOS_IPadPro_11)
-	{
-		if (!self.bDeviceInPortraitMode)
-		{
-			[ImageString appendString : @"-Landscape-1194"];
-		}
-		else
-		{
-			[ImageString appendString : @"-Portrait-1194"];
-		}
-
-		if (NativeScale > 1.0f)
-		{
-			[ImageString appendString : @"@2x"];
-		}
-	}
-	else
-	{
-		if (MainFrame.size.height == 320 && MainFrame.size.width != 480 && !self.bDeviceInPortraitMode)
-		{
-			[ImageString appendString:@"-568h"];
-			orient = UIImageOrientationRight;
-		}
-		else if (MainFrame.size.height == 320 && MainFrame.size.width == 480 && !self.bDeviceInPortraitMode)
-		{
-			orient = UIImageOrientationRight;
-		}
-		else if (MainFrame.size.height == 568 || Device == FPlatformMisc::IOS_IPodTouch6 || Device == FPlatformMisc::IOS_IPodTouch7)
-		{
-			[ImageString appendString:@"-568h"];
-		}
-		else if (MainFrame.size.height == 1024 && !self.bDeviceInPortraitMode)
-		{
-			[ImageString appendString:@"-Landscape"];
-			orient = UIImageOrientationRight;
-		}
-		else if (MainFrame.size.height == 1024)
-		{
-			[ImageString appendString:@"-Portrait"];
-		}
-		else if (MainFrame.size.height == 768 && !self.bDeviceInPortraitMode)
-		{
-			[ImageString appendString:@"-Landscape"];
-		}
-        
-        if (NativeScale > 1.0f)
-        {
-            [ImageString appendString:@"@2x"];
-        }
-	}
-
-	[PngString appendString : ImageString];
-	[PngString appendString : @".png"];
-	[ImageString appendString : @".jpg"];
-    [path setString: [path stringByAppendingPathComponent:ImageString]];
-    UIImage* image = [[UIImage alloc] initWithContentsOfFile: path];
-	if (image == nil)
-	{
-        [path setString: [[NSBundle mainBundle] resourcePath]];
-		[path setString : [path stringByAppendingPathComponent : PngString]];
-		image = [[UIImage alloc] initWithContentsOfFile:path];
-	}
-	[path release];
-	
-    UIImage* imageToDisplay = [UIImage imageWithCGImage: [image CGImage] scale: 1.0 orientation: orient];
-    UIImageView* imageView = [[UIImageView alloc] initWithImage: imageToDisplay];
-    imageView.frame = MainFrame;
-    imageView.tag = 200;
-    [self.Window addSubview: imageView];
+    FAppEntry::PreInit(self, application);
+    
     GShowSplashScreen = true;
 
-	timer = [NSTimer scheduledTimerWithTimeInterval: 0.05f target:self selector:@selector(timerForSplashScreen) userInfo:nil repeats:YES];
+    timer = [NSTimer scheduledTimerWithTimeInterval: 0.05f target:self selector:@selector(timerForSplashScreen) userInfo:nil repeats:YES];
 
 	[self StartGameThread];
 
