LaunchScreen Storyboard displaying AppDelegate fix

diff --git a/Engine/Source/Runtime/ApplicationCore/Private/IOS/IOSAppDelegate.cpp b/Engine/Source/Runtime/ApplicationCore/Private/IOS/IOSAppDelegate.cpp
--- a/Engine/Source/Runtime/ApplicationCore/Private/IOS/IOSAppDelegate.cpp
+++ b/Engine/Source/Runtime/ApplicationCore/Private/IOS/IOSAppDelegate.cpp
@@ -947,21 +947,28 @@ static FAutoConsoleVariableRef CVarGEnableThermalsReport(
 	self.PeakMemoryTimer = [NSTimer scheduledTimerWithTimeInterval:0.1f target:self selector:@selector(RecordPeakMemory) userInfo:nil repeats:YES];
 
 #if !BUILD_EMBEDDED_APP
+    
+    
+    CGRect MainFrame = [[UIScreen mainScreen] bounds];
+    self.Window = [[UIWindow alloc] initWithFrame:MainFrame];
+    self.Window.screen = [UIScreen mainScreen];
+    
     // get the native scale
     const float NativeScale = [[UIScreen mainScreen] scale];
+    
+    [self.Window makeKeyAndVisible];
+
+    FAppEntry::PreInit(self, application);
 
-    self.Window = [[UIWindow alloc] initWithFrame:UIScreen.mainScreen.bounds];
     UIStoryboard *storyboard = [UIStoryboard storyboardWithName:@"LaunchScreen" bundle:nil];
     if (storyboard != nil)
     {
         UIViewController *viewController = [storyboard instantiateViewControllerWithIdentifier:@"LaunchScreen"];
-        self.Window.rootViewController = viewController;
+        viewController.view.tag = 200;
+        [self.Window addSubview: viewController.view];
+        GShowSplashScreen = true;
     }
-    [self.Window makeKeyAndVisible];
-
-    FAppEntry::PreInit(self, application);
     
-    GShowSplashScreen = true;
 
     timer = [NSTimer scheduledTimerWithTimeInterval: 0.05f target:self selector:@selector(timerForSplashScreen) userInfo:nil repeats:YES];
 
