Fixed issue with tvOS not responding to input from AppleTV remote

diff --git a/Engine/Source/Runtime/ApplicationCore/Private/IOS/IOSInputInterface.cpp b/Engine/Source/Runtime/ApplicationCore/Private/IOS/IOSInputInterface.cpp
--- a/Engine/Source/Runtime/ApplicationCore/Private/IOS/IOSInputInterface.cpp
+++ b/Engine/Source/Runtime/ApplicationCore/Private/IOS/IOSInputInterface.cpp
@@ -382,7 +382,10 @@ void FIOSInputInterface::SendControllerEvents()
 #endif
     for(int32 i = 0; i < UE_ARRAY_COUNT(Controllers); ++i)
  	{
-        if(!Controllers[i].bIsGamepadConnected) continue;
+		if (!(Controllers[i].bIsGamepadConnected || Controllers[i].bIsRemoteConnected))
+		{
+			continue;
+		}
         
         GCController* Cont = Controllers[i].Controller;
         GCExtendedGamepadSnapshot* ExtendedGamepad = nullptr;
