Avoid trying to set the clipboard if SDL2 video is not init'ed.
Move fatal log to a warning log, as this is not a fatal issue

diff --git a/Engine/Source/Runtime/ApplicationCore/Private/Linux/LinuxPlatformApplicationMisc.cpp b/Engine/Source/Runtime/ApplicationCore/Private/Linux/LinuxPlatformApplicationMisc.cpp
--- a/Engine/Source/Runtime/ApplicationCore/Private/Linux/LinuxPlatformApplicationMisc.cpp
+++ b/Engine/Source/Runtime/ApplicationCore/Private/Linux/LinuxPlatformApplicationMisc.cpp
@@ -516,9 +516,9 @@ float FLinuxPlatformApplicationMisc::GetDPIScaleFactorAtPoint(float X, float Y)
 
 void FLinuxPlatformApplicationMisc::ClipboardCopy(const TCHAR* Str)
 {
-	if (SDL_SetClipboardText(TCHAR_TO_UTF8(Str)))
+	if (SDL_WasInit(SDL_INIT_VIDEO) && SDL_SetClipboardText(TCHAR_TO_UTF8(Str)))
 	{
-		UE_LOG(LogInit, Fatal, TEXT("Error copying clipboard contents: %s\n"), UTF8_TO_TCHAR(SDL_GetError()));
+		UE_LOG(LogInit, Warning, TEXT("Error copying clipboard contents: %s\n"), UTF8_TO_TCHAR(SDL_GetError()));
 	}
 }
 
