We no longer create CreateInputDevices on several platforms until GIsRunning is true, otherwise we can miss some plugins depending on load/init order.

diff --git a/Engine/Source/Runtime/ApplicationCore/Private/Android/AndroidApplication.cpp b/Engine/Source/Runtime/ApplicationCore/Private/Android/AndroidApplication.cpp
--- a/Engine/Source/Runtime/ApplicationCore/Private/Android/AndroidApplication.cpp
+++ b/Engine/Source/Runtime/ApplicationCore/Private/Android/AndroidApplication.cpp
@@ -54,7 +54,7 @@ void FAndroidApplication::AddExternalInputDevice(TSharedPtr<IInputDevice> InputD
 void FAndroidApplication::PollGameDeviceState( const float TimeDelta )
 {
 	// initialize any externally-implemented input devices (we delay load initialize the array so any plugins have had time to load)
-	if (!bHasLoadedInputPlugins)
+	if (!bHasLoadedInputPlugins && GIsRunning)
 	{
 		TArray<IInputDeviceModule*> PluginImplementations = IModularFeatures::Get().GetModularFeatureImplementations<IInputDeviceModule>(IInputDeviceModule::GetModularFeatureName());
 		for (auto InputPluginIt = PluginImplementations.CreateIterator(); InputPluginIt; ++InputPluginIt)
diff --git a/Engine/Source/Runtime/ApplicationCore/Private/HoloLens/HoloLensApplication.cpp b/Engine/Source/Runtime/ApplicationCore/Private/HoloLens/HoloLensApplication.cpp
--- a/Engine/Source/Runtime/ApplicationCore/Private/HoloLens/HoloLensApplication.cpp
+++ b/Engine/Source/Runtime/ApplicationCore/Private/HoloLens/HoloLensApplication.cpp
@@ -78,7 +78,7 @@ void FHoloLensApplication::PollGameDeviceState( const float TimeDelta )
 {
 	// @MIXEDREALITY_CHANGE : BEGIN Mixed Reality spatial input support
 	// Initialize any externally-implemented input devices (we delay load initialize the array so any plugins have had time to load)
-	if (!bHasLoadedInputPlugins)
+	if (!bHasLoadedInputPlugins && GIsRunning)
 	{
 		TArray< IInputDeviceModule* > PluginImplementations = IModularFeatures::Get().GetModularFeatureImplementations< IInputDeviceModule >(IInputDeviceModule::GetModularFeatureName());
 		for (auto InputPluginIt = PluginImplementations.CreateIterator(); InputPluginIt; ++InputPluginIt)
diff --git a/Engine/Source/Runtime/ApplicationCore/Private/IOS/IOSApplication.cpp b/Engine/Source/Runtime/ApplicationCore/Private/IOS/IOSApplication.cpp
--- a/Engine/Source/Runtime/ApplicationCore/Private/IOS/IOSApplication.cpp
+++ b/Engine/Source/Runtime/ApplicationCore/Private/IOS/IOSApplication.cpp
@@ -73,7 +73,7 @@ void FIOSApplication::AddExternalInputDevice(TSharedPtr<IInputDevice> InputDevic
 void FIOSApplication::PollGameDeviceState( const float TimeDelta )
 {
 	// initialize any externally-implemented input devices (we delay load initialize the array so any plugins have had time to load)
-	if (!bHasLoadedInputPlugins)
+	if (!bHasLoadedInputPlugins && GIsRunning)
 	{
 		TArray<IInputDeviceModule*> PluginImplementations = IModularFeatures::Get().GetModularFeatureImplementations<IInputDeviceModule>(IInputDeviceModule::GetModularFeatureName());
 		for (auto InputPluginIt = PluginImplementations.CreateIterator(); InputPluginIt; ++InputPluginIt)
diff --git a/Engine/Source/Runtime/ApplicationCore/Private/Linux/LinuxApplication.cpp b/Engine/Source/Runtime/ApplicationCore/Private/Linux/LinuxApplication.cpp
--- a/Engine/Source/Runtime/ApplicationCore/Private/Linux/LinuxApplication.cpp
+++ b/Engine/Source/Runtime/ApplicationCore/Private/Linux/LinuxApplication.cpp
@@ -1093,7 +1093,7 @@ void FLinuxApplication::PollGameDeviceState( const float TimeDelta )
 	}
 
 	// initialize any externally-implemented input devices (we delay load initialize the array so any plugins have had time to load)
-	if (!bHasLoadedInputPlugins)
+	if (!bHasLoadedInputPlugins && GIsRunning)
 	{
 		TArray<IInputDeviceModule*> PluginImplementations = IModularFeatures::Get().GetModularFeatureImplementations<IInputDeviceModule>(IInputDeviceModule::GetModularFeatureName());
 		for (auto InputPluginIt = PluginImplementations.CreateIterator(); InputPluginIt; ++InputPluginIt)
diff --git a/Engine/Source/Runtime/ApplicationCore/Private/Mac/MacApplication.cpp b/Engine/Source/Runtime/ApplicationCore/Private/Mac/MacApplication.cpp
--- a/Engine/Source/Runtime/ApplicationCore/Private/Mac/MacApplication.cpp
+++ b/Engine/Source/Runtime/ApplicationCore/Private/Mac/MacApplication.cpp
@@ -233,7 +233,7 @@ void FMacApplication::SetAccessibleMessageHandler(const TSharedRef<FGenericAcces
 void FMacApplication::PollGameDeviceState(const float TimeDelta)
 {
 	// initialize any externally-implemented input devices (we delay load initialize the array so any plugins have had time to load)
-	if (!bHasLoadedInputPlugins)
+	if (!bHasLoadedInputPlugins && GIsRunning)
 	{
 		TArray<IInputDeviceModule*> PluginImplementations = IModularFeatures::Get().GetModularFeatureImplementations<IInputDeviceModule>( IInputDeviceModule::GetModularFeatureName() );
 		for( auto InputPluginIt = PluginImplementations.CreateIterator(); InputPluginIt; ++InputPluginIt )
diff --git a/Engine/Source/Runtime/ApplicationCore/Private/Windows/WindowsApplication.cpp b/Engine/Source/Runtime/ApplicationCore/Private/Windows/WindowsApplication.cpp
--- a/Engine/Source/Runtime/ApplicationCore/Private/Windows/WindowsApplication.cpp
+++ b/Engine/Source/Runtime/ApplicationCore/Private/Windows/WindowsApplication.cpp
@@ -2641,7 +2641,7 @@ void FWindowsApplication::PollGameDeviceState( const float TimeDelta )
 	}
 
 	// initialize any externally-implemented input devices (we delay load initialize the array so any plugins have had time to load)
-	if (!bHasLoadedInputPlugins)
+	if (!bHasLoadedInputPlugins && GIsRunning )
 	{
 		TArray<IInputDeviceModule*> PluginImplementations = IModularFeatures::Get().GetModularFeatureImplementations<IInputDeviceModule>( IInputDeviceModule::GetModularFeatureName() );
 		for( auto InputPluginIt = PluginImplementations.CreateIterator(); InputPluginIt; ++InputPluginIt )
