addition of iPad Pro 2020 profiles (12.9 and 11)

diff --git a/Engine/Config/BaseDeviceProfiles.ini b/Engine/Config/BaseDeviceProfiles.ini
--- a/Engine/Config/BaseDeviceProfiles.ini
+++ b/Engine/Config/BaseDeviceProfiles.ini
@@ -44,7 +44,11 @@
 +DeviceProfileNameAndTypes=iPad6,IOS
 +DeviceProfileNameAndTypes=iPad7,IOS
 +DeviceProfileNameAndTypes=iPadPro11,IOS
++DeviceProfileNameAndTypes=iPadPro2_11,IOS
 +DeviceProfileNameAndTypes=iPadPro3_129,IOS
++DeviceProfileNameAndTypes=iPadPro4_129,IOS
++DeviceProfileNameAndTypes=iPadAir3,IOS
++DeviceProfileNameAndTypes=iPadMini5,IOS
 +DeviceProfileNameAndTypes=AppleTV,IOS
 +DeviceProfileNameAndTypes=AppleTV4K,IOS
 +DeviceProfileNameAndTypes=TVOS,AppleTV
@@ -365,6 +369,10 @@ BaseProfileName=IPadPro
 DeviceType=IOS
 BaseProfileName=IPadPro
 
+[iPadPro2_11 DeviceProfile]
+DeviceType=IOS
+BaseProfileName=IPadPro
+
 [iPad5 DeviceProfile]
 DeviceType=IOS
 BaseProfileName=IPadPro97
@@ -385,6 +393,10 @@ BaseProfileName=IPadPro
 DeviceType=IOS
 BaseProfileName=IPadPro
 
+[iPadPro4_129 DeviceProfile]
+DeviceType=IOS
+BaseProfileName=IPadPro
+
 [iPadAir3 DeviceProfile]
 DeviceType=IOS
 BaseProfileName=IPadPro
diff --git a/Engine/Source/Runtime/ApplicationCore/Private/IOS/IOSPlatformApplicationMisc.cpp b/Engine/Source/Runtime/ApplicationCore/Private/IOS/IOSPlatformApplicationMisc.cpp
--- a/Engine/Source/Runtime/ApplicationCore/Private/IOS/IOSPlatformApplicationMisc.cpp
+++ b/Engine/Source/Runtime/ApplicationCore/Private/IOS/IOSPlatformApplicationMisc.cpp
@@ -258,7 +258,7 @@ void FIOSPlatformApplicationMisc::ClipboardPaste(class FString& Result)
 EScreenPhysicalAccuracy FIOSPlatformApplicationMisc::ComputePhysicalScreenDensity(int32& ScreenDensity)
 {
 	FPlatformMisc::EIOSDevice Device = FPlatformMisc::GetIOSDeviceType();
-	static_assert( FPlatformMisc::EIOSDevice::IOS_Unknown == 54, "Every device needs to be handled here." );
+	static_assert( FPlatformMisc::EIOSDevice::IOS_Unknown == 56, "Every device needs to be handled here." );
 
 	ScreenDensity = 0;
 	EScreenPhysicalAccuracy Accuracy = EScreenPhysicalAccuracy::Unknown;
@@ -322,8 +322,10 @@ EScreenPhysicalAccuracy FIOSPlatformApplicationMisc::ComputePhysicalScreenDensit
 	case FPlatformMisc::IOS_IPadPro_129:
 	case FPlatformMisc::IOS_IPadPro_105:
 	case FPlatformMisc::IOS_IPadPro_11:
+	case FPlatformMisc::IOS_IPadPro2_11:
 	case FPlatformMisc::IOS_IPadPro2_129:
 	case FPlatformMisc::IOS_IPadPro3_129:
+	case FPlatformMisc::IOS_IPadPro4_129:
 		ScreenDensity = 264;
 		Accuracy = EScreenPhysicalAccuracy::Truth;
 		break;
diff --git a/Engine/Source/Runtime/Core/Private/IOS/IOSPlatformMisc.cpp b/Engine/Source/Runtime/Core/Private/IOS/IOSPlatformMisc.cpp
--- a/Engine/Source/Runtime/Core/Private/IOS/IOSPlatformMisc.cpp
+++ b/Engine/Source/Runtime/Core/Private/IOS/IOSPlatformMisc.cpp
@@ -387,6 +387,8 @@ const TCHAR* FIOSPlatformMisc::GetDefaultDeviceProfileName()
 		TEXT("IPodTouch7"),
 		TEXT("IPad7"),
 		TEXT("IPhoneSE2"),
+		TEXT("IPadPro2_11"),
+		TEXT("IPadPro4_129"),
 		TEXT("NewIDevice1"),
 		TEXT("NewIDevice2"),
 		TEXT("NewIDevice3"),
@@ -558,10 +560,18 @@ FIOSPlatformMisc::EIOSDevice FIOSPlatformMisc::GetIOSDeviceType()
 			{
 				DeviceType = IOS_IPadPro_11;
 			}
-			else
+			else if (Minor <= 8)
 			{
 				DeviceType = IOS_IPadPro3_129;
 			}
+			else if (Minor <= 10)
+			{
+				DeviceType = IOS_IPadPro2_11;
+			}
+			else
+			{
+				DeviceType = IOS_IPadPro4_129;
+			}
 		}
         else if (Major == 11)
         {
@@ -577,7 +587,7 @@ FIOSPlatformMisc::EIOSDevice FIOSPlatformMisc::GetIOSDeviceType()
 		// Default to highest settings currently available for any future device
 		else if (Major >= 9)
 		{
-			DeviceType = IOS_IPadPro3_129;
+			DeviceType = IOS_IPadPro4_129;
 		}
 	}
 	// iPhones
diff --git a/Engine/Source/Runtime/Core/Public/IOS/IOSPlatformMisc.h b/Engine/Source/Runtime/Core/Public/IOS/IOSPlatformMisc.h
--- a/Engine/Source/Runtime/Core/Public/IOS/IOSPlatformMisc.h
+++ b/Engine/Source/Runtime/Core/Public/IOS/IOSPlatformMisc.h
@@ -156,6 +156,8 @@ struct CORE_API FIOSPlatformMisc : public FApplePlatformMisc
 		IOS_IPodTouch7,
 		IOS_IPad7,
 		IOS_IPhoneSE2,
+		IOS_IPadPro2_11,
+		IOS_IPadPro4_129,
 
 		// We can use the entries below for any iOS devices released during the hotfix cycle
 		// They should be moved to real device enum above these values in the next full release.
