Trying out drop shadows on top level windows

diff --git a/Engine/Source/Runtime/ApplicationCore/Private/Windows/WindowsApplication.cpp b/Engine/Source/Runtime/ApplicationCore/Private/Windows/WindowsApplication.cpp
--- a/Engine/Source/Runtime/ApplicationCore/Private/Windows/WindowsApplication.cpp
+++ b/Engine/Source/Runtime/ApplicationCore/Private/Windows/WindowsApplication.cpp
@@ -291,7 +291,7 @@ bool FWindowsApplication::RegisterClass( const HINSTANCE HInstance, const HICON
 {
 	WNDCLASS wc;
 	FMemory::Memzero( &wc, sizeof(wc) );
-	wc.style = CS_DBLCLKS; // We want to receive double clicks
+	wc.style = CS_DBLCLKS | CS_DROPSHADOW; // We want to receive double clicks
 	wc.lpfnWndProc = AppWndProc;
 	wc.cbClsExtra = 0;
 	wc.cbWndExtra = 0;
