Fix wrong styling when app has an OS window border

diff --git a/Engine/Source/Runtime/ApplicationCore/Private/Windows/WindowsWindow.cpp b/Engine/Source/Runtime/ApplicationCore/Private/Windows/WindowsWindow.cpp
--- a/Engine/Source/Runtime/ApplicationCore/Private/Windows/WindowsWindow.cpp
+++ b/Engine/Source/Runtime/ApplicationCore/Private/Windows/WindowsWindow.cpp
@@ -468,7 +468,7 @@ void FWindowsWindow::ReshapeWindow( int32 NewX, int32 NewY, int32 NewWidth, int3
 	// prevents us being able to change to a higher resolution while in fullscreen mode
 	::SetWindowPos( HWnd, nullptr, WindowX, WindowY, NewWidth, NewHeight, SWP_NOZORDER | SWP_NOACTIVATE | ((WindowMode == EWindowMode::Fullscreen) ? SWP_NOSENDCHANGING : 0) );
 
-	if( (Definition->SizeWillChangeOften && bVirtualSizeChanged) || Definition->CornerRadius > 0 )
+	if( !Definition->HasOSWindowBorder && (((Definition->SizeWillChangeOften && bVirtualSizeChanged) || Definition->CornerRadius > 0)) )
 	{
 		AdjustWindowRegion( VirtualWidth, VirtualHeight );
 	}
