Change log severity from error to warning

diff --git a/Engine/Source/Runtime/ApplicationCore/Private/Windows/WindowsPlatformApplicationMisc.cpp b/Engine/Source/Runtime/ApplicationCore/Private/Windows/WindowsPlatformApplicationMisc.cpp
--- a/Engine/Source/Runtime/ApplicationCore/Private/Windows/WindowsPlatformApplicationMisc.cpp
+++ b/Engine/Source/Runtime/ApplicationCore/Private/Windows/WindowsPlatformApplicationMisc.cpp
@@ -379,7 +379,7 @@ void FWindowsPlatformApplicationMisc::ClipboardCopy(const TCHAR* Str)
 	}
 	else
 	{
-		UE_LOG(LogWindows, Error, TEXT("OpenClipboard failed with error code %i"), (uint32)GetLastError());
+		UE_LOG(LogWindows, Warning, TEXT("OpenClipboard failed with error code %i"), (uint32)GetLastError());
 	}
 }
 
@@ -424,7 +424,7 @@ void FWindowsPlatformApplicationMisc::ClipboardPaste(class FString& Result)
 	else 
 	{
 		Result=TEXT("");
-		UE_LOG(LogWindows, Error, TEXT("OpenClipboard failed with error code %i"), (uint32)GetLastError());
+		UE_LOG(LogWindows, Warning, TEXT("OpenClipboard failed with error code %i"), (uint32)GetLastError());
 	}
 }
 
