Fix iOS application launch.

#rnx

diff --git a/Engine/Source/Runtime/ApplicationCore/Private/IOS/IOSAppDelegate.cpp b/Engine/Source/Runtime/ApplicationCore/Private/IOS/IOSAppDelegate.cpp
--- a/Engine/Source/Runtime/ApplicationCore/Private/IOS/IOSAppDelegate.cpp
+++ b/Engine/Source/Runtime/ApplicationCore/Private/IOS/IOSAppDelegate.cpp
@@ -303,6 +303,8 @@ static IOSAppDelegate* CachedDelegate = nil;
 		}
 	}
 
+    FTaskTagScope Scope(ETaskTag::EGameThread);
+
 	FAppEntry::Init();
 
 	// check for update on app store if cvar is enabled
diff --git a/Engine/Source/Runtime/ApplicationCore/Private/IOS/IOSView.cpp b/Engine/Source/Runtime/ApplicationCore/Private/IOS/IOSView.cpp
--- a/Engine/Source/Runtime/ApplicationCore/Private/IOS/IOSView.cpp
+++ b/Engine/Source/Runtime/ApplicationCore/Private/IOS/IOSView.cpp
@@ -274,15 +274,19 @@ id<MTLDevice> GMetalDevice = nil;
 		// look up what the device can support
 		const float NativeScale = [[UIScreen mainScreen] scale];
 
-		// look up the CVar for the scale factor
-		static IConsoleVariable* CVar = IConsoleManager::Get().FindConsoleVariable(TEXT("r.MobileContentScaleFactor"));
-		float RequestedContentScaleFactor = CVar->GetFloat();
-
-		FString CmdLineCSF;
-		if (FParse::Value(FCommandLine::Get(), TEXT("mcsf="), CmdLineCSF, false))
+		float RequestedContentScaleFactor = 1.0;
+		[IOSAppDelegate WaitAndRunOnGameThread:[&RequestedContentScaleFactor]()
 		{
-			RequestedContentScaleFactor = FCString::Atof(*CmdLineCSF);
-		}
+			// look up the CVar for the scale factor
+			static IConsoleVariable* CVar = IConsoleManager::Get().FindConsoleVariable(TEXT("r.MobileContentScaleFactor"));
+			RequestedContentScaleFactor = CVar->GetFloat();
+
+			FString CmdLineCSF;
+			if (FParse::Value(FCommandLine::Get(), TEXT("mcsf="), CmdLineCSF, false))
+			{
+				RequestedContentScaleFactor = FCString::Atof(*CmdLineCSF);
+			}
+		}];
 
 		// 0 means to leave the scale alone, use native
 		if (RequestedContentScaleFactor == 0.0f)
diff --git a/Engine/Source/Runtime/Core/Core.Build.cs b/Engine/Source/Runtime/Core/Core.Build.cs
--- a/Engine/Source/Runtime/Core/Core.Build.cs
+++ b/Engine/Source/Runtime/Core/Core.Build.cs
@@ -223,7 +223,7 @@ public class Core : ModuleRules
 			PublicDefinitions.Add("WITH_DIRECTXMATH=0");
 		}
 
-		if ((Target.Platform == UnrealTargetPlatform.Mac) || (Target.Platform == UnrealTargetPlatform.HoloLens))
+		if ((Target.Platform == UnrealTargetPlatform.Mac) || (Target.Platform == UnrealTargetPlatform.IOS) || (Target.Platform == UnrealTargetPlatform.TVOS) || (Target.Platform == UnrealTargetPlatform.HoloLens))
 		{
 			PublicDefinitions.Add("IS_RUNNING_GAMETHREAD_ON_EXTERNAL_THREAD=1");
 		}
