Remove legacy code that change the game-thread priority while the editor doesn't have focus, we have too many threads in the engine for that to make a meaningful difference and it is causing latency issues by making the background workers have the same priority as the game thread when not in focus...
Windows already does that better for the application that has the focus

diff --git a/Engine/Source/Runtime/ApplicationCore/Private/Windows/WindowsPlatformApplicationMisc.cpp b/Engine/Source/Runtime/ApplicationCore/Private/Windows/WindowsPlatformApplicationMisc.cpp
--- a/Engine/Source/Runtime/ApplicationCore/Private/Windows/WindowsPlatformApplicationMisc.cpp
+++ b/Engine/Source/Runtime/ApplicationCore/Private/Windows/WindowsPlatformApplicationMisc.cpp
@@ -138,16 +138,6 @@ void FWindowsPlatformApplicationMisc::PumpMessages(bool bFromMainLoop)
 	// If editor thread doesn't have the focus, don't suck up too much CPU time.
 	if( GIsEditor )
 	{
-		if( HadFocus && !HasFocus )
-		{
-			// Drop our priority to speed up whatever is in the foreground.
-			SetThreadPriority( GetCurrentThread(), THREAD_PRIORITY_BELOW_NORMAL );
-		}
-		else if( HasFocus && !HadFocus )
-		{
-			// Boost our priority back to above normal as initially set in WindowsRunnableThread::CreateInternal.
-			SetThreadPriority( GetCurrentThread(), THREAD_PRIORITY_ABOVE_NORMAL );
-		}
 		if( !HasFocus )
 		{
 			// Sleep for a bit to not eat up all CPU time.
