Move calls to setAlphaValue: to the main application thread.

#rnx

diff --git a/Engine/Source/Programs/CrashReportClient/Private/Mac/CrashReportClientMainMac.cpp b/Engine/Source/Programs/CrashReportClient/Private/Mac/CrashReportClientMainMac.cpp
--- a/Engine/Source/Programs/CrashReportClient/Private/Mac/CrashReportClientMainMac.cpp
+++ b/Engine/Source/Programs/CrashReportClient/Private/Mac/CrashReportClientMainMac.cpp
@@ -18,7 +18,7 @@ void CrashReporterCrashHandler(const FGenericCrashContext& GenericContext)
 }
 
 static FString GSavedCommandLine;
-static bool GIsUnattended;
+static bool GIsUnattended = false;
 
 @interface UE4AppDelegate : NSObject <NSApplicationDelegate, NSFileManagerDelegate>
 {
diff --git a/Engine/Source/Runtime/ApplicationCore/Private/Mac/CocoaWindow.cpp b/Engine/Source/Runtime/ApplicationCore/Private/Mac/CocoaWindow.cpp
--- a/Engine/Source/Runtime/ApplicationCore/Private/Mac/CocoaWindow.cpp
+++ b/Engine/Source/Runtime/ApplicationCore/Private/Mac/CocoaWindow.cpp
@@ -149,7 +149,10 @@ NSString* NSPerformDragOperation = @"NSPerformDragOperation";
 	Opacity = WindowAlpha;
 	if (bRenderInitialized)
 	{
-		[super setAlphaValue:WindowAlpha];
+		MainThreadCall(^{
+			[super setAlphaValue:WindowAlpha];
+		}, UE4ShowEventMode, false);
+
 	}
 }
 
@@ -158,7 +161,9 @@ NSString* NSPerformDragOperation = @"NSPerformDragOperation";
 	if (!bRenderInitialized)
 	{
 		bRenderInitialized = true;
-		[super setAlphaValue:Opacity];
+		MainThreadCall(^{
+			[super setAlphaValue:Opacity];
+		}, UE4ShowEventMode, false);
 	}
 }
 
