Remove window size hack

diff --git a/Engine/Source/Runtime/ApplicationCore/Private/Linux/LinuxWindow.cpp b/Engine/Source/Runtime/ApplicationCore/Private/Linux/LinuxWindow.cpp
--- a/Engine/Source/Runtime/ApplicationCore/Private/Linux/LinuxWindow.cpp
+++ b/Engine/Source/Runtime/ApplicationCore/Private/Linux/LinuxWindow.cpp
@@ -510,30 +510,6 @@ static void GetBestFullscreenResolution( SDL_HWindow hWnd, int32 *pWidth, int32
 
 void FLinuxWindow::ReshapeWindow( int32 NewX, int32 NewY, int32 NewWidth, int32 NewHeight )
 {
-	// Some vulkan video drivers have issues with specific height ranges causing them to corrupt the texture rendered
-	// Moving these nearest values removes this corruption.
-	static bool bDisableVulkanWorkaround = FParse::Param(FCommandLine::Get(), TEXT("disablevulkanworkaround"));
-
-	if (!bDisableVulkanWorkaround)
-	{
-		if (NewHeight >= 9 && NewHeight <= 10)
-		{
-			NewHeight = 11;
-		}
-		else if (NewHeight >= 17 && NewHeight <= 21)
-		{
-			NewHeight = 22;
-		}
-		else if (NewHeight >= 33 && NewHeight <= 43)
-		{
-			NewHeight = 44;
-		}
-		else if (NewHeight >= 65 && NewHeight <= 85)
-		{
-			NewHeight = 86;
-		}
-	}
-
 	// If we have set our self to 0,0 Width/Height it will not be allowed we will still show the window
 	// this is a work around to at least reduce the visibile impact of a window that is lingering
 	NewWidth  = FMath::Max(NewWidth, 1);
