Unshelved from pending changelist '15636883':

Remove iOS 12 and set Metal shading default to 2.2 #jira UE-106400 #rb jack.porter #review-15719081 jack.porter

diff --git a/Engine/Config/BaseEngine.ini b/Engine/Config/BaseEngine.ini
--- a/Engine/Config/BaseEngine.ini
+++ b/Engine/Config/BaseEngine.ini
@@ -2206,7 +2206,7 @@ DistroCodeSigningIdentity=iPhone Distribution:;Apple Distribution:
 VersionInfo=1.0
 FrameRateLock=PUFRL_30
 bEnableDynamicMaxFPS=False
-MinimumiOSVersion=IOS_12
+MinimumiOSVersion=IOS_13
 bSupportsIPad=True
 bSupportsIPhone=True
 AdditionalPlistData=
diff --git a/Engine/Plugins/Experimental/AppleVision/Source/AppleVision/Public/AppleVisionAvailability.h b/Engine/Plugins/Experimental/AppleVision/Source/AppleVision/Public/AppleVisionAvailability.h
--- a/Engine/Plugins/Experimental/AppleVision/Source/AppleVision/Public/AppleVisionAvailability.h
+++ b/Engine/Plugins/Experimental/AppleVision/Source/AppleVision/Public/AppleVisionAvailability.h
@@ -37,10 +37,7 @@ public:
 		{
 			bSupportChecked = true;
 			// This call is slow, so cache the result
-			if (@available(iOS 11.0, *))
-			{
-				bSupportsAppleVision10 = true;
-			}
+			bSupportsAppleVision10 = true;
 		}
 #endif
 		return bSupportsAppleVision10;
diff --git a/Engine/Plugins/Online/OnlineSubsystemApple/Source/Private/OnlineExternalUIInterfaceApple.cpp b/Engine/Plugins/Online/OnlineSubsystemApple/Source/Private/OnlineExternalUIInterfaceApple.cpp
--- a/Engine/Plugins/Online/OnlineSubsystemApple/Source/Private/OnlineExternalUIInterfaceApple.cpp
+++ b/Engine/Plugins/Online/OnlineSubsystemApple/Source/Private/OnlineExternalUIInterfaceApple.cpp
@@ -199,87 +199,85 @@ FOnlineExternalUIApple::~FOnlineExternalUIApple()
 
 bool FOnlineExternalUIApple::ShowLoginUI(const int ControllerIndex, bool bShowOnlineOnly, bool bShowSkipButton, const FOnLoginUIClosedDelegate& Delegate)
 {
-	bool bStarted = false;
-	FString ErrorStr;
-
+    bool bStarted = false;
+    FString ErrorStr;
+    
 #if SIWA_SUPPORTED
-	if (@available(iOS 13, tvOS 13, macOS 10.15, *))
-	{
-		if (ControllerIndex >= 0 && ControllerIndex < MAX_LOCAL_PLAYERS)
-		{
-			FOnlineIdentityApplePtr IdentityInterface = StaticCastSharedPtr<FOnlineIdentityApple>(Subsystem->GetIdentityInterface());
-			if (IdentityInterface.IsValid())
-			{
-				// Is the user already logged in?
-				if (IdentityInterface->GetLoginStatus(ControllerIndex)!= ELoginStatus::NotLoggedIn)
-				{
-					Subsystem->ExecuteNextTick([IdentityInterface, ControllerIndex, Delegate]()
-					{
-						Delegate.ExecuteIfBound(IdentityInterface->GetUniquePlayerId(ControllerIndex), ControllerIndex, FOnlineError::Success());
-					});
-					bStarted = true;
-				}
-				else
-				{
-					LoginUIClosedDelegate = Delegate;
-
-					// Show the Sign in with Apple login UI
-					OnSignInCompleteHandle = SignInBridge->AddOnSignInCompleteDelegate_Handle(FOnSignInCompleteDelegate::CreateRaw(this, &FOnlineExternalUIApple::OnSignInComplete, ControllerIndex));
-
-					ASAuthorizationAppleIDProvider *Provider = [ASAuthorizationAppleIDProvider new];
-					ASAuthorizationAppleIDRequest *Request = Provider.createRequest;
-
-					if (ScopeFields.Num() > 0)
-					{
-						// Define any requested scopes
-						// NOTE: We will only receieve scopes for the initial login only, once the app is authorised extended details are no longer available
-						NSMutableArray* Permissions = [NSMutableArray array];
-						for (int32 ScopeIdx = 0; ScopeIdx < ScopeFields.Num(); ScopeIdx++)
-						{
-							if (ScopeFields[ScopeIdx].Equals(TEXT(APPLE_PERM_EMAIL), ESearchCase::CaseSensitive))
-							{
-								[Permissions addObject: ASAuthorizationScopeEmail];
-							}
-							else if (ScopeFields[ScopeIdx].Equals(TEXT(APPLE_PERM_FULLNAME), ESearchCase::CaseSensitive))
-							{
-								[Permissions addObject: ASAuthorizationScopeFullName];
-							}
-							else
-							{
-								UE_LOG_ONLINE_IDENTITY(Warning, TEXT("Login invalid scope. %s"), *ScopeFields[ScopeIdx]);
-							}
-						}
-
-						Request.requestedScopes = Permissions;
-					}
-
-					ASAuthorizationController *Controller = [[ASAuthorizationController alloc] initWithAuthorizationRequests:@[Request]];
-
-					SignInController* SignInController = SignInBridge->GetSignInController();
-					Controller.delegate = SignInController;
-					Controller.presentationContextProvider = SignInController;
-					[Controller performRequests];
-
-					bStarted = true;
-				}
-			}
-			else
-			{
-				ErrorStr = TEXT("ShowLoginUI: Missing identity interface");
-			}
-		}
-		else
-		{
-			ErrorStr = FString::Printf(TEXT("ShowLoginUI: Invalid controller index (%d)"), ControllerIndex);
-		}
-	}
+    if (ControllerIndex >= 0 && ControllerIndex < MAX_LOCAL_PLAYERS)
+    {
+        FOnlineIdentityApplePtr IdentityInterface = StaticCastSharedPtr<FOnlineIdentityApple>(Subsystem->GetIdentityInterface());
+        if (IdentityInterface.IsValid())
+        {
+            // Is the user already logged in?
+            if (IdentityInterface->GetLoginStatus(ControllerIndex)!= ELoginStatus::NotLoggedIn)
+            {
+                Subsystem->ExecuteNextTick([IdentityInterface, ControllerIndex, Delegate]()
+                                           {
+                    Delegate.ExecuteIfBound(IdentityInterface->GetUniquePlayerId(ControllerIndex), ControllerIndex, FOnlineError::Success());
+                });
+                bStarted = true;
+            }
+            else
+            {
+                LoginUIClosedDelegate = Delegate;
+                
+                // Show the Sign in with Apple login UI
+                OnSignInCompleteHandle = SignInBridge->AddOnSignInCompleteDelegate_Handle(FOnSignInCompleteDelegate::CreateRaw(this, &FOnlineExternalUIApple::OnSignInComplete, ControllerIndex));
+                
+                ASAuthorizationAppleIDProvider *Provider = [ASAuthorizationAppleIDProvider new];
+                ASAuthorizationAppleIDRequest *Request = Provider.createRequest;
+                
+                if (ScopeFields.Num() > 0)
+                {
+                    // Define any requested scopes
+                    // NOTE: We will only receieve scopes for the initial login only, once the app is authorised extended details are no longer available
+                    NSMutableArray* Permissions = [NSMutableArray array];
+                    for (int32 ScopeIdx = 0; ScopeIdx < ScopeFields.Num(); ScopeIdx++)
+                    {
+                        if (ScopeFields[ScopeIdx].Equals(TEXT(APPLE_PERM_EMAIL), ESearchCase::CaseSensitive))
+                        {
+                            [Permissions addObject: ASAuthorizationScopeEmail];
+                        }
+                        else if (ScopeFields[ScopeIdx].Equals(TEXT(APPLE_PERM_FULLNAME), ESearchCase::CaseSensitive))
+                        {
+                            [Permissions addObject: ASAuthorizationScopeFullName];
+                        }
+                        else
+                        {
+                            UE_LOG_ONLINE_IDENTITY(Warning, TEXT("Login invalid scope. %s"), *ScopeFields[ScopeIdx]);
+                        }
+                    }
+                    
+                    Request.requestedScopes = Permissions;
+                }
+                
+                ASAuthorizationController *Controller = [[ASAuthorizationController alloc] initWithAuthorizationRequests:@[Request]];
+                
+                SignInController* SignInController = SignInBridge->GetSignInController();
+                Controller.delegate = SignInController;
+                Controller.presentationContextProvider = SignInController;
+                [Controller performRequests];
+                
+                bStarted = true;
+            }
+        }
+        else
+        {
+            ErrorStr = TEXT("ShowLoginUI: Missing identity interface");
+        }
+    }
+    else
+    {
+        ErrorStr = FString::Printf(TEXT("ShowLoginUI: Invalid controller index (%d)"), ControllerIndex);
+    }
+    
 #else
-	ErrorStr = TEXT("Sign in with Apple is not available");
+    ErrorStr = TEXT("Sign in with Apple is not available");
 #endif
-
-	if (!bStarted)
-	{
-		UE_LOG_ONLINE_EXTERNALUI(Warning, TEXT("%s"), *ErrorStr);
+    
+    if (!bStarted)
+    {
+        UE_LOG_ONLINE_EXTERNALUI(Warning, TEXT("%s"), *ErrorStr);
 
 		FOnlineError Error(false);
 		Error.SetFromErrorCode(MoveTemp(ErrorStr));
@@ -288,7 +286,6 @@ bool FOnlineExternalUIApple::ShowLoginUI(const int ControllerIndex, bool bShowOn
 			Delegate.ExecuteIfBound(nullptr, ControllerIndex, Error);
 		});
 	}
-
 	return bStarted;
 }
 
diff --git a/Engine/Plugins/Online/OnlineSubsystemApple/Source/Private/OnlineIdentityInterfaceApple.cpp b/Engine/Plugins/Online/OnlineSubsystemApple/Source/Private/OnlineIdentityInterfaceApple.cpp
--- a/Engine/Plugins/Online/OnlineSubsystemApple/Source/Private/OnlineIdentityInterfaceApple.cpp
+++ b/Engine/Plugins/Online/OnlineSubsystemApple/Source/Private/OnlineIdentityInterfaceApple.cpp
@@ -105,89 +105,86 @@ bool FOnlineIdentityApple::Login(int32 LocalUserNum, const FOnlineAccountCredent
 	UE_LOG_ONLINE_IDENTITY(Verbose, TEXT("FOnlineIdentityApple::Login"));
 
 #if SIWA_SUPPORTED
-	if (@available(iOS 13, tvOS 13, macOS 10.15, *))
-	{
-		ELoginStatus::Type LoginStatus = GetLoginStatus(LocalUserNum);
-
-		if (LoginStatus != ELoginStatus::NotLoggedIn)
-		{
-			Subsystem->ExecuteNextTick([this, LocalUserNum]()
-			{
-				TriggerOnLoginCompleteDelegates(LocalUserNum, true, *GetUniquePlayerId(LocalUserNum), TEXT("Already logged in"));
-			});
-			return false;
-		}
-
-		ensure(LoginStatus == ELoginStatus::NotLoggedIn);
-		if (!AccountCredentials.Id.IsEmpty())
-		{
-			FString PlayerId = AccountCredentials.Id;
-
-			FTCHARToUTF8 UserIdStr(*PlayerId);
-			NSString *UserId = [NSString stringWithUTF8String:UserIdStr.Get()];
-
-			// Attempt to login with existing credentials, we only validate the Id, no full account details are created or need storing
-			ASAuthorizationAppleIDProvider *Provider = [ASAuthorizationAppleIDProvider new];
-			[Provider getCredentialStateForUserID:UserId completion:^(ASAuthorizationAppleIDProviderCredentialState credentialState, NSError *error)
-			{
-				bool bWasSuccessful = false;
-				FString ErrorStr;
-
-				switch (credentialState) {
-					case ASAuthorizationAppleIDProviderCredentialAuthorized:
-						{
-							// We receive no additional information or updated token, only that the supplied Id is still valid, or not on this device/application
-							TSharedRef<FUserOnlineAccountApple> User = MakeShared<FUserOnlineAccountApple>(PlayerId, FString());
-
-							// update/add cached entry for user
-							UserAccounts.Add(User->GetUserId()->ToString(), User);
-							// keep track of user ids for local users
-							UserIds.Add(LocalUserNum, User->GetUserId());
-
-							bWasSuccessful = true;
-
-							UE_LOG_ONLINE_IDENTITY(Display, TEXT("Apple login with credentials was successful"));
-						}
-						break;
-
-					case ASAuthorizationAppleIDProviderCredentialRevoked:
-						ErrorStr = TEXT("Credentials have been revoked");
-						break;
-
-					case ASAuthorizationAppleIDProviderCredentialNotFound:
-					default:
-						ErrorStr = TEXT("Credentials not found");
-						break;
-				}
-
-				if (error)
-				{
-					NSString* errstr = [error localizedDescription];
-					ErrorStr = FString::Printf(TEXT("Login failure %s"), *FString(errstr));
-				}
-
-				OnLoginAttemptComplete(LocalUserNum, ErrorStr);
-			}];
-
-			return true;
-		}
-		else
-		{
-			IOnlineExternalUIPtr OnlineExternalUI = Subsystem->GetExternalUIInterface();
-			if (OnlineExternalUI.IsValid())
-			{
-				FOnLoginUIClosedDelegate CompletionDelegate = FOnLoginUIClosedDelegate::CreateRaw(this, &FOnlineIdentityApple::OnExternalUILoginComplete);
- 				OnlineExternalUI->ShowLoginUI(LocalUserNum, true, false, CompletionDelegate);
-				return true;
-			}
-			else
-			{
-				const FString ErrorStr = TEXT("External interface missing");
-				OnLoginAttemptComplete(LocalUserNum, ErrorStr);
-				return false;
-			}
-		}
-	}
+    ELoginStatus::Type LoginStatus = GetLoginStatus(LocalUserNum);
+    
+    if (LoginStatus != ELoginStatus::NotLoggedIn)
+    {
+        Subsystem->ExecuteNextTick([this, LocalUserNum]()
+                                   {
+            TriggerOnLoginCompleteDelegates(LocalUserNum, true, *GetUniquePlayerId(LocalUserNum), TEXT("Already logged in"));
+        });
+        return false;
+    }
+    
+    ensure(LoginStatus == ELoginStatus::NotLoggedIn);
+    if (!AccountCredentials.Id.IsEmpty())
+    {
+        FString PlayerId = AccountCredentials.Id;
+        
+        FTCHARToUTF8 UserIdStr(*PlayerId);
+        NSString *UserId = [NSString stringWithUTF8String:UserIdStr.Get()];
+        
+        // Attempt to login with existing credentials, we only validate the Id, no full account details are created or need storing
+        ASAuthorizationAppleIDProvider *Provider = [ASAuthorizationAppleIDProvider new];
+        [Provider getCredentialStateForUserID:UserId completion:^(ASAuthorizationAppleIDProviderCredentialState credentialState, NSError *error)
+         {
+            bool bWasSuccessful = false;
+            FString ErrorStr;
+            
+            switch (credentialState) {
+                case ASAuthorizationAppleIDProviderCredentialAuthorized:
+                {
+                    // We receive no additional information or updated token, only that the supplied Id is still valid, or not on this device/application
+                    TSharedRef<FUserOnlineAccountApple> User = MakeShared<FUserOnlineAccountApple>(PlayerId, FString());
+                    
+                    // update/add cached entry for user
+                    UserAccounts.Add(User->GetUserId()->ToString(), User);
+                    // keep track of user ids for local users
+                    UserIds.Add(LocalUserNum, User->GetUserId());
+                    
+                    bWasSuccessful = true;
+                    
+                    UE_LOG_ONLINE_IDENTITY(Display, TEXT("Apple login with credentials was successful"));
+                }
+                    break;
+                    
+                case ASAuthorizationAppleIDProviderCredentialRevoked:
+                    ErrorStr = TEXT("Credentials have been revoked");
+                    break;
+                    
+                case ASAuthorizationAppleIDProviderCredentialNotFound:
+                default:
+                    ErrorStr = TEXT("Credentials not found");
+                    break;
+            }
+            
+            if (error)
+            {
+                NSString* errstr = [error localizedDescription];
+                ErrorStr = FString::Printf(TEXT("Login failure %s"), *FString(errstr));
+            }
+            
+            OnLoginAttemptComplete(LocalUserNum, ErrorStr);
+        }];
+        
+        return true;
+    }
+    else
+    {
+        IOnlineExternalUIPtr OnlineExternalUI = Subsystem->GetExternalUIInterface();
+        if (OnlineExternalUI.IsValid())
+        {
+            FOnLoginUIClosedDelegate CompletionDelegate = FOnLoginUIClosedDelegate::CreateRaw(this, &FOnlineIdentityApple::OnExternalUILoginComplete);
+            OnlineExternalUI->ShowLoginUI(LocalUserNum, true, false, CompletionDelegate);
+            return true;
+        }
+        else
+        {
+            const FString ErrorStr = TEXT("External interface missing");
+            OnLoginAttemptComplete(LocalUserNum, ErrorStr);
+            return false;
+        }
+    }
 #endif
 
 	Subsystem->ExecuteNextTick([this, LocalUserNum]()
diff --git a/Engine/Plugins/Online/WebAuth/Source/Private/IOS/IOSPlatformWebAuth.cpp b/Engine/Plugins/Online/WebAuth/Source/Private/IOS/IOSPlatformWebAuth.cpp
--- a/Engine/Plugins/Online/WebAuth/Source/Private/IOS/IOSPlatformWebAuth.cpp
+++ b/Engine/Plugins/Online/WebAuth/Source/Private/IOS/IOSPlatformWebAuth.cpp
@@ -72,12 +72,8 @@ bool FIOSWebAuth::AuthSessionWithURL(const FString &UrlStr, const FString &Schem
     }];
 
 #ifdef WEBAUTH_PLATFORM_IOS_13
-    // If we are using iOS 13+ ensure we supply the context provider
-    if (@available(iOS 13.0, *))
-    {
-        check(PresentationContextProvider);
-        ((ASWebAuthenticationSession*)SavedAuthSession).presentationContextProvider = PresentationContextProvider;
-    }
+    check(PresentationContextProvider);
+    ((ASWebAuthenticationSession*)SavedAuthSession).presentationContextProvider = PresentationContextProvider;
 #endif
 
     [(ASWebAuthenticationSession*)SavedAuthSession start];
diff --git a/Engine/Source/Developer/IOS/IOSPlatformEditor/Private/IOSTargetSettingsCustomization.cpp b/Engine/Source/Developer/IOS/IOSPlatformEditor/Private/IOSTargetSettingsCustomization.cpp
--- a/Engine/Source/Developer/IOS/IOSPlatformEditor/Private/IOSTargetSettingsCustomization.cpp
+++ b/Engine/Source/Developer/IOS/IOSPlatformEditor/Private/IOSTargetSettingsCustomization.cpp
@@ -1754,7 +1754,7 @@ void FIOSTargetSettingsCustomization::SetShaderStandard(int32 Value)
 	{
 		FText Message;
 		
-		uint8 EnumValue = (uint8)EIOSVersion::IOS_12;
+		uint8 EnumValue = (uint8)EIOSVersion::IOS_13;
 		if (MinOSPropertyHandle.IsValid())
 		{
 			MinOSPropertyHandle->GetValue(EnumValue);
@@ -1767,9 +1767,9 @@ void FIOSTargetSettingsCustomization::SetShaderStandard(int32 Value)
 		}
 		
 		// make sure we never set the min version to less than current supported
-		if (((EIOSVersion)EnumValue < EIOSVersion::IOS_12))
+		if (((EIOSVersion)EnumValue < EIOSVersion::IOS_13))
 		{
-			SetMinVersion((int32)EIOSVersion::IOS_12);
+			SetMinVersion((int32)EIOSVersion::IOS_13);
 		}
 
 		
@@ -1800,9 +1800,9 @@ void FIOSTargetSettingsCustomization::UpdateOSVersionWarning()
 		{
 			uint8 EnumValue;
 			MinOSPropertyHandle->GetValue(EnumValue);
-			if (EnumValue < (uint8)EIOSVersion::IOS_12)
+			if (EnumValue < (uint8)EIOSVersion::IOS_13)
 			{
-				SetMinVersion((int32)EIOSVersion::IOS_12);
+				SetMinVersion((int32)EIOSVersion::IOS_13);
 				
 				FText Message;
 				Message = LOCTEXT("MetalMRTStandardv1.2","Enabling the Desktop Forward Renderer Metal requires Shader Standard v2.0 which increases the minimum operating system requirement for Metal from iOS 10.0 or later to iOS 11.0 or later.");
@@ -1828,9 +1828,9 @@ void FIOSTargetSettingsCustomization::UpdateMetalMRTWarning()
 		{
 			uint8 EnumValue;
 			MinOSPropertyHandle->GetValue(EnumValue);
-			if (EnumValue < (uint8)EIOSVersion::IOS_12)
+			if (EnumValue < (uint8)EIOSVersion::IOS_13)
 			{
-				SetMinVersion((int32)EIOSVersion::IOS_12);
+				SetMinVersion((int32)EIOSVersion::IOS_13);
 				
 				FText Message;
 				Message = LOCTEXT("MetalMRTStandardv1.2","Enabling the Desktop Forward Renderer Metal requires Shader Standard v2.0 which increases the minimum operating system requirement for Metal from iOS 10.0 or later to iOS 11.0 or later.");
@@ -1838,9 +1838,9 @@ void FIOSTargetSettingsCustomization::UpdateMetalMRTWarning()
 			}
 			
 			ShaderVersionPropertyHandle->GetValue(EnumValue);
-			if (EnumValue < (uint8)EIOSMetalShaderStandard::IOSMetalSLStandard_2_0)
+			if (EnumValue < (uint8)EIOSMetalShaderStandard::IOSMetalSLStandard_2_2)
 			{
-				SetShaderStandard((int32)EIOSMetalShaderStandard::IOSMetalSLStandard_2_0);
+				SetShaderStandard((int32)EIOSMetalShaderStandard::IOSMetalSLStandard_2_2);
 				
 				FText Message;
 				Message = LOCTEXT("MetalMRTStandardv1.2","Enabling the Desktop Forward Renderer Metal requires Shader Standard v2.0 which increases the minimum operating system requirement for Metal from iOS 10.0 or later to iOS 11.0 or later.");
diff --git a/Engine/Source/Programs/UnrealBuildTool/Platform/IOS/IOSToolChain.cs b/Engine/Source/Programs/UnrealBuildTool/Platform/IOS/IOSToolChain.cs
--- a/Engine/Source/Programs/UnrealBuildTool/Platform/IOS/IOSToolChain.cs
+++ b/Engine/Source/Programs/UnrealBuildTool/Platform/IOS/IOSToolChain.cs
@@ -33,7 +33,7 @@ namespace UnrealBuildTool
 		/// The version of the iOS to allow at build time.
 		/// </summary>
 		[XmlConfigFile(Category = "IOSToolChain")]
-		public string BuildIOSVersion = "7.0";
+		public string BuildIOSVersion = "13.0";
 
 		/// <summary>
 		/// Directory for the developer binaries
@@ -155,7 +155,7 @@ namespace UnrealBuildTool
 					FileReference AssetFile = FileReference.Combine(Binary.OutputFilePath.Directory, "AssetCatalog", "Assets.car");
 					BuildProducts.Add(AssetFile, BuildProductType.RequiredResource);
 				}
-				else if (Target.Platform == UnrealTargetPlatform.IOS && Settings.Value.IOSSDKVersionFloat >= 11.0f)
+				else if (Target.Platform == UnrealTargetPlatform.IOS && Settings.Value.IOSSDKVersionFloat >= 13.0f)
 				{
 					int Index = Binary.OutputFilePath.GetFileNameWithoutExtension().IndexOf("-");
 					string OutputFile = Binary.OutputFilePath.GetFileNameWithoutExtension().Substring(0, Index > 0 ? Index : Binary.OutputFilePath.GetFileNameWithoutExtension().Length);
diff --git a/Engine/Source/Programs/UnrealBuildTool/Platform/IOS/UEBuildIOS.cs b/Engine/Source/Programs/UnrealBuildTool/Platform/IOS/UEBuildIOS.cs
--- a/Engine/Source/Programs/UnrealBuildTool/Platform/IOS/UEBuildIOS.cs
+++ b/Engine/Source/Programs/UnrealBuildTool/Platform/IOS/UEBuildIOS.cs
@@ -346,14 +346,12 @@ namespace UnrealBuildTool
 			{
 				switch (MinimumIOSVersion)
 				{
-					case "IOS_12":
-						return "12.0";
 					case "IOS_13":
 						return "13.0";
 					case "IOS_14":
 						return "14.0";
 					default:
-						return "12.0";
+						return "13.0";
 				}
 			}
 		}
diff --git a/Engine/Source/Programs/UnrealBuildTool/Platform/IOS/UEDeployIOS.cs b/Engine/Source/Programs/UnrealBuildTool/Platform/IOS/UEDeployIOS.cs
--- a/Engine/Source/Programs/UnrealBuildTool/Platform/IOS/UEDeployIOS.cs
+++ b/Engine/Source/Programs/UnrealBuildTool/Platform/IOS/UEDeployIOS.cs
@@ -197,6 +197,25 @@ namespace UnrealBuildTool
 			return result;
 		}
 
+		public static string GetMinimumOSVersion(string MinVersion)
+		{
+			string MinVersionToReturn = "";
+			switch (MinVersion)
+			{
+				case "IOS_13":
+					MinVersionToReturn = "13.0";
+					break;
+				case "IOS_14":
+					MinVersionToReturn = "14.0";
+					break;
+				default:
+					Log.TraceWarning("MinimumiOSVersion {0} specified in ini file is no longer supported, defaulting to 13.0", MinVersion);
+					MinVersionToReturn = "13.0";
+					break;
+			}
+			return MinVersionToReturn;
+		}
+
 		public static bool GenerateIOSPList(FileReference ProjectFile, UnrealTargetConfiguration Config, string ProjectDirectory, bool bIsUnrealGame, string GameName, bool bIsClient, string ProjectName, string InEngineDir, string AppDirectory, VersionNumber SdkVersion, UnrealPluginLanguage UPL, string BundleID, bool bBuildAsFramework, out bool bSupportsPortrait, out bool bSupportsLandscape, out bool bSkipIcons)
 		{
 			// generate the Info.plist for future use
@@ -310,30 +329,9 @@ namespace UnrealBuildTool
 			RequiredCaps += bSupported ? "\t\t<string>metal</string>\n" : "";
 
 			// minimum iOS version
-			string MinVersion;
-			if (Ini.GetString("/Script/IOSRuntimeSettings.IOSRuntimeSettings", "MinimumiOSVersion", out MinVersion))
-			{
-				switch (MinVersion)
-				{
-				case "IOS_12":
-					MinVersion = "12.0";
-					break;
-				case "IOS_13":
-					MinVersion = "13.0";
-					break;
-				case "IOS_14":
-					MinVersion = "14.0";
-					break;
-				default:
-					Log.TraceWarning("MinimumiOSVersion {0} specified in ini file is no longer supported, defaulting to 12.0", MinVersion);
-					MinVersion = "12.0";
-					break;
-				}
-			}
-			else
-			{
-				MinVersion = "12.0";
-			}
+			string MinVersionSetting = "";
+			Ini.GetString("/Script/IOSRuntimeSettings.IOSRuntimeSettings", "MinimumiOSVersion", out MinVersionSetting);
+			string MinVersion  = GetMinimumOSVersion(MinVersionSetting);
 
 			// Get Facebook Support details
 			bool bEnableFacebookSupport = true;
diff --git a/Engine/Source/Programs/UnrealBuildTool/Platform/TVOS/UEBuildTVOS.cs b/Engine/Source/Programs/UnrealBuildTool/Platform/TVOS/UEBuildTVOS.cs
--- a/Engine/Source/Programs/UnrealBuildTool/Platform/TVOS/UEBuildTVOS.cs
+++ b/Engine/Source/Programs/UnrealBuildTool/Platform/TVOS/UEBuildTVOS.cs
@@ -10,14 +10,6 @@ namespace UnrealBuildTool
 {
 	class TVOSProjectSettings : IOSProjectSettings
 	{
-		/// <summary>
-		/// Which version of the iOS to allow at run time
-		/// </summary>
-		public override string RuntimeVersion
-		{
-			get { return "10.0"; }
-		}
-
 		/// <summary>
 		/// which devices the game is allowed to run on
 		/// </summary>
diff --git a/Engine/Source/Programs/UnrealBuildTool/Platform/TVOS/UEDeployTVOS.cs b/Engine/Source/Programs/UnrealBuildTool/Platform/TVOS/UEDeployTVOS.cs
--- a/Engine/Source/Programs/UnrealBuildTool/Platform/TVOS/UEDeployTVOS.cs
+++ b/Engine/Source/Programs/UnrealBuildTool/Platform/TVOS/UEDeployTVOS.cs
@@ -73,20 +73,9 @@ namespace UnrealBuildTool
 			string RequiredCaps = "\t\t<string>arm64</string>\n";
 
 			// minimum iOS version
-			string MinVersion;
-			if (Ini.GetString("/Script/IOSRuntimeSettings.IOSRuntimeSettings", "MinimumTVOSVersion", out MinVersion))
-			{
-				switch (MinVersion)
-				{
-					case "TVOS_9":
-						MinVersion = "9.0";
-						break;
-				}
-			}
-			else
-			{
-				MinVersion = "9.0";
-			}
+			string MinVersionSetting = "";
+			Ini.GetString("/Script/IOSRuntimeSettings.IOSRuntimeSettings", "MinimumiOSVersion", out MinVersionSetting);
+			string MinVersion = GetMinimumOSVersion(MinVersionSetting);
 
 			// extra plist data
 			string ExtraData = "";
diff --git a/Engine/Source/Runtime/Apple/MetalRHI/Private/MetalBuffer.cpp b/Engine/Source/Runtime/Apple/MetalRHI/Private/MetalBuffer.cpp
--- a/Engine/Source/Runtime/Apple/MetalRHI/Private/MetalBuffer.cpp
+++ b/Engine/Source/Runtime/Apple/MetalRHI/Private/MetalBuffer.cpp
@@ -1257,10 +1257,7 @@ FMetalTexture FMetalTexturePool::CreateTexture(mtlpp::Device Device, mtlpp::Text
 		Pool.Remove(Descriptor);
 		if (GMetalResourcePurgeInPool)
 		{
-			if (@available(iOS 12.0, macOS 10.13, *))
-			{
-        		Texture.SetPurgeableState(mtlpp::PurgeableState::NonVolatile);
-			}
+            Texture.SetPurgeableState(mtlpp::PurgeableState::NonVolatile);
         }
 	}
 	else
@@ -1289,13 +1286,10 @@ void FMetalTexturePool::ReleaseTexture(FMetalTexture& Texture)
 	Descriptor.usage = Texture.GetUsage();
 	Descriptor.freedFrame = GFrameNumberRenderThread;
 	
-	if (@available(iOS 12.0, macOS 10.13, *))
-	{
-		if (GMetalResourcePurgeInPool && Texture.SetPurgeableState(mtlpp::PurgeableState::KeepCurrent) == mtlpp::PurgeableState::NonVolatile)
-		{
-			Texture.SetPurgeableState(mtlpp::PurgeableState::Volatile);
-		}
-	}
+    if (GMetalResourcePurgeInPool && Texture.SetPurgeableState(mtlpp::PurgeableState::KeepCurrent) == mtlpp::PurgeableState::NonVolatile)
+    {
+        Texture.SetPurgeableState(mtlpp::PurgeableState::Volatile);
+    }
 	
 	FScopeLock Lock(&PoolMutex);
 	Pool.Add(Descriptor, Texture);
@@ -1320,10 +1314,7 @@ void FMetalTexturePool::Drain(bool const bForce)
             {
 				if (GMetalResourcePurgeInPool && (GFrameNumberRenderThread - It->Key.freedFrame) >= PurgeAfterNumFrames)
 				{
-					if (@available(iOS 12.0, macOS 10.13, *))
-					{
-						It->Value.SetPurgeableState(mtlpp::PurgeableState::Empty);
-					}
+					It->Value.SetPurgeableState(mtlpp::PurgeableState::Empty);
 				}
             }
         }
@@ -1522,13 +1513,10 @@ FMetalBuffer FMetalResourceHeap::CreateBuffer(uint32 Size, uint32 Alignment, uin
 				 else
 				 {
                     Buffer = ManagedBuffers.CreatePooledResource(FMetalPooledBufferArgs(Queue->GetDevice(), BlockSize, Flags, StorageMode));
-					if (@available(iOS 12.0, macOS 10.13, *))
-					{
-						if (GMetalResourcePurgeInPool)
-						{
-							Buffer.SetPurgeableState(mtlpp::PurgeableState::NonVolatile);
-						}
-					}
+                     if (GMetalResourcePurgeInPool)
+                     {
+                         Buffer.SetPurgeableState(mtlpp::PurgeableState::NonVolatile);
+                     }
 					DEC_MEMORY_STAT_BY(STAT_MetalBufferUnusedMemory, Buffer.GetLength());
 					DEC_MEMORY_STAT_BY(STAT_MetalPooledBufferUnusedMemory, Buffer.GetLength());
 					INC_MEMORY_STAT_BY(STAT_MetalPooledBufferMemory, Buffer.GetLength());
@@ -1604,10 +1592,7 @@ FMetalBuffer FMetalResourceHeap::CreateBuffer(uint32 Size, uint32 Alignment, uin
                     Buffer = Buffers[Storage].CreatePooledResource(FMetalPooledBufferArgs(Queue->GetDevice(), BlockSize, Flags, StorageMode));
 					if (GMetalResourcePurgeInPool)
 					{
-						if (@available(iOS 12.0, macOS 10.13, *))
-						{
-                    		Buffer.SetPurgeableState(mtlpp::PurgeableState::NonVolatile);
-						}
+                   		Buffer.SetPurgeableState(mtlpp::PurgeableState::NonVolatile);
 					}
 					DEC_MEMORY_STAT_BY(STAT_MetalBufferUnusedMemory, Buffer.GetLength());
 					DEC_MEMORY_STAT_BY(STAT_MetalPooledBufferUnusedMemory, Buffer.GetLength());
@@ -1655,10 +1640,7 @@ void FMetalResourceHeap::ReleaseBuffer(FMetalBuffer& Buffer)
 		
 		if (GMetalResourcePurgeInPool)
 		{
-			if (@available(iOS 12.0, macOS 10.13, *))
-			{
-        		Buffer.SetPurgeableState(mtlpp::PurgeableState::Volatile);
-			}
+       		Buffer.SetPurgeableState(mtlpp::PurgeableState::Volatile);
 		}
         
 		switch (StorageMode)
diff --git a/Engine/Source/Runtime/Apple/MetalRHI/Private/MetalRHI.cpp b/Engine/Source/Runtime/Apple/MetalRHI/Private/MetalRHI.cpp
--- a/Engine/Source/Runtime/Apple/MetalRHI/Private/MetalRHI.cpp
+++ b/Engine/Source/Runtime/Apple/MetalRHI/Private/MetalRHI.cpp
@@ -602,14 +602,11 @@ FMetalDynamicRHI::FMetalDynamicRHI(ERHIFeatureLevel::Type RequestedFeatureLevel)
 	GRHISupportsFirstInstance = GRHISupportsBaseVertexIndex;
 	
 	// TODO: Move this into IOSPlatform
-	if (@available(iOS 11.0, *))
-	{
-		@autoreleasepool {
-			UIScreen* mainScreen = [UIScreen mainScreen];
-			UIDisplayGamut gamut = mainScreen.traitCollection.displayGamut;
-			GRHISupportsHDROutput = FPlatformMisc::IOSVersionCompare(10, 0, 0) && gamut == UIDisplayGamutP3;
-		}
-	}
+    @autoreleasepool {
+        UIScreen* mainScreen = [UIScreen mainScreen];
+        UIDisplayGamut gamut = mainScreen.traitCollection.displayGamut;
+        GRHISupportsHDROutput = FPlatformMisc::IOSVersionCompare(10, 0, 0) && gamut == UIDisplayGamutP3;
+    }
 	
 	GRHIHDRDisplayOutputFormat = (GRHISupportsHDROutput) ? PF_PLATFORM_HDR_0 : PF_B8G8R8A8;
 	// Based on the spec below, the maxTotalThreadsPerThreadgroup is not a fixed number but calculated according to the device current ability, so the available threads could less than the maximum number.
diff --git a/Engine/Source/Runtime/ApplicationCore/Private/IOS/IOSAppDelegate.cpp b/Engine/Source/Runtime/ApplicationCore/Private/IOS/IOSAppDelegate.cpp
--- a/Engine/Source/Runtime/ApplicationCore/Private/IOS/IOSAppDelegate.cpp
+++ b/Engine/Source/Runtime/ApplicationCore/Private/IOS/IOSAppDelegate.cpp
@@ -934,7 +934,7 @@ static FAutoConsoleVariableRef CVarGEnableThermalsReport(
 	self.bDeviceInPortraitMode = false;
 #else
 	// use the status bar orientation to properly determine landscape vs portrait
-	self.bDeviceInPortraitMode = UIInterfaceOrientationIsPortrait([[UIApplication sharedApplication] statusBarOrientation]);
+	self.bDeviceInPortraitMode = UIInterfaceOrientationIsPortrait([self.Window.windowScene interfaceOrientation]);
 	printf("========= This app is in %s mode\n", self.bDeviceInPortraitMode ? "PORTRAIT" : "LANDSCAPE");
 #endif
 
@@ -953,7 +953,6 @@ static FAutoConsoleVariableRef CVarGEnableThermalsReport(
     
     CGRect MainFrame = [[UIScreen mainScreen] bounds];
     self.Window = [[UIWindow alloc] initWithFrame:MainFrame];
-    self.Window.screen = [UIScreen mainScreen];
     
     // get the native scale
     const float NativeScale = [[UIScreen mainScreen] scale];
@@ -970,7 +969,7 @@ static FAutoConsoleVariableRef CVarGEnableThermalsReport(
         [self.Window addSubview: viewController.view];
         GShowSplashScreen = true;
     }
-    
+
 
     timer = [NSTimer scheduledTimerWithTimeInterval: 0.05f target:self selector:@selector(timerForSplashScreen) userInfo:nil repeats:YES];
 
@@ -985,7 +984,7 @@ static FAutoConsoleVariableRef CVarGEnableThermalsReport(
 	Center.delegate = self;
 	// Register for device orientation changes
 	[[UIDevice currentDevice] beginGeneratingDeviceOrientationNotifications];
-	[[NSNotificationCenter defaultCenter] addObserver:self selector:@selector(didRotate:) name:UIApplicationDidChangeStatusBarOrientationNotification object:nil];
+    [[NSNotificationCenter defaultCenter] addObserver:self selector:@selector(didRotate:) name:UIDeviceOrientationDidChangeNotification object:nil];
 
 #if !UE_BUILD_SHIPPING
 	// make a history buffer
@@ -999,59 +998,52 @@ static FAutoConsoleVariableRef CVarGEnableThermalsReport(
 	}
 	self.ConsoleHistoryValuesIndex = -1;
 
-	if (@available(iOS 11, *))
-	{
-		FCoreDelegates::OnGetOnScreenMessages.AddLambda(
-			[&EnableThermalsReport = GEnableThermalsReport](TMultiMap<FCoreDelegates::EOnScreenMessageSeverity, FText >& OutMessages)
-			{
-				if (EnableThermalsReport)
-				{
-					switch ([[NSProcessInfo processInfo] thermalState])
-					{
-						case NSProcessInfoThermalStateNominal:	OutMessages.Add(FCoreDelegates::EOnScreenMessageSeverity::Info, FText::FromString(TEXT("Thermals are Nominal"))); break;
-						case NSProcessInfoThermalStateFair:		OutMessages.Add(FCoreDelegates::EOnScreenMessageSeverity::Info, FText::FromString(TEXT("Thermals are Fair"))); break;
-						case NSProcessInfoThermalStateSerious:	OutMessages.Add(FCoreDelegates::EOnScreenMessageSeverity::Warning, FText::FromString(TEXT("Thermals are Serious"))); break;
-						case NSProcessInfoThermalStateCritical:	OutMessages.Add(FCoreDelegates::EOnScreenMessageSeverity::Error, FText::FromString(TEXT("Thermals are Critical"))); break;
-					}
-				}
-
-				// Uncomment to view the state of the AVAudioSession category, mode, and options.
-//#define VIEW_AVAUDIOSESSION_INFO
+    FCoreDelegates::OnGetOnScreenMessages.AddLambda([&EnableThermalsReport = GEnableThermalsReport](TMultiMap<FCoreDelegates::EOnScreenMessageSeverity, FText >& OutMessages)
+    {
+        if (EnableThermalsReport)
+        {
+            switch ([[NSProcessInfo processInfo] thermalState])
+            {
+                case NSProcessInfoThermalStateNominal:	OutMessages.Add(FCoreDelegates::EOnScreenMessageSeverity::Info, FText::FromString(TEXT("Thermals are Nominal"))); break;
+                case NSProcessInfoThermalStateFair:		OutMessages.Add(FCoreDelegates::EOnScreenMessageSeverity::Info, FText::FromString(TEXT("Thermals are Fair"))); break;
+                case NSProcessInfoThermalStateSerious:	OutMessages.Add(FCoreDelegates::EOnScreenMessageSeverity::Warning, FText::FromString(TEXT("Thermals are Serious"))); break;
+                case NSProcessInfoThermalStateCritical:	OutMessages.Add(FCoreDelegates::EOnScreenMessageSeverity::Error, FText::FromString(TEXT("Thermals are Critical"))); break;
+            }
+        }
+        
+        // Uncomment to view the state of the AVAudioSession category, mode, and options.
+        //#define VIEW_AVAUDIOSESSION_INFO
 #if defined(VIEW_AVAUDIOSESSION_INFO)
-				FString Message = FString::Printf(
-					TEXT("Session Category: %s, Mode: %s, Options: %x"),
-					UTF8_TO_TCHAR([[AVAudioSession sharedInstance].category UTF8String]),
-					UTF8_TO_TCHAR([[AVAudioSession sharedInstance].mode UTF8String]),
-					[AVAudioSession sharedInstance].categoryOptions);
-				OutMessages.Add(FCoreDelegates::EOnScreenMessageSeverity::Info, FText::FromString(Message));
+        FString Message = FString::Printf(
+                                          TEXT("Session Category: %s, Mode: %s, Options: %x"),
+                                          UTF8_TO_TCHAR([[AVAudioSession sharedInstance].category UTF8String]),
+                                          UTF8_TO_TCHAR([[AVAudioSession sharedInstance].mode UTF8String]),
+                                          [AVAudioSession sharedInstance].categoryOptions);
+        OutMessages.Add(FCoreDelegates::EOnScreenMessageSeverity::Info, FText::FromString(Message));
 #endif // defined(VIEW_AVAUDIOSESSION_INFO)
-			});
-	}
-
+    });
+    
 
 #endif // UE_BUILD_SHIPPING
 #endif // !TVOS
 
 #if !PLATFORM_TVOS
-	if (@available(iOS 11, *))
-	{
-        UIDevice* UiDevice = [UIDevice currentDevice];
-        UiDevice.batteryMonitoringEnabled = YES;
-        
-        // Battery level is from 0.0 to 1.0, get it in terms of 0-100
-        self.BatteryLevel = ((int)([UiDevice batteryLevel] * 100));
-        UIDeviceBatteryState State = UiDevice.batteryState;
-        self.bBatteryState = State == UIDeviceBatteryStateUnplugged || State == UIDeviceBatteryStateUnknown;
-        self.ThermalState = [[NSProcessInfo processInfo] thermalState];
-
-		[[NSNotificationCenter defaultCenter] addObserver:self selector:@selector(temperatureChanged:) name:NSProcessInfoThermalStateDidChangeNotification object:nil];
-		[[NSNotificationCenter defaultCenter] addObserver:self selector:@selector(lowPowerModeChanged:) name:NSProcessInfoPowerStateDidChangeNotification object:nil];
-        [[NSNotificationCenter defaultCenter] addObserver:self selector:@selector(batteryChanged:) name:UIDeviceBatteryLevelDidChangeNotification object:nil];
-        [[NSNotificationCenter defaultCenter] addObserver:self selector:@selector(batteryStateChanged:) name:UIDeviceBatteryStateDidChangeNotification object:nil];
-	}
+    UIDevice* UiDevice = [UIDevice currentDevice];
+    UiDevice.batteryMonitoringEnabled = YES;
+    
+    // Battery level is from 0.0 to 1.0, get it in terms of 0-100
+    self.BatteryLevel = ((int)([UiDevice batteryLevel] * 100));
+    UIDeviceBatteryState State = UiDevice.batteryState;
+    self.bBatteryState = State == UIDeviceBatteryStateUnplugged || State == UIDeviceBatteryStateUnknown;
+    self.ThermalState = [[NSProcessInfo processInfo] thermalState];
+    
+    [[NSNotificationCenter defaultCenter] addObserver:self selector:@selector(temperatureChanged:) name:NSProcessInfoThermalStateDidChangeNotification object:nil];
+    [[NSNotificationCenter defaultCenter] addObserver:self selector:@selector(lowPowerModeChanged:) name:NSProcessInfoPowerStateDidChangeNotification object:nil];
+    [[NSNotificationCenter defaultCenter] addObserver:self selector:@selector(batteryChanged:) name:UIDeviceBatteryLevelDidChangeNotification object:nil];
+    [[NSNotificationCenter defaultCenter] addObserver:self selector:@selector(batteryStateChanged:) name:UIDeviceBatteryStateDidChangeNotification object:nil];
 #endif
-	
-	self.bAudioSessionInitialized = false;
+    
+    self.bAudioSessionInitialized = false;
 	
 	// InitializeAudioSession is now called from FEngineLoop::AppInit after the config system is initialized
 //	[self InitializeAudioSession];
@@ -1210,12 +1202,10 @@ extern EDeviceScreenOrientation ConvertFromUIInterfaceOrientation(UIInterfaceOri
 {   
 #if !PLATFORM_TVOS
 	// get the interfaec orientation
-	NSNumber* OrientationNumber = [notification.userInfo objectForKey:UIApplicationStatusBarOrientationUserInfoKey];
-	UIInterfaceOrientation Orientation = (UIInterfaceOrientation)[OrientationNumber intValue];
 	
-	NSLog(@"didRotate orientation = %d, statusBar = %d", (int)Orientation, (int)[[UIApplication sharedApplication] statusBarOrientation]);
+	NSLog(@"didRotate orientation = %d", (int)[self.Window.windowScene interfaceOrientation]);
 	
-	Orientation = [[UIApplication sharedApplication] statusBarOrientation];
+    UIInterfaceOrientation Orientation = [self.Window.windowScene interfaceOrientation];
 	
 	extern UIInterfaceOrientation GInterfaceOrientation;
 	GInterfaceOrientation = Orientation;
@@ -1722,10 +1712,7 @@ didReceiveNotificationResponse:(UNNotificationResponse *)response
 #if !PLATFORM_TVOS
 	// create the leaderboard display object 
 	GKGameCenterViewController* GameCenterDisplay = [[[GKGameCenterViewController alloc] init] autorelease];
-    if (@available(iOS 7, tvOS 999, *))
-    {
 	GameCenterDisplay.viewState = GKGameCenterViewControllerStateAchievements;
-    }
 	GameCenterDisplay.gameCenterDelegate = self;
 
 	// show it 
@@ -1824,43 +1811,37 @@ CORE_API bool IOSShowDashboardUI()
 -(void)temperatureChanged:(NSNotification *)notification
 {
 #if !PLATFORM_TVOS
-	if (@available(iOS 11, *))
-	{
-		// send game callback with new temperature severity
-		FCoreDelegates::ETemperatureSeverity Severity;
-        FString Level = TEXT("Unknown");
-        self.ThermalState = [[NSProcessInfo processInfo] thermalState];
-		switch (self.ThermalState)
-		{
-			case NSProcessInfoThermalStateNominal:	Severity = FCoreDelegates::ETemperatureSeverity::Good; Level = TEXT("Good"); break;
-			case NSProcessInfoThermalStateFair:		Severity = FCoreDelegates::ETemperatureSeverity::Bad; Level = TEXT("Bad"); break;
-			case NSProcessInfoThermalStateSerious:	Severity = FCoreDelegates::ETemperatureSeverity::Serious; Level = TEXT("Serious"); break;
-			case NSProcessInfoThermalStateCritical:	Severity = FCoreDelegates::ETemperatureSeverity::Critical; Level = TEXT("Critical"); break;
-		}
-
-		[FIOSAsyncTask CreateTaskWithBlock : ^ bool(void)
-		{
-			UE_LOG(LogIOS, Display, TEXT("Temperature Changed: %s"), *Level);
-			FCoreDelegates::OnTemperatureChange.Broadcast(Severity);
-			return true;
-		}];
-	}
+    // send game callback with new temperature severity
+    FCoreDelegates::ETemperatureSeverity Severity;
+    FString Level = TEXT("Unknown");
+    self.ThermalState = [[NSProcessInfo processInfo] thermalState];
+    switch (self.ThermalState)
+    {
+        case NSProcessInfoThermalStateNominal:	Severity = FCoreDelegates::ETemperatureSeverity::Good; Level = TEXT("Good"); break;
+        case NSProcessInfoThermalStateFair:		Severity = FCoreDelegates::ETemperatureSeverity::Bad; Level = TEXT("Bad"); break;
+        case NSProcessInfoThermalStateSerious:	Severity = FCoreDelegates::ETemperatureSeverity::Serious; Level = TEXT("Serious"); break;
+        case NSProcessInfoThermalStateCritical:	Severity = FCoreDelegates::ETemperatureSeverity::Critical; Level = TEXT("Critical"); break;
+    }
+    
+    [FIOSAsyncTask CreateTaskWithBlock : ^ bool(void)
+     {
+        UE_LOG(LogIOS, Display, TEXT("Temperature Changed: %s"), *Level);
+        FCoreDelegates::OnTemperatureChange.Broadcast(Severity);
+        return true;
+    }];
 #endif
 }
 
 -(void)lowPowerModeChanged:(NSNotification *)notification
 {
 #if !PLATFORM_TVOS
-	if (@available(iOS 11, *))
-	{	
-		[FIOSAsyncTask CreateTaskWithBlock : ^ bool(void)
-        {
-            bool bInLowPowerMode = [[NSProcessInfo processInfo] isLowPowerModeEnabled];
-            UE_LOG(LogIOS, Display, TEXT("Low Power Mode Changed: %d"), bInLowPowerMode);
-            FCoreDelegates::OnLowPowerMode.Broadcast(bInLowPowerMode);
-            return true;
-        }];
-	}
+    [FIOSAsyncTask CreateTaskWithBlock : ^ bool(void)
+     {
+        bool bInLowPowerMode = [[NSProcessInfo processInfo] isLowPowerModeEnabled];
+        UE_LOG(LogIOS, Display, TEXT("Low Power Mode Changed: %d"), bInLowPowerMode);
+        FCoreDelegates::OnLowPowerMode.Broadcast(bInLowPowerMode);
+        return true;
+    }];
 #endif
 }
 
diff --git a/Engine/Source/Runtime/ApplicationCore/Private/IOS/IOSApplication.cpp b/Engine/Source/Runtime/ApplicationCore/Private/IOS/IOSApplication.cpp
--- a/Engine/Source/Runtime/ApplicationCore/Private/IOS/IOSApplication.cpp
+++ b/Engine/Source/Runtime/ApplicationCore/Private/IOS/IOSApplication.cpp
@@ -141,73 +141,63 @@ void FDisplayMetrics::RebuildDisplayMetrics(FDisplayMetrics& OutDisplayMetrics)
 	OutDisplayMetrics.PrimaryDisplayWorkAreaRect = FIOSWindow::GetScreenRect();
 	OutDisplayMetrics.VirtualDisplayRect = OutDisplayMetrics.PrimaryDisplayWorkAreaRect;
 
-	// Total screen size of the primary monitor
-	OutDisplayMetrics.PrimaryDisplayWidth = OutDisplayMetrics.PrimaryDisplayWorkAreaRect.Right - OutDisplayMetrics.PrimaryDisplayWorkAreaRect.Left;
-	OutDisplayMetrics.PrimaryDisplayHeight = OutDisplayMetrics.PrimaryDisplayWorkAreaRect.Bottom - OutDisplayMetrics.PrimaryDisplayWorkAreaRect.Top;
-
-	// Get ui window rect
-	OutDisplayMetrics.IosUiWindowAreaRect = FIOSWindow::GetUIWindowRect();
-
+    // Total screen size of the primary monitor
+    OutDisplayMetrics.PrimaryDisplayWidth = OutDisplayMetrics.PrimaryDisplayWorkAreaRect.Right - OutDisplayMetrics.PrimaryDisplayWorkAreaRect.Left;
+    OutDisplayMetrics.PrimaryDisplayHeight = OutDisplayMetrics.PrimaryDisplayWorkAreaRect.Bottom - OutDisplayMetrics.PrimaryDisplayWorkAreaRect.Top;
+    
+    // Get ui window rect
+    OutDisplayMetrics.IosUiWindowAreaRect = FIOSWindow::GetUIWindowRect();
+    
 #if !PLATFORM_TVOS
-	if (@available(iOS 11, *))
-	{
-		const float RequestedContentScaleFactor = [[IOSAppDelegate GetDelegate].IOSView contentScaleFactor];
-
-		//we need to set these according to the orientation
-		TAutoConsoleVariable<float>* CVar_Left = nullptr;
-		TAutoConsoleVariable<float>* CVar_Top = nullptr;
-		TAutoConsoleVariable<float>* CVar_Right = nullptr;
-		TAutoConsoleVariable<float>* CVar_Bottom = nullptr;
-
-		//making an assumption that the "normal" landscape mode is Landscape right
-		if (FIOSApplication::CachedOrientation == UIInterfaceOrientationLandscapeLeft)
-		{
-			CVar_Left = &CVarSafeZone_Landscape_Left;
-			CVar_Right = &CVarSafeZone_Landscape_Right;
-            CVar_Top = &CVarSafeZone_Landscape_Top;
-            CVar_Bottom = &CVarSafeZone_Landscape_Bottom;
-		}
-		else if (FIOSApplication::CachedOrientation == UIInterfaceOrientationLandscapeRight)
-		{
-			CVar_Left = &CVarSafeZone_Landscape_Right;
-			CVar_Right = &CVarSafeZone_Landscape_Left;
-            CVar_Top = &CVarSafeZone_Landscape_Top;
-            CVar_Bottom = &CVarSafeZone_Landscape_Bottom;
-		}
-        
-		// of the CVars are set, use their values. If not, use what comes from iOS
-		const float Inset_Left = (!CVar_Left || CVar_Left->AsVariable()->GetFloat() < 0.0f) ? CachedInsets.left : CVar_Left->AsVariable()->GetFloat();
-		const float Inset_Top = (!CVar_Top || CVar_Top->AsVariable()->GetFloat() < 0.0f) ? CachedInsets.top : CVar_Top->AsVariable()->GetFloat();
-		const float Inset_Right = (!CVar_Right || CVar_Right->AsVariable()->GetFloat() < 0.0f) ? CachedInsets.right : CVar_Right->AsVariable()->GetFloat();
-		const float Inset_Bottom = (!CVar_Bottom || CVar_Bottom->AsVariable()->GetFloat() < 0.0f) ? CachedInsets.bottom : CVar_Bottom->AsVariable()->GetFloat();
-
-		//setup the asymmetrical padding
-		OutDisplayMetrics.TitleSafePaddingSize.X = Inset_Left;
-		OutDisplayMetrics.TitleSafePaddingSize.Y = Inset_Top;
-		OutDisplayMetrics.TitleSafePaddingSize.Z = Inset_Right;
-		OutDisplayMetrics.TitleSafePaddingSize.W = Inset_Bottom;
-
-		//scale the thing
-		OutDisplayMetrics.TitleSafePaddingSize *= RequestedContentScaleFactor;
-
-		OutDisplayMetrics.ActionSafePaddingSize = OutDisplayMetrics.TitleSafePaddingSize;
-	}
-	else
+    const float RequestedContentScaleFactor = [[IOSAppDelegate GetDelegate].IOSView contentScaleFactor];
+    
+    //we need to set these according to the orientation
+    TAutoConsoleVariable<float>* CVar_Left = nullptr;
+    TAutoConsoleVariable<float>* CVar_Top = nullptr;
+    TAutoConsoleVariable<float>* CVar_Right = nullptr;
+    TAutoConsoleVariable<float>* CVar_Bottom = nullptr;
+    
+    //making an assumption that the "normal" landscape mode is Landscape right
+    if (FIOSApplication::CachedOrientation == UIInterfaceOrientationLandscapeLeft)
+    {
+        CVar_Left = &CVarSafeZone_Landscape_Left;
+        CVar_Right = &CVarSafeZone_Landscape_Right;
+        CVar_Top = &CVarSafeZone_Landscape_Top;
+        CVar_Bottom = &CVarSafeZone_Landscape_Bottom;
+    }
+    else if (FIOSApplication::CachedOrientation == UIInterfaceOrientationLandscapeRight)
+    {
+        CVar_Left = &CVarSafeZone_Landscape_Right;
+        CVar_Right = &CVarSafeZone_Landscape_Left;
+        CVar_Top = &CVarSafeZone_Landscape_Top;
+        CVar_Bottom = &CVarSafeZone_Landscape_Bottom;
+    }
+    
+    // of the CVars are set, use their values. If not, use what comes from iOS
+    const float Inset_Left = (!CVar_Left || CVar_Left->AsVariable()->GetFloat() < 0.0f) ? CachedInsets.left : CVar_Left->AsVariable()->GetFloat();
+    const float Inset_Top = (!CVar_Top || CVar_Top->AsVariable()->GetFloat() < 0.0f) ? CachedInsets.top : CVar_Top->AsVariable()->GetFloat();
+    const float Inset_Right = (!CVar_Right || CVar_Right->AsVariable()->GetFloat() < 0.0f) ? CachedInsets.right : CVar_Right->AsVariable()->GetFloat();
+    const float Inset_Bottom = (!CVar_Bottom || CVar_Bottom->AsVariable()->GetFloat() < 0.0f) ? CachedInsets.bottom : CVar_Bottom->AsVariable()->GetFloat();
+    
+    //setup the asymmetrical padding
+    OutDisplayMetrics.TitleSafePaddingSize.X = Inset_Left;
+    OutDisplayMetrics.TitleSafePaddingSize.Y = Inset_Top;
+    OutDisplayMetrics.TitleSafePaddingSize.Z = Inset_Right;
+    OutDisplayMetrics.TitleSafePaddingSize.W = Inset_Bottom;
+    
+    //scale the thing
+    OutDisplayMetrics.TitleSafePaddingSize *= RequestedContentScaleFactor;
+    
+    OutDisplayMetrics.ActionSafePaddingSize = OutDisplayMetrics.TitleSafePaddingSize;
 #endif
-	{
-		OutDisplayMetrics.ApplyDefaultSafeZones();
-	}
+    
 }
 
 void FIOSApplication::CacheDisplayMetrics()
-	{
-
+{
 #if !PLATFORM_TVOS
-	if (@available(iOS 11, *))
-	{
-		CachedInsets = [[[[UIApplication sharedApplication] delegate] window] safeAreaInsets];
-		CachedOrientation = [[UIApplication sharedApplication] statusBarOrientation];
-	}
+    CachedInsets = [[[[UIApplication sharedApplication] delegate] window] safeAreaInsets];
+    CachedOrientation = [[[[[UIApplication sharedApplication] delegate] window] windowScene] interfaceOrientation];
 #endif
 }
 
diff --git a/Engine/Source/Runtime/ApplicationCore/Private/IOS/IOSPlatformFramePacer.cpp b/Engine/Source/Runtime/ApplicationCore/Private/IOS/IOSPlatformFramePacer.cpp
--- a/Engine/Source/Runtime/ApplicationCore/Private/IOS/IOSPlatformFramePacer.cpp
+++ b/Engine/Source/Runtime/ApplicationCore/Private/IOS/IOSPlatformFramePacer.cpp
@@ -76,10 +76,7 @@ namespace IOSDisplayConstants
 		FScopeLock Lock(&HandlersMutex);
 		CADisplayLink* displayLink = (CADisplayLink*)param;
 		double OutputSeconds = displayLink.duration + displayLink.timestamp;
-		if (@available(iOS 10.0, tvOS 10.0, *))
-		{
-			OutputSeconds = displayLink.targetTimestamp;
-		}
+		OutputSeconds = displayLink.targetTimestamp;
 		double OutputDuration = displayLink.duration;
 		for (FIOSFramePacerHandler Handler in Handlers)
 		{
diff --git a/Engine/Source/Runtime/ApplicationCore/Private/IOS/IOSView.cpp b/Engine/Source/Runtime/ApplicationCore/Private/IOS/IOSView.cpp
--- a/Engine/Source/Runtime/ApplicationCore/Private/IOS/IOSView.cpp
+++ b/Engine/Source/Runtime/ApplicationCore/Private/IOS/IOSView.cpp
@@ -332,7 +332,7 @@ id<MTLDevice> GMetalDevice = nil;
 - (void)layoutSubviews
 {
 #if !PLATFORM_TVOS
-	auto orientation = [[UIApplication sharedApplication] statusBarOrientation];
+    auto orientation = [self.window.windowScene interfaceOrientation];
 	FIOSApplication::OrientationChanged(orientation);
 #endif
 }
@@ -892,18 +892,15 @@ self.accessibilityElements = @[Window.accessibilityContainer];
 	return nil;
 }
 
-// UE 4.25 note: type UITextWritingDirection is deprecated and the new type NSWritingDirection used from Xcode 11 but we need to 
-// continue to use UITextWritingDirection for now in order to maintain compatibility with Xcode 9 and 10.
-
-- (UITextWritingDirection)baseWritingDirectionForPosition:(UITextPosition *)position inDirection:(UITextStorageDirection)direction
+- (NSWritingDirection)baseWritingDirectionForPosition:(UITextPosition *)position inDirection:(UITextStorageDirection)direction
 {
 	REPORT_EVENT;
 	// assume left to right for now
-	return UITextWritingDirectionLeftToRight;
+	return NSWritingDirectionLeftToRight;
 }
 
 
-- (void)setBaseWritingDirection:(UITextWritingDirection)writingDirection forRange:(UITextRange *)range
+- (void)setBaseWritingDirection:(NSWritingDirection)writingDirection forRange:(UITextRange *)range
 {
 	// @todo keyboard: This is called
 }
diff --git a/Engine/Source/Runtime/ApplicationCore/Private/IOS/IOSWindow.cpp b/Engine/Source/Runtime/ApplicationCore/Private/IOS/IOSWindow.cpp
--- a/Engine/Source/Runtime/ApplicationCore/Private/IOS/IOSWindow.cpp
+++ b/Engine/Source/Runtime/ApplicationCore/Private/IOS/IOSWindow.cpp
@@ -23,8 +23,17 @@ void FIOSWindow::Initialize( class FIOSApplication* const Application, const TSh
 {
 	OwningApplication = Application;
 	Definition = InDefinition;
-
-	Window = [[UIApplication sharedApplication] keyWindow];
+    
+    UIWindow        *foundKeyWindow = nil;
+    NSArray         *allWindows = [[UIApplication sharedApplication]windows];
+    for (UIWindow   *currentWindow in allWindows)
+    {
+        if (currentWindow.isKeyWindow)
+        {
+            Window = foundKeyWindow;
+            break;
+        }
+    }
 
 #if !PLATFORM_TVOS
 	if(InParent.Get() != NULL)
diff --git a/Engine/Source/Runtime/Core/Private/Apple/ApplePlatformMemory.cpp b/Engine/Source/Runtime/Core/Private/Apple/ApplePlatformMemory.cpp
--- a/Engine/Source/Runtime/Core/Private/Apple/ApplePlatformMemory.cpp
+++ b/Engine/Source/Runtime/Core/Private/Apple/ApplePlatformMemory.cpp
@@ -345,20 +345,13 @@ FPlatformMemoryStats FApplePlatformMemory::GetStats()
 	// Gather platform memory stats.
 	uint64_t FreeMem = 0;
 #if PLATFORM_IOS
-#if defined(__IPHONE_13_0)
-	if (@available(iOS 13.0,*))
-	{
-		FreeMem = os_proc_available_memory();
-	}
-	else
-#endif
+	FreeMem = os_proc_available_memory();
+#else
+    vm_statistics Stats;
+    mach_msg_type_number_t StatsSize = sizeof(Stats);
+    host_statistics(mach_host_self(), HOST_VM_INFO, (host_info_t)&Stats, &StatsSize);
+    FreeMem = (Stats.free_count + Stats.inactive_count) * MemoryConstants.PageSize;
 #endif
-	{
-		vm_statistics Stats;
-		mach_msg_type_number_t StatsSize = sizeof(Stats);
-		host_statistics(mach_host_self(), HOST_VM_INFO, (host_info_t)&Stats, &StatsSize);
-		FreeMem = (Stats.free_count + Stats.inactive_count) * MemoryConstants.PageSize;
-	}
 	MemoryStats.AvailablePhysical = FreeMem;
 	
 	// Just get memory information for the process and report the working set instead
@@ -366,17 +359,11 @@ FPlatformMemoryStats FApplePlatformMemory::GetStats()
 	mach_msg_type_number_t TaskInfoCount = MACH_TASK_BASIC_INFO_COUNT;
 	task_info( mach_task_self(), MACH_TASK_BASIC_INFO, (task_info_t)&TaskInfo, &TaskInfoCount );
 #if PLATFORM_IOS
-#if defined(__IPHONE_13_0)
-	if (@available(iOS 13.0,*))
-	{
-		MemoryStats.UsedPhysical = MemoryConstants.TotalPhysical - FreeMem;
-	}
-	else
-#endif
+    MemoryStats.UsedPhysical = MemoryConstants.TotalPhysical - FreeMem;
+#else
+    MemoryStats.UsedPhysical = TaskInfo.resident_size;
 #endif
-	{
-		MemoryStats.UsedPhysical = TaskInfo.resident_size;
-	}
+
 	if(MemoryStats.UsedPhysical > MemoryStats.PeakUsedPhysical)
 	{
 		MemoryStats.PeakUsedPhysical = MemoryStats.UsedPhysical;
@@ -405,43 +392,35 @@ const FPlatformMemoryConstants& FApplePlatformMemory::GetConstants()
 		// Get memory.
 		int64 AvailablePhysical = 0;
 #if PLATFORM_IOS
-#if defined(__IPHONE_13_0)
-		if (@available(iOS 13.0,*))
-		{
-			AvailablePhysical = os_proc_available_memory();
-			
-			// quantize to the known jetsam limits, we should be within 50MB of the correct one
-			uint64 JetsamLimits[] = { 1520435200, 1939865600, 2201170740, 2252710350, 3006477100 }; // { 2GB, gimped 3GB, gimped 4GB, 3GB, 4GB
-			if (AvailablePhysical < JetsamLimits[0])
-			{
-				AvailablePhysical = JetsamLimits[0];
-			}
-			else if (AvailablePhysical < JetsamLimits[1])
-			{
-				AvailablePhysical = JetsamLimits[1];
-			}
-			else if (AvailablePhysical < JetsamLimits[2])
-			{
-				AvailablePhysical = JetsamLimits[2];
-			}
-			else if (AvailablePhysical < JetsamLimits[3])
-			{
-				AvailablePhysical = JetsamLimits[3];
-			}
-			else if (AvailablePhysical < JetsamLimits[4])
-			{
-				AvailablePhysical = JetsamLimits[4];
-			}
-		}
-		else
-#endif
+        AvailablePhysical = os_proc_available_memory();
+        
+        // quantize to the known jetsam limits, we should be within 50MB of the correct one
+        uint64 JetsamLimits[] = { 1520435200, 1939865600, 2201170740, 2252710350, 3006477100 }; // { 2GB, gimped 3GB, gimped 4GB, 3GB, 4GB
+        if (AvailablePhysical < JetsamLimits[0])
+        {
+            AvailablePhysical = JetsamLimits[0];
+        }
+        else if (AvailablePhysical < JetsamLimits[1])
+        {
+            AvailablePhysical = JetsamLimits[1];
+        }
+        else if (AvailablePhysical < JetsamLimits[2])
+        {
+            AvailablePhysical = JetsamLimits[2];
+        }
+        else if (AvailablePhysical < JetsamLimits[3])
+        {
+            AvailablePhysical = JetsamLimits[3];
+        }
+        else if (AvailablePhysical < JetsamLimits[4])
+        {
+            AvailablePhysical = JetsamLimits[4];
+        }
+#else
+        int Mib[] = {CTL_HW, HW_MEMSIZE};
+        size_t Length = sizeof(int64);
+        sysctl(Mib, 2, &AvailablePhysical, &Length, NULL, 0);
 #endif
-		{
-			int Mib[] = {CTL_HW, HW_MEMSIZE};
-			size_t Length = sizeof(int64);
-			sysctl(Mib, 2, &AvailablePhysical, &Length, NULL, 0);
-		}
-		
 		MemoryConstants.TotalPhysical = AvailablePhysical;
 		MemoryConstants.TotalVirtual = AvailablePhysical;
 		MemoryConstants.PageSize = (uint32)vm_page_size;
diff --git a/Engine/Source/Runtime/Core/Private/IOS/IOSPlatformMisc.cpp b/Engine/Source/Runtime/Core/Private/IOS/IOSPlatformMisc.cpp
--- a/Engine/Source/Runtime/Core/Private/IOS/IOSPlatformMisc.cpp
+++ b/Engine/Source/Runtime/Core/Private/IOS/IOSPlatformMisc.cpp
@@ -263,16 +263,13 @@ bool FIOSPlatformMisc::IsRunningOnBattery()
 float FIOSPlatformMisc::GetDeviceTemperatureLevel()
 {
 #if !PLATFORM_TVOS
-	if (@available(iOS 11, *))
-	{
-		switch ([[IOSAppDelegate GetDelegate] GetThermalState])
-		{
-		case NSProcessInfoThermalStateNominal:	return (float)FCoreDelegates::ETemperatureSeverity::Good; break;
-		case NSProcessInfoThermalStateFair:		return (float)FCoreDelegates::ETemperatureSeverity::Bad; break;
-		case NSProcessInfoThermalStateSerious:	return (float)FCoreDelegates::ETemperatureSeverity::Serious; break;
-		case NSProcessInfoThermalStateCritical:	return (float)FCoreDelegates::ETemperatureSeverity::Critical; break;
-		}
-	}
+    switch ([[IOSAppDelegate GetDelegate] GetThermalState])
+    {
+        case NSProcessInfoThermalStateNominal:	return (float)FCoreDelegates::ETemperatureSeverity::Good; break;
+        case NSProcessInfoThermalStateFair:		return (float)FCoreDelegates::ETemperatureSeverity::Bad; break;
+        case NSProcessInfoThermalStateSerious:	return (float)FCoreDelegates::ETemperatureSeverity::Serious; break;
+        case NSProcessInfoThermalStateCritical:	return (float)FCoreDelegates::ETemperatureSeverity::Critical; break;
+    }
 #endif
 	return -1.0f;
 }
@@ -280,11 +277,8 @@ float FIOSPlatformMisc::GetDeviceTemperatureLevel()
 bool FIOSPlatformMisc::IsInLowPowerMode()
 {
 #if !PLATFORM_TVOS
-    if (@available(iOS 11, *))
-    {
-        bool bInLowPowerMode = [[NSProcessInfo processInfo] isLowPowerModeEnabled];
-        return bInLowPowerMode;
-    }
+    bool bInLowPowerMode = [[NSProcessInfo processInfo] isLowPowerModeEnabled];
+    return bInLowPowerMode;
 #endif
     return false;
 }
@@ -314,7 +308,7 @@ EDeviceScreenOrientation FIOSPlatformMisc::GetDeviceOrientation()
 #if !PLATFORM_TVOS
 	if (GInterfaceOrientation == UIInterfaceOrientationUnknown)
 	{
-		GInterfaceOrientation = [[UIApplication sharedApplication] statusBarOrientation];
+		GInterfaceOrientation = [[[[[UIApplication sharedApplication] delegate] window] windowScene] interfaceOrientation];
 	}
 
 	return ConvertFromUIInterfaceOrientation(GInterfaceOrientation);
@@ -818,31 +812,28 @@ bool FIOSPlatformMisc::GetDiskTotalAndFreeSpace(const FString& InPath, uint64& T
 {
     //On iOS 11 use new method to return disk space available for important usages
 #if !PLATFORM_TVOS
-    if (@available(iOS 11, *))
+    bool GetValueSuccess = false;
+    
+    NSNumber *FreeBytes = nil;
+    NSURL *URL = [NSURL fileURLWithPath : NSHomeDirectory()];
+    GetValueSuccess = [URL getResourceValue : &FreeBytes forKey : NSURLVolumeAvailableCapacityForImportantUsageKey error : nil];
+    if (FreeBytes)
     {
-	    bool GetValueSuccess = false;
-
-	    NSNumber *FreeBytes = nil;
-	    NSURL *URL = [NSURL fileURLWithPath : NSHomeDirectory()];
-	    GetValueSuccess = [URL getResourceValue : &FreeBytes forKey : NSURLVolumeAvailableCapacityForImportantUsageKey error : nil];
-	    if (FreeBytes)
-	    {
-	        NumberOfFreeBytes = [FreeBytes longLongValue];
-	    }
-
-	    NSNumber *TotalBytes = nil;
-	    GetValueSuccess = GetValueSuccess &&[URL getResourceValue : &TotalBytes forKey : NSURLVolumeTotalCapacityKey error : nil];
-	    if (TotalBytes)
-	    {
-	        TotalNumberOfBytes = [TotalBytes longLongValue];
-	    }
-
-	    if (GetValueSuccess
-	        && (NumberOfFreeBytes > 0)
-	        && (TotalNumberOfBytes > 0))
-	    {
-	        return true;
-	    }
+        NumberOfFreeBytes = [FreeBytes longLongValue];
+    }
+    
+    NSNumber *TotalBytes = nil;
+    GetValueSuccess = GetValueSuccess &&[URL getResourceValue : &TotalBytes forKey : NSURLVolumeTotalCapacityKey error : nil];
+    if (TotalBytes)
+    {
+        TotalNumberOfBytes = [TotalBytes longLongValue];
+    }
+    
+    if (GetValueSuccess
+        && (NumberOfFreeBytes > 0)
+        && (TotalNumberOfBytes > 0))
+    {
+        return true;
     }
 #endif
 
@@ -865,10 +856,7 @@ bool FIOSPlatformMisc::GetDiskTotalAndFreeSpace(const FString& InPath, uint64& T
 void FIOSPlatformMisc::RequestStoreReview()
 {
 #if !PLATFORM_TVOS
-	if (@available(iOS 10, *))
-	{
-		[SKStoreReviewController requestReview];
-	}
+	[SKStoreReviewController requestReview];
 #endif
 }
 
diff --git a/Engine/Source/Runtime/IOS/IOSRuntimeSettings/Classes/IOSRuntimeSettings.h b/Engine/Source/Runtime/IOS/IOSRuntimeSettings/Classes/IOSRuntimeSettings.h
--- a/Engine/Source/Runtime/IOS/IOSRuntimeSettings/Classes/IOSRuntimeSettings.h
+++ b/Engine/Source/Runtime/IOS/IOSRuntimeSettings/Classes/IOSRuntimeSettings.h
@@ -30,9 +30,6 @@ enum class EPowerUsageFrameRateLock : uint8
 UENUM()
 	enum class EIOSVersion : uint8
 {
-    /** iOS 12 */
-	IOS_12 = 12 UMETA(DisplayName = "12.0"),
-
 	/** iOS 13 */
 	IOS_13 = 13 UMETA(DisplayName = "13.0"),
 
@@ -43,20 +40,11 @@ UENUM()
 UENUM()
 enum class EIOSMetalShaderStandard : uint8
 {
-	/** Metal Shaders Compatible With iOS 10.0/tvOS 10.0 or later (std=ios-metal1.2) */
-	IOSMetalSLStandard_1_2 = 2 UMETA(DisplayName="Metal v1.2 (iOS 10.0/tvOS 10.0)", Hidden),
-	
-    /** Metal Shaders Compatible With iOS 11.0/tvOS 11.0 or later (std=ios-metal2.0) */
-	IOSMetalSLStandard_2_0 = 3 UMETA(DisplayName="Metal v2.0 (iOS 11.0/tvOS 11.0)"),
-    
-    /** Metal Shaders Compatible With iOS 12.0/tvOS 12.0 or later (std=ios-metal2.1) */
-    IOSMetalSLStandard_2_1 = 4 UMETA(DisplayName="Metal v2.1 (iOS 12.0/tvOS 12.0)"),
-
-	/** Metal Shaders Compatible With iOS 13.0/tvOS 13.0  or later (std=ios-metal2.2) */
-	IOSMetalSLStandard_2_2 = 5 UMETA(DisplayName="Metal v2.2 (iOS 13.0/tvOS 13.0)"),
+    /** Metal Shaders Compatible With iOS 13.0/tvOS 13.0 or later (std=ios-metal2.2) */
+    IOSMetalSLStandard_2_2 = 5 UMETA(DisplayName="Metal v2.2 (iOS 13.0/tvOS 13.0)"),
+    /** Metal Shaders Compatible With iOS 14.0/tvOS 14.0 or later (std=ios-metal2.3) */
+    IOSMetalSLStandard_2_3 = 6 UMETA(DisplayName="Metal v2.3 (iOS 14.0/tvOS 14.0)"),
 
-	/** Metal Shaders Compatible With iOS 14.0/tvOS 14.0  or later (std=ios-metal2.3) */
-	IOSMetalSLStandard_2_3 = 6 UMETA(DisplayName="Metal v2.3 (iOS 14.0/tvOS 14.0)"),
 };
 
 UENUM()
diff --git a/Engine/Source/Runtime/IOS/IOSRuntimeSettings/Private/IOSRuntimeSettings.cpp b/Engine/Source/Runtime/IOS/IOSRuntimeSettings/Private/IOSRuntimeSettings.cpp
--- a/Engine/Source/Runtime/IOS/IOSRuntimeSettings/Private/IOSRuntimeSettings.cpp
+++ b/Engine/Source/Runtime/IOS/IOSRuntimeSettings/Private/IOSRuntimeSettings.cpp
@@ -24,7 +24,7 @@ UIOSRuntimeSettings::UIOSRuntimeSettings(const FObjectInitializer& ObjectInitial
 	bEnableDynamicMaxFPS = false;
 	bSupportsIPad = true;
 	bSupportsIPhone = true;
-	MinimumiOSVersion = EIOSVersion::IOS_12;
+	MinimumiOSVersion = EIOSVersion::IOS_13;
     bBuildAsFramework = true;
 	bGeneratedSYMFile = false;
 	bGeneratedSYMBundle = false;
@@ -115,9 +115,9 @@ void UIOSRuntimeSettings::PostInitProperties()
 		}
 	}
 
-	if (MinimumiOSVersion < EIOSVersion::IOS_12)
+	if (MinimumiOSVersion < EIOSVersion::IOS_13)
 	{
-		MinimumiOSVersion = EIOSVersion::IOS_12;
+		MinimumiOSVersion = EIOSVersion::IOS_13;
 		UpdateSinglePropertyInConfigFile(GetClass()->FindPropertyByName(GET_MEMBER_NAME_CHECKED(UIOSRuntimeSettings, MinimumiOSVersion)), GetDefaultConfigFilename());
 	}
 	if (!bSupportsMetal && !bSupportsMetalMRT)
diff --git a/Engine/Source/Runtime/Launch/Private/IOS/GameLaunchDaemonMessageHandler.cpp b/Engine/Source/Runtime/Launch/Private/IOS/GameLaunchDaemonMessageHandler.cpp
--- a/Engine/Source/Runtime/Launch/Private/IOS/GameLaunchDaemonMessageHandler.cpp
+++ b/Engine/Source/Runtime/Launch/Private/IOS/GameLaunchDaemonMessageHandler.cpp
@@ -7,6 +7,7 @@
 #include "IOSMessageProtocol.h"
 #include "MessageEndpoint.h"
 #include "MessageEndpointBuilder.h"
+#include <UIKit/UIDevice.h>
 
 #import <TargetConditionals.h>
 
@@ -42,7 +43,7 @@ void FGameLaunchDaemonMessageHandler::HandlePingMessage(const FIOSLaunchDaemonPi
 		MessageEndpoint->Send(new FIOSLaunchDaemonPong(FString(FPlatformProperties::PlatformName()) + (TARGET_IPHONE_SIMULATOR ? FString(TEXT("Simulator:")) : FString(TEXT("@"))) + FString(FPlatformProcess::ComputerName())
 			, FPlatformProcess::ComputerName()
 			, "Game_Running"
-			, UI_USER_INTERFACE_IDIOM() == UIUserInterfaceIdiomPhone? "Phone" : "Tablet"
+			, [[UIDevice currentDevice] userInterfaceIdiom] == UIUserInterfaceIdiomPhone? "Phone" : "Tablet"
 			, false
 			, false
 			, false
diff --git a/Templates/TP_FirstPerson/Config/DefaultEngine.ini b/Templates/TP_FirstPerson/Config/DefaultEngine.ini
--- a/Templates/TP_FirstPerson/Config/DefaultEngine.ini
+++ b/Templates/TP_FirstPerson/Config/DefaultEngine.ini
@@ -17,5 +17,5 @@ GlobalDefaultGameMode=/Script/TP_FirstPerson.TP_FirstPersonGameMode
 GlobalDefaultServerGameMode=None
 
 [/Script/IOSRuntimeSettings.IOSRuntimeSettings]
-MinimumiOSVersion=IOS_12
+MinimumiOSVersion=IOS_13
 
diff --git a/Templates/TP_FirstPersonBP/Config/DefaultEngine.ini b/Templates/TP_FirstPersonBP/Config/DefaultEngine.ini
--- a/Templates/TP_FirstPersonBP/Config/DefaultEngine.ini
+++ b/Templates/TP_FirstPersonBP/Config/DefaultEngine.ini
@@ -17,5 +17,5 @@ GlobalDefaultGameMode=/Game/FirstPersonBP/Blueprints/FirstPersonGameMode.FirstPe
 GlobalDefaultServerGameMode=None
 
 [/Script/IOSRuntimeSettings.IOSRuntimeSettings]
-MinimumiOSVersion=IOS_12
+MinimumiOSVersion=IOS_13
 
diff --git a/Templates/TP_HandheldARBP/Config/DefaultEngine.ini b/Templates/TP_HandheldARBP/Config/DefaultEngine.ini
--- a/Templates/TP_HandheldARBP/Config/DefaultEngine.ini
+++ b/Templates/TP_HandheldARBP/Config/DefaultEngine.ini
@@ -109,7 +109,7 @@ MobileProvision=iOS11.mobileprovision
 bSupportsPortraitOrientation=True
 bSupportsUpsideDownOrientation=False
 FrameRateLock=PUFRL_60
-MinimumiOSVersion=IOS_12
+MinimumiOSVersion=IOS_13
 
 [/Script/Engine.Engine]
 bUseFixedFrameRate=True
diff --git a/Templates/TP_Puzzle/Config/DefaultEngine.ini b/Templates/TP_Puzzle/Config/DefaultEngine.ini
--- a/Templates/TP_Puzzle/Config/DefaultEngine.ini
+++ b/Templates/TP_Puzzle/Config/DefaultEngine.ini
@@ -13,5 +13,5 @@ GlobalDefaultGameMode=/Script/TP_Puzzle.TP_PuzzleGameMode
 GlobalDefaultServerGameMode=None
 
 [/Script/IOSRuntimeSettings.IOSRuntimeSettings]
-MinimumiOSVersion=IOS_12
+MinimumiOSVersion=IOS_13
 
diff --git a/Templates/TP_PuzzleBP/Config/DefaultEngine.ini b/Templates/TP_PuzzleBP/Config/DefaultEngine.ini
--- a/Templates/TP_PuzzleBP/Config/DefaultEngine.ini
+++ b/Templates/TP_PuzzleBP/Config/DefaultEngine.ini
@@ -12,4 +12,4 @@ GlobalDefaultGameMode=/Game/PuzzleBP/Blueprints/PuzzleGameMode.PuzzleGameMode_C
 GlobalDefaultServerGameMode=None
 
 [/Script/IOSRuntimeSettings.IOSRuntimeSettings]
-MinimumiOSVersion=IOS_12
+MinimumiOSVersion=IOS_13
diff --git a/Templates/TP_ThirdPerson/Config/DefaultEngine.ini b/Templates/TP_ThirdPerson/Config/DefaultEngine.ini
--- a/Templates/TP_ThirdPerson/Config/DefaultEngine.ini
+++ b/Templates/TP_ThirdPerson/Config/DefaultEngine.ini
@@ -4,4 +4,4 @@ EditorStartupMap=/Game/ThirdPersonCPP/Maps/ThirdPersonExampleMap
 GlobalDefaultGameMode="/Script/TP_ThirdPerson.TP_ThirdPersonGameMode"
 
 [/Script/IOSRuntimeSettings.IOSRuntimeSettings]
-MinimumiOSVersion=IOS_12
+MinimumiOSVersion=IOS_13
diff --git a/Templates/TP_ThirdPersonBP/Config/DefaultEngine.ini b/Templates/TP_ThirdPersonBP/Config/DefaultEngine.ini
--- a/Templates/TP_ThirdPersonBP/Config/DefaultEngine.ini
+++ b/Templates/TP_ThirdPersonBP/Config/DefaultEngine.ini
@@ -12,4 +12,4 @@ GlobalDefaultGameMode=/Game/ThirdPersonBP/Blueprints/ThirdPersonGameMode.ThirdPe
 GlobalDefaultServerGameMode=None
 
 [/Script/IOSRuntimeSettings.IOSRuntimeSettings]
-MinimumiOSVersion=IOS_12
+MinimumiOSVersion=IOS_13
diff --git a/Templates/TP_TopDown/Config/DefaultEngine.ini b/Templates/TP_TopDown/Config/DefaultEngine.ini
--- a/Templates/TP_TopDown/Config/DefaultEngine.ini
+++ b/Templates/TP_TopDown/Config/DefaultEngine.ini
@@ -4,4 +4,4 @@ EditorStartupMap=/Game/TopDownCPP/Maps/TopDownExampleMap
 GlobalDefaultGameMode="/Script/TP_TopDown.TP_TopDownGameMode"
 
 [/Script/IOSRuntimeSettings.IOSRuntimeSettings]
-MinimumiOSVersion=IOS_12
+MinimumiOSVersion=IOS_13
diff --git a/Templates/TP_TopDownBP/Config/DefaultEngine.ini b/Templates/TP_TopDownBP/Config/DefaultEngine.ini
--- a/Templates/TP_TopDownBP/Config/DefaultEngine.ini
+++ b/Templates/TP_TopDownBP/Config/DefaultEngine.ini
@@ -50,5 +50,5 @@ VerticalDeviationFromGroundCompensation=0.000000
 bForceRebuildOnLoad=True
 
 [/Script/IOSRuntimeSettings.IOSRuntimeSettings]
-MinimumiOSVersion=IOS_12
+MinimumiOSVersion=IOS_13
 
diff --git a/Templates/TP_VehicleAdv/Config/DefaultEngine.ini b/Templates/TP_VehicleAdv/Config/DefaultEngine.ini
--- a/Templates/TP_VehicleAdv/Config/DefaultEngine.ini
+++ b/Templates/TP_VehicleAdv/Config/DefaultEngine.ini
@@ -15,5 +15,5 @@ GlobalDefaultGameMode=/Script/TP_VehicleAdv.TP_VehicleAdvGameMode
 GlobalDefaultServerGameMode=None
 
 [/Script/IOSRuntimeSettings.IOSRuntimeSettings]
-MinimumiOSVersion=IOS_12
+MinimumiOSVersion=IOS_13
 
diff --git a/Templates/TP_VehicleAdvBP/Config/DefaultEngine.ini b/Templates/TP_VehicleAdvBP/Config/DefaultEngine.ini
--- a/Templates/TP_VehicleAdvBP/Config/DefaultEngine.ini
+++ b/Templates/TP_VehicleAdvBP/Config/DefaultEngine.ini
@@ -18,5 +18,5 @@ GlobalDefaultGameMode=/Game/VehicleAdvBP/Blueprints/VehicleAdvGameMode.VehicleAd
 GlobalDefaultServerGameMode=None
 
 [/Script/IOSRuntimeSettings.IOSRuntimeSettings]
-MinimumiOSVersion=IOS_12
+MinimumiOSVersion=IOS_13
 
diff --git a/Templates/TP_VirtualRealityBP/Config/DefaultEngine.ini b/Templates/TP_VirtualRealityBP/Config/DefaultEngine.ini
--- a/Templates/TP_VirtualRealityBP/Config/DefaultEngine.ini
+++ b/Templates/TP_VirtualRealityBP/Config/DefaultEngine.ini
@@ -1,262 +1,61 @@
-
-
-[/Script/Engine.RendererSettings]
-r.Mobile.DisableVertexFog=True
-r.Shadow.CSM.MaxMobileCascades=2
-r.MobileMSAA=4
-r.Mobile.AllowDitheredLODTransition=False
-r.Mobile.AllowSoftwareOcclusion=False
-r.Mobile.VirtualTextures=False
-r.DiscardUnusedQuality=False
-r.AllowOcclusionQueries=True
-r.MinScreenRadiusForLights=0.030000
-r.MinScreenRadiusForDepthPrepass=0.030000
-r.MinScreenRadiusForCSMDepth=0.010000
-r.PrecomputedVisibilityWarning=False
-r.TextureStreaming=True
-Compat.UseDXT5NormalMaps=False
-r.VirtualTextures=False
-r.VirtualTexturedLightmaps=False
-r.VT.TileSize=128
-r.VT.TileBorderSize=4
-r.vt.FeedbackFactor=16
-r.VT.EnableCompressZlib=True
-r.VT.EnableCompressCrunch=False
-r.ClearCoatNormal=False
-r.ReflectionCaptureResolution=128
-r.Mobile.ReflectionCaptureCompression=False
-r.ReflectionEnvironmentLightmapMixBasedOnRoughness=True
-r.ForwardShading=True
-r.VertexFoggingForOpaque=True
-r.AllowStaticLighting=True
-r.NormalMapsForStaticLighting=False
-r.GenerateMeshDistanceFields=False
-r.DistanceFieldBuild.EightBit=False
-r.GenerateLandscapeGIData=False
-r.DistanceFieldBuild.Compress=False
-r.TessellationAdaptivePixelsPerTriangle=48.000000
-r.SeparateTranslucency=False
-r.TranslucentSortPolicy=0
-TranslucentSortAxis=(X=0.000000,Y=-1.000000,Z=0.000000)
-r.CustomDepth=1
-r.CustomDepthTemporalAAJitter=True
-r.PostProcessing.PropagateAlpha=0
-r.DefaultFeature.Bloom=False
-r.DefaultFeature.AmbientOcclusion=False
-r.DefaultFeature.AmbientOcclusionStaticFraction=True
-r.DefaultFeature.AutoExposure=False
-r.DefaultFeature.AutoExposure.Method=0
-r.DefaultFeature.AutoExposure.Bias=1.000000
-r.DefaultFeature.AutoExposure.ExtendDefaultLuminanceRange=True
-r.UsePreExposure=True
-r.EyeAdaptation.EditorOnly=False
-r.DefaultFeature.MotionBlur=False
-r.DefaultFeature.LensFlare=False
-r.TemporalAA.Upsampling=False
-r.SSGI.Enable=False
-r.DefaultFeature.AntiAliasing=3
-r.DefaultFeature.LightUnits=1
-r.DefaultBackBufferPixelFormat=4
-r.Shadow.UnbuiltPreviewInGame=True
-r.StencilForLODDither=False
-r.EarlyZPass=3
-r.EarlyZPassOnlyMaterialMasking=False
-r.DBuffer=True
-r.ClearSceneMethod=1
-r.BasePassOutputsVelocity=False
-r.VertexDeformationOutputsVelocity=False
-r.SelectiveBasePassOutputs=False
-bDefaultParticleCutouts=False
-fx.GPUSimulationTextureSizeX=1024
-fx.GPUSimulationTextureSizeY=1024
-r.AllowGlobalClipPlane=False
-r.GBufferFormat=1
-r.MorphTarget.Mode=True
-r.GPUCrashDebugging=False
-vr.InstancedStereo=True
-r.MobileHDR=False
-vr.MobileMultiView=True
-r.Mobile.UseHWsRGBEncoding=True
-vr.RoundRobinOcclusion=False
-vr.ODSCapture=False
-r.MeshStreaming=False
-r.WireframeCullThreshold=5.000000
-r.RayTracing=False
-r.RayTracing.UseTextureLod=False
-r.SupportStationarySkylight=True
-r.SupportLowQualityLightmaps=True
-r.SupportPointLightWholeSceneShadows=True
-r.SupportAtmosphericFog=True
-r.SupportSkyAtmosphere=True
-r.SupportSkyAtmosphereAffectsHeightFog=False
-r.SkinCache.CompileShaders=False
-r.SkinCache.DefaultBehavior=1
-r.SkinCache.SceneMemoryLimitInMB=128.000000
-r.Mobile.EnableStaticAndCSMShadowReceivers=True
-r.Mobile.EnableMovableLightCSMShaderCulling=True
-r.Mobile.AllowDistanceFieldShadows=True
-r.Mobile.AllowMovableDirectionalLights=True
-r.MobileNumDynamicPointLights=4
-r.MobileDynamicPointLightsUseStaticBranch=True
-r.Mobile.EnableMovableSpotlights=False
-r.Mobile.EnableMovableSpotlightsShadow=False
-r.GPUSkin.Support16BitBoneIndex=False
-r.GPUSkin.Limit2BoneInfluences=False
-r.SupportDepthOnlyIndexBuffers=True
-r.SupportReversedIndexBuffers=True
-r.LightPropagationVolume=False
-r.Mobile.AmbientOcclusion=False
-r.GPUSkin.UnlimitedBoneInfluences=False
-r.GPUSkin.UnlimitedBoneInfluencesThreshold=8
-r.Mobile.PlanarReflectionMode=0
-bStreamSkeletalMeshLODs=(Default=False,PerPlatform=())
-bDiscardSkeletalMeshOptionalLODs=(Default=False,PerPlatform=())
-VisualizeCalibrationColorMaterialPath=None
-VisualizeCalibrationCustomMaterialPath=None
-VisualizeCalibrationGrayscaleMaterialPath=None
-
-[/Script/HardwareTargeting.HardwareTargetingSettings]
-TargetedHardwareClass=Mobile
-AppliedTargetedHardwareClass=Mobile
-DefaultGraphicsPerformance=Scalable
-AppliedDefaultGraphicsPerformance=Scalable
-
 [/Script/EngineSettings.GameMapsSettings]
-EditorStartupMap=/Game/VRTemplate/Maps/VRTemplateMap.VRTemplateMap
+EditorStartupMap=/Game/VirtualRealityBP/Maps/StartupMap.StartupMap
 LocalMapOptions=
-TransitionMap=None
-bUseSplitscreen=False
+TransitionMap=
+bUseSplitscreen=True
 TwoPlayerSplitscreenLayout=Horizontal
 ThreePlayerSplitscreenLayout=FavorTop
-FourPlayerSplitscreenLayout=Grid
-bOffsetPlayerGamepadIds=False
 GameInstanceClass=/Script/Engine.GameInstance
-GameDefaultMap=/Game/VRTemplate/Maps/VRTemplateMap.VRTemplateMap
-ServerDefaultMap=/Engine/Maps/Entry.Entry
-GlobalDefaultGameMode=/Game/VRTemplate/Blueprints/VRGameMode.VRGameMode_C
+GameDefaultMap=/Game/VirtualRealityBP/Maps/MotionControllerMap.MotionControllerMap
+ServerDefaultMap=/Engine/Maps/Entry
+GlobalDefaultGameMode=/Script/Engine.GameModeBase
 GlobalDefaultServerGameMode=None
 
-[/Script/Slate.SlateSettings]
-bExplicitCanvasChildZOrder=True
+[/Script/Engine.Engine]
++ActiveGameNameRedirects=(OldGameName="TP_FirstPersonBP",NewGameName="/Script/VRTemplate")
++ActiveGameNameRedirects=(OldGameName="/Script/TP_FirstPersonBP",NewGameName="/Script/VRTemplate")
 
-[/Script/WindowsTargetPlatform.WindowsTargetSettings]
-SpatializationPlugin=Resonance Audio
-ReverbPlugin=Resonance Audio
+[/Script/HardwareTargeting.HardwareTargetingSettings]
+TargetedHardwareClass=Desktop
+AppliedTargetedHardwareClass=Desktop
+DefaultGraphicsPerformance=Maximum
+AppliedDefaultGraphicsPerformance=Maximum
 
-[/Script/AndroidRuntimeSettings.AndroidRuntimeSettings]
-PackageName=com.epicgames.[PROJECT]
-StoreVersion=1
-StoreVersionOffsetArmV7=0
-StoreVersionOffsetArm64=0
-StoreVersionOffsetX8664=0
-ApplicationDisplayName=
-VersionDisplayName=1.0
-MinSDKVersion=23
-TargetSDKVersion=25
-InstallLocation=InternalOnly
-bEnableGradle=True
-bEnableLint=False
-bPackageDataInsideApk=True
-bCreateAllPlatformsInstall=False
-bDisableVerifyOBBOnStartUp=False
-bForceSmallOBBFiles=False
-bAllowLargeOBBFiles=False
-bAllowPatchOBBFile=False
-bAllowOverflowOBBFiles=False
-bUseExternalFilesDir=False
-bPublicLogFiles=True
-Orientation=SensorLandscape
-MaxAspectRatio=2.100000
-bUseDisplayCutout=False
-bRestoreNotificationsOnReboot=False
-bFullScreen=True
-bEnableNewKeyboard=True
-DepthBufferPreference=Default
-bValidateTextureFormats=True
-bEnableBundle=False
-bEnableUniversalAPK=True
-bBundleABISplit=True
-bBundleLanguageSplit=True
-bBundleDensitySplit=True
-ExtraApplicationSettings=
-ExtraActivitySettings=
-bAndroidVoiceEnabled=False
-+PackageForOculusMobile=Quest2
-+PackageForOculusMobile=Quest
-bRemoveOSIG=False
-+GoogleVRCaps=Daydream33
-bGoogleVRSustainedPerformance=False
-KeyStore=UE.keystore
-KeyAlias=DistroKey
-KeyStorePassword=codesigner
-KeyPassword=
-bBuildForArmV7=True
-bBuildForArm64=False
-bBuildForX8664=False
-bBuildForES31=True
-bSupportsVulkan=True
-bSupportsVulkanSM5=False
-ClangSanitizer=None
-bAndroidOpenGLSupportsBackbufferSampling=False
-bDetectVulkanByDefault=True
-bBuildWithHiddenSymbolVisibility=False
-bSaveSymbols=False
-bForceLDLinker=False
-bEnableGooglePlaySupport=False
-bUseGetAccounts=False
-GamesAppID=
-bEnableSnapshots=False
-bSupportAdMob=True
-AdMobAdUnitID=
-GooglePlayLicenseKey=
-GCMClientSenderID=
-bShowLaunchImage=True
-bAllowIMU=True
-bAllowControllers=True
-bBlockAndroidKeysOnControllers=False
-bControllersBlockDeviceFeedback=False
-AndroidAudio=Default
-AudioSampleRate=44100
-AudioCallbackBufferFrameSize=1024
-AudioNumBuffersToEnqueue=4
-AudioMaxChannels=0
-AudioNumSourceWorkers=0
-SpatializationPlugin=Resonance Audio
-ReverbPlugin=Resonance Audio
-OcclusionPlugin=
-CompressionOverrides=(bOverrideCompressionTimes=False,DurationThreshold=5.000000,MaxNumRandomBranches=0,SoundCueQualityIndex=0)
-bUseAudioStreamCaching=False
-CacheSizeKB=0
-MaxChunkSizeOverrideKB=0
-bResampleForDevice=False
-SoundCueCookQualityIndex=-1
-MaxSampleRate=0.000000
-HighSampleRate=0.000000
-MedSampleRate=0.000000
-LowSampleRate=0.000000
-MinSampleRate=0.000000
-CompressionQualityModifier=0.000000
-AutoStreamingThreshold=0.000000
-AndroidGraphicsDebugger=None
-MaliGraphicsDebuggerPath=(Path="")
-bEnableMaliPerfCounters=False
-bMultiTargetFormat_ETC2=True
-bMultiTargetFormat_DXT=True
-bMultiTargetFormat_ASTC=True
-TextureFormatPriority_ETC2=0.200000
-TextureFormatPriority_DXT=0.600000
-TextureFormatPriority_ASTC=0.900000
-SDKAPILevelOverride=
-NDKAPILevelOverride=
-bStreamLandscapeMeshLODs=False
+[/Script/Engine.RendererSettings]
+r.DefaultFeature.AmbientOcclusion=False
+r.DefaultFeature.AmbientOcclusionStaticFraction=False
+vr.InstancedStereo=True
+r.SeparateTranslucency=0
+r.HZBOcclusion=0
+r.EyeAdaptionQuality=0
+r.SSR.MaxRoughness=0
+r.rhicmdbypass=0
+r.TiledReflectionEnvironmentMinimumCount=10
+
+[/Script/NavigationSystem.NavigationSystemV1]
+bAutoCreateNavigationData=True
+bAllowClientSideNavigation=False
+bInitialBuildingLocked=False
+bSkipAgentHeightCheckWhenPickingNavData=False
+DataGatheringMode=Instant
+bGenerateNavigationOnlyAroundNavigationInvokers=False
+ActiveTilesUpdateInterval=1.000000
+DirtyAreasUpdateFreq=60.000000
+
+[/Script/Engine.UserInterfaceSettings]
+RenderFocusRule=NavigationOnly
+DefaultCursor=None
+TextEditBeamCursor=None
+CrosshairsCursor=None
+HandCursor=None
+GrabHandCursor=None
+GrabHandClosedCursor=None
+SlashedCircleCursor=None
+ApplicationScale=1.000000
+UIScaleRule=ShortestSide
+CustomScalingRuleClass=None
+UIScaleCurve=(EditorCurveData=(PreInfinityExtrap=RCCE_Constant,PostInfinityExtrap=RCCE_Constant,Keys=((Time=480.000000,Value=0.444000),(Time=720.000000,Value=0.666000),(Time=1080.000000,Value=1.000000),(Time=8640.000000,Value=8.000000)),DefaultValue=340282346638528859811704183484516925440.000000),ExternalCurve=None)
 
 [/Script/IOSRuntimeSettings.IOSRuntimeSettings]
-SpatializationPlugin=Resonance Audio
-
-[/Script/LuminRuntimeSettings.LuminRuntimeSettings]
-SpatializationPlugin=Built-in Spatialization
-
-[/Script/MacTargetPlatform.MacTargetSettings]
-SpatializationPlugin=Resonance Audio
-ReverbPlugin=Resonance Audio
+MinimumiOSVersion=IOS_13
 
