Reinstate menu button special right (pause) workaround logic due to iOS/tvOS 14 unreliable buttonMenu behavior.

diff --git a/Engine/Source/Runtime/ApplicationCore/Private/Apple/AppleControllerInterface.cpp b/Engine/Source/Runtime/ApplicationCore/Private/Apple/AppleControllerInterface.cpp
--- a/Engine/Source/Runtime/ApplicationCore/Private/Apple/AppleControllerInterface.cpp
+++ b/Engine/Source/Runtime/ApplicationCore/Private/Apple/AppleControllerInterface.cpp
@@ -136,6 +136,15 @@ void FAppleControllerInterface::HandleConnection(GCController* Controller)
         Controllers[ControllerIndex].Controller = Controller;
         SetControllerType(ControllerIndex);
         
+        // Deprecated but buttonMenu behavior is unreliable in iOS/tvOS 14
+		Controllers[ControllerIndex].bPauseWasPressed = false;
+PRAGMA_DISABLE_DEPRECATION_WARNINGS
+		Controller.controllerPausedHandler = ^(GCController* Cont)
+		{
+			Controllers[ControllerIndex].bPauseWasPressed = true;
+		};
+PRAGMA_ENABLE_DEPRECATION_WARNINGS
+        
         bFoundSlot = true;
         
         UE_LOG(LogAppleController, Log, TEXT("New %s controller inserted, assigned to playerIndex %d"),
@@ -184,6 +193,15 @@ void FAppleControllerInterface::SendControllerEvents()
 		// Assumes iOS13, tvOS 13 & macOS 10.15
         GCExtendedGamepad* ExtendedGamepad = [Cont capture].extendedGamepad;
 		GCMotion* Motion = Cont.motion;
+		
+		// Workaround for unreliable buttonMenu behavior in iOS/tvOS 14
+		if (Controller.bPauseWasPressed)
+        {
+            MessageHandler->OnControllerButtonPressed(FGamepadKeyNames::SpecialRight, Controllers[i].PlayerIndex, false);
+            MessageHandler->OnControllerButtonReleased(FGamepadKeyNames::SpecialRight, Controllers[i].PlayerIndex, false);
+
+            Controller.bPauseWasPressed = false;
+        }
         
 		if (ExtendedGamepad != nil)
 		{
diff --git a/Engine/Source/Runtime/ApplicationCore/Private/Apple/AppleControllerInterface.h b/Engine/Source/Runtime/ApplicationCore/Private/Apple/AppleControllerInterface.h
--- a/Engine/Source/Runtime/ApplicationCore/Private/Apple/AppleControllerInterface.h
+++ b/Engine/Source/Runtime/ApplicationCore/Private/Apple/AppleControllerInterface.h
@@ -105,6 +105,9 @@ protected:
 		bool bNeedsReferenceAttitude;
 		bool bHasReferenceAttitude;
 #endif
+
+		// Workaround for unreliable buttonMenu behavior in iOS/tvOS 14
+        bool bPauseWasPressed;
 	};
 	
     // there is a hardcoded limit of 4 controllers in the API
