Removal of dependencies to Windows Vista or earlier.

diff --git a/Engine/Build/Windows/Resources/Default-Win64.manifest b/Engine/Build/Windows/Resources/Default-Win64.manifest
--- a/Engine/Build/Windows/Resources/Default-Win64.manifest
+++ b/Engine/Build/Windows/Resources/Default-Win64.manifest
@@ -1,7 +1,7 @@
 <?xml version="1.0" encoding="UTF-8" standalone="yes"?>
 <assembly xmlns="urn:schemas-microsoft-com:asm.v1" manifestVersion="1.0" xmlns:asmv3="urn:schemas-microsoft-com:asm.v3">
 
-  <!-- Tell Vista we don't need administrator level access-->
+  <!-- Tell Windows we don't need administrator level access-->
   <ms_asmv2:trustInfo xmlns:ms_asmv2="urn:schemas-microsoft-com:asm.v2">
     <ms_asmv2:security>
       <ms_asmv2:requestedPrivileges>
@@ -24,8 +24,6 @@
       <supportedOS Id="{8e0f7a12-bfb3-4fe8-b9a5-48fd50a15a9a}"/>
       <!-- Windows 8.1 -->
       <supportedOS Id="{1f676c76-80e1-4239-95bb-83d0f6d0da78}"/>
-      <!-- Windows Vista -->
-      <supportedOS Id="{e2011457-1546-43c5-a5fe-008deee3d3f0}"/>
       <!-- Windows 7 -->
       <supportedOS Id="{35138b9a-5d96-4fbd-8e2d-a2440225f93a}"/>
       <!-- Windows 8 -->
diff --git a/Engine/Config/BaseEngine.ini b/Engine/Config/BaseEngine.ini
--- a/Engine/Config/BaseEngine.ini
+++ b/Engine/Config/BaseEngine.ini
@@ -2386,7 +2386,6 @@ bEnable2DPhysics=false
 
 [/Script/WindowsTargetPlatform.WindowsTargetSettings]
 +TargetedRHIs=PCD3D_SM5
-MinimumOSVersion=MSOS_Vista
 bEnableRayTracing=true
 bTarget32Bit=false
 
diff --git a/Engine/Plugins/Media/MfMedia/Source/MfMedia/Private/MfMediaPrivate.h b/Engine/Plugins/Media/MfMedia/Source/MfMedia/Private/MfMediaPrivate.h
--- a/Engine/Plugins/Media/MfMedia/Source/MfMedia/Private/MfMediaPrivate.h
+++ b/Engine/Plugins/Media/MfMedia/Source/MfMedia/Private/MfMediaPrivate.h
@@ -5,7 +5,7 @@
 #include "Logging/LogMacros.h"
 
 #if !defined(MFMEDIA_SUPPORTED_PLATFORM)
-	#define MFMEDIA_SUPPORTED_PLATFORM (PLATFORM_XBOXONE || PLATFORM_HOLOLENS || (PLATFORM_WINDOWS && WINVER >= 0x0601 /*Win7*/ && !UE_SERVER))
+	#define MFMEDIA_SUPPORTED_PLATFORM (PLATFORM_XBOXONE || PLATFORM_HOLOLENS || (PLATFORM_WINDOWS && !UE_SERVER))
 	#define MFMEDIA_NEED_PLATFORM_PRIVATE 0
 #else
 	#define MFMEDIA_NEED_PLATFORM_PRIVATE 1
diff --git a/Engine/Plugins/Media/MfMedia/Source/MfMediaFactory/Private/MfMediaFactoryPrivate.h b/Engine/Plugins/Media/MfMedia/Source/MfMediaFactory/Private/MfMediaFactoryPrivate.h
--- a/Engine/Plugins/Media/MfMedia/Source/MfMediaFactory/Private/MfMediaFactoryPrivate.h
+++ b/Engine/Plugins/Media/MfMedia/Source/MfMediaFactory/Private/MfMediaFactoryPrivate.h
@@ -4,7 +4,7 @@
 
 #include "Logging/LogMacros.h"
 
-#define MFMEDIAFACTORY_WINDOWS (PLATFORM_WINDOWS && (WINVER >= 0x601 /*Win7*/))
+#define MFMEDIAFACTORY_WINDOWS (PLATFORM_WINDOWS)
 
 
 /** Log category for the MfMediaFactory module. */
diff --git a/Engine/Plugins/Media/WmfMedia/Source/WmfMedia/Private/Wmf/WmfMediaUtils.cpp b/Engine/Plugins/Media/WmfMedia/Source/WmfMedia/Private/Wmf/WmfMediaUtils.cpp
--- a/Engine/Plugins/Media/WmfMedia/Source/WmfMedia/Private/Wmf/WmfMediaUtils.cpp
+++ b/Engine/Plugins/Media/WmfMedia/Source/WmfMedia/Private/Wmf/WmfMediaUtils.cpp
@@ -1814,10 +1814,7 @@ namespace WmfMedia
 		case MF_TOPOSTATUS_READY: return TEXT("Ready");
 		case MF_TOPOSTATUS_SINK_SWITCHED: return TEXT("Sink Switched");
 		case MF_TOPOSTATUS_STARTED_SOURCE: return TEXT("Started Source");
-
-#if (WINVER >= _WIN32_WINNT_WIN7)
 		case MF_TOPOSTATUS_DYNAMIC_CHANGED: return TEXT("Dynamic Changed");
-#endif
 
 		default:
 			return FString::Printf(TEXT("Unknown status %i"), (int32)Status);
diff --git a/Engine/Plugins/Media/WmfMedia/Source/WmfMedia/Public/WmfMediaCommon.h b/Engine/Plugins/Media/WmfMedia/Source/WmfMedia/Public/WmfMediaCommon.h
--- a/Engine/Plugins/Media/WmfMedia/Source/WmfMedia/Public/WmfMediaCommon.h
+++ b/Engine/Plugins/Media/WmfMedia/Source/WmfMedia/Public/WmfMediaCommon.h
@@ -4,7 +4,7 @@
 
 #include "Logging/LogMacros.h"
 
-#define WMFMEDIA_SUPPORTED_PLATFORM (PLATFORM_WINDOWS && (WINVER >= 0x0600 /*Vista*/) && !UE_SERVER)
+#define WMFMEDIA_SUPPORTED_PLATFORM (PLATFORM_WINDOWS && !UE_SERVER)
 
 #if WMFMEDIA_SUPPORTED_PLATFORM
 	#include "WmfMediaSettings.h"
diff --git a/Engine/Plugins/Runtime/Oculus/OculusVR/Source/OculusHMD/Public/IOculusHMDModule.h b/Engine/Plugins/Runtime/Oculus/OculusVR/Source/OculusHMD/Public/IOculusHMDModule.h
--- a/Engine/Plugins/Runtime/Oculus/OculusVR/Source/OculusHMD/Public/IOculusHMDModule.h
+++ b/Engine/Plugins/Runtime/Oculus/OculusVR/Source/OculusHMD/Public/IOculusHMDModule.h
@@ -5,8 +5,7 @@
 #include "IHeadMountedDisplayModule.h"
 #include "HeadMountedDisplayTypes.h"
 
-// Oculus support is not available on Windows XP
-#define OCULUS_HMD_SUPPORTED_PLATFORMS (PLATFORM_WINDOWS && WINVER > 0x0502)  || (PLATFORM_ANDROID_ARM || PLATFORM_ANDROID_ARM64)
+#define OCULUS_HMD_SUPPORTED_PLATFORMS (PLATFORM_WINDOWS || PLATFORM_ANDROID_ARM || PLATFORM_ANDROID_ARM64)
 
 //-------------------------------------------------------------------------------------------------
 // IOculusHMDModule
diff --git a/Engine/Plugins/Runtime/Oculus/OculusVR/Source/OculusInput/Public/IOculusInputModule.h b/Engine/Plugins/Runtime/Oculus/OculusVR/Source/OculusInput/Public/IOculusInputModule.h
--- a/Engine/Plugins/Runtime/Oculus/OculusVR/Source/OculusInput/Public/IOculusInputModule.h
+++ b/Engine/Plugins/Runtime/Oculus/OculusVR/Source/OculusInput/Public/IOculusInputModule.h
@@ -4,7 +4,7 @@
 #include "Modules/ModuleManager.h"
 #include "IInputDeviceModule.h"
 
-#define OCULUS_INPUT_SUPPORTED_PLATFORMS (PLATFORM_WINDOWS && WINVER > 0x0502) || (PLATFORM_ANDROID_ARM || PLATFORM_ANDROID_ARM64)
+#define OCULUS_INPUT_SUPPORTED_PLATFORMS (PLATFORM_WINDOWS || PLATFORM_ANDROID_ARM || PLATFORM_ANDROID_ARM64)
 
 
 /**
diff --git a/Engine/Plugins/Runtime/Oculus/OculusVR/Source/OculusMR/Public/IOculusMRModule.h b/Engine/Plugins/Runtime/Oculus/OculusVR/Source/OculusMR/Public/IOculusMRModule.h
--- a/Engine/Plugins/Runtime/Oculus/OculusVR/Source/OculusMR/Public/IOculusMRModule.h
+++ b/Engine/Plugins/Runtime/Oculus/OculusVR/Source/OculusMR/Public/IOculusMRModule.h
@@ -3,8 +3,7 @@
 #pragma once
 #include "Modules/ModuleManager.h"
 
-// Oculus support is not available on Windows XP
-#define OCULUS_MR_SUPPORTED_PLATFORMS ((PLATFORM_WINDOWS && WINVER > 0x0502) || PLATFORM_ANDROID)
+#define OCULUS_MR_SUPPORTED_PLATFORMS (PLATFORM_WINDOWS || PLATFORM_ANDROID)
 
 /**
  * The public interface to this module.  In most cases, this interface is only public to sibling modules 
diff --git a/Engine/Plugins/Runtime/Steam/SteamVR/Source/SteamVR/Public/ISteamVRPlugin.h b/Engine/Plugins/Runtime/Steam/SteamVR/Source/SteamVR/Public/ISteamVRPlugin.h
--- a/Engine/Plugins/Runtime/Steam/SteamVR/Source/SteamVR/Public/ISteamVRPlugin.h
+++ b/Engine/Plugins/Runtime/Steam/SteamVR/Source/SteamVR/Public/ISteamVRPlugin.h
@@ -2,7 +2,7 @@
 
 #pragma once
 
-#define STEAMVR_SUPPORTED_PLATFORMS (PLATFORM_MAC || (PLATFORM_LINUX && PLATFORM_CPU_X86_FAMILY && PLATFORM_64BITS) || (PLATFORM_WINDOWS && WINVER > 0x0502))
+#define STEAMVR_SUPPORTED_PLATFORMS (PLATFORM_MAC || (PLATFORM_LINUX && PLATFORM_CPU_X86_FAMILY && PLATFORM_64BITS) || PLATFORM_WINDOWS)
 
 #include "Modules/ModuleManager.h"
 #include "IHeadMountedDisplayModule.h"
diff --git a/Engine/Source/Developer/DesktopPlatform/Private/Windows/WindowsRegistry.cpp b/Engine/Source/Developer/DesktopPlatform/Private/Windows/WindowsRegistry.cpp
--- a/Engine/Source/Developer/DesktopPlatform/Private/Windows/WindowsRegistry.cpp
+++ b/Engine/Source/Developer/DesktopPlatform/Private/Windows/WindowsRegistry.cpp
@@ -3,11 +3,6 @@
 #include "WindowsRegistry.h"
 #include "Windows/AllowWindowsPlatformTypes.h"
 
-#if WINVER == 0x0502
-	// Some registry functions are not available on windows XP, so we need the XP compat alternative
-	#include <Shlwapi.h>
-#endif
-
 ////////////////////////////////////////////////////////////////////
 
 FRegistryValue::FRegistryValue()
@@ -182,12 +177,7 @@ bool FRegistryKey::Write(HKEY hKey, bool bRemoveDifferences) const
 		// Remove any keys that don't exist any more
 		for (int32 Idx = 0; Idx < KeyNames.Num(); Idx++)
 		{
-#if WINVER == 0x0502
-			if (!Keys.Contains(KeyNames[Idx]) && SHDeleteKey(hKey, *KeyNames[Idx]) != ERROR_SUCCESS)
-#else
-
 			if (!Keys.Contains(KeyNames[Idx]) && RegDeleteTree(hKey, *KeyNames[Idx]) != ERROR_SUCCESS)
-#endif // WINVER == 0x0502
 			{
 				return false;
 			}
diff --git a/Engine/Source/Developer/DirectoryWatcher/Private/Windows/DirectoryWatchRequestWindows.cpp b/Engine/Source/Developer/DirectoryWatcher/Private/Windows/DirectoryWatchRequestWindows.cpp
--- a/Engine/Source/Developer/DirectoryWatcher/Private/Windows/DirectoryWatchRequestWindows.cpp
+++ b/Engine/Source/Developer/DirectoryWatcher/Private/Windows/DirectoryWatchRequestWindows.cpp
@@ -137,11 +137,7 @@ void FDirectoryWatchRequestWindows::EndWatchRequest()
 	{
 		if ( DirectoryHandle != INVALID_HANDLE_VALUE )
 		{
-#if WINVER >= 0x600		// CancelIoEx() is only supported on Windows Vista and higher
 			CancelIoEx(DirectoryHandle, &Overlapped);
-#else
-			CancelIo(DirectoryHandle);
-#endif
 			// Clear the handle so we don't setup any more requests, and wait for the operation to finish
 			HANDLE TempDirectoryHandle = DirectoryHandle;
 			DirectoryHandle = INVALID_HANDLE_VALUE;
diff --git a/Engine/Source/Developer/StandaloneRenderer/Private/StandaloneRenderer.cpp b/Engine/Source/Developer/StandaloneRenderer/Private/StandaloneRenderer.cpp
--- a/Engine/Source/Developer/StandaloneRenderer/Private/StandaloneRenderer.cpp
+++ b/Engine/Source/Developer/StandaloneRenderer/Private/StandaloneRenderer.cpp
@@ -23,13 +23,13 @@ TSharedRef<FSlateRenderer> GetStandardStandaloneRenderer()
 {
 	// create a standalone renderer object
 	TSharedPtr<FSlateRenderer> Renderer = NULL;
-#if PLATFORM_WINDOWS && ( WINVER > 0x502 )	// D3D11 not supported on Windows XP
+#if PLATFORM_WINDOWS
 	bool bUseOpenGL = FParse::Param( FCommandLine::Get(), TEXT("opengl") );
 	if( bUseOpenGL )
 	{
 #endif
 		Renderer = TSharedPtr<FSlateRenderer>( new FSlateOpenGLRenderer( FCoreStyle::Get() ) );
-#if PLATFORM_WINDOWS && ( WINVER > 0x502 )	// D3D11 not supported on Windows XP
+#if PLATFORM_WINDOWS
 	}
 	else
 	{
diff --git a/Engine/Source/Developer/Windows/WindowsTargetPlatform/Classes/WindowsTargetSettings.h b/Engine/Source/Developer/Windows/WindowsTargetPlatform/Classes/WindowsTargetSettings.h
--- a/Engine/Source/Developer/Windows/WindowsTargetPlatform/Classes/WindowsTargetSettings.h
+++ b/Engine/Source/Developer/Windows/WindowsTargetPlatform/Classes/WindowsTargetSettings.h
@@ -11,12 +11,6 @@
 #include "AudioCompressionSettings.h"
 #include "WindowsTargetSettings.generated.h"
 
-UENUM()
-enum class EMinimumSupportedOS : uint8
-{
-	MSOS_Vista = 0 UMETA(DisplayName = "Windows Vista"),
-};
-
 UENUM()
 enum class ECompilerVersion : uint8
 {
@@ -67,12 +61,6 @@ public:
 	UPROPERTY(EditAnywhere, config, Category="Targeted RHIs", Meta = (DisplayName = "Default RHI"))
 	EDefaultGraphicsRHI DefaultGraphicsRHI;
 
-	/**
-	 * Determine the minimum supported 
-	 */
-	UPROPERTY(EditAnywhere, config, Category="OS Info", Meta=(DisplayName = "Minimum OS Version"))
-	EMinimumSupportedOS MinimumOSVersion;
-
 	/** Sample rate to run the audio mixer with. */
 	UPROPERTY(config, EditAnywhere, Category = "Audio", Meta = (DisplayName = "Audio Mixer Sample Rate"))
 	int32 AudioSampleRate;
diff --git a/Engine/Source/Developer/Windows/WindowsTargetPlatform/Private/WindowsTargetPlatformClasses.cpp b/Engine/Source/Developer/Windows/WindowsTargetPlatform/Private/WindowsTargetPlatformClasses.cpp
--- a/Engine/Source/Developer/Windows/WindowsTargetPlatform/Private/WindowsTargetPlatformClasses.cpp
+++ b/Engine/Source/Developer/Windows/WindowsTargetPlatform/Private/WindowsTargetPlatformClasses.cpp
@@ -8,8 +8,6 @@
 UWindowsTargetSettings::UWindowsTargetSettings( const FObjectInitializer& ObjectInitializer )
 	: Super(ObjectInitializer)
 {
-	MinimumOSVersion = EMinimumSupportedOS::MSOS_Vista;
-
 	// Default windows settings
 	AudioSampleRate = 48000;
 	AudioCallbackBufferFrameSize = 1024;
diff --git a/Engine/Source/Editor/DetailCustomizations/Private/WindowsTargetSettingsDetails.cpp b/Engine/Source/Editor/DetailCustomizations/Private/WindowsTargetSettingsDetails.cpp
--- a/Engine/Source/Editor/DetailCustomizations/Private/WindowsTargetSettingsDetails.cpp
+++ b/Engine/Source/Editor/DetailCustomizations/Private/WindowsTargetSettingsDetails.cpp
@@ -167,15 +167,6 @@ void FWindowsTargetSettingsDetails::CustomizeDetails( IDetailLayoutBuilder& Deta
 	TargetShaderFormatsDetails = MakeShareable(new FShaderFormatsPropertyDetails(&DetailBuilder));
 	TargetShaderFormatsDetails->CreateTargetShaderFormatsPropertyView(TargetPlatform, GetFriendlyNameFromWindowsRHIName);
 
-	TSharedRef<IPropertyHandle> MinOSProperty = DetailBuilder.GetProperty("MinimumOSVersion");
-	IDetailCategoryBuilder& OSInfoCategory = DetailBuilder.EditCategory(TEXT("OS Info"));
-	
-	// Setup edit condition and tool tip of Min OS property. Determined by whether the engine is installed or not.
-	bool bIsMinOSSelectionAvailable = FApp::IsEngineInstalled() == false;
-	IDetailPropertyRow& MinOSRow = OSInfoCategory.AddProperty(MinOSProperty);
-	MinOSRow.IsEnabled(bIsMinOSSelectionAvailable);
-	MinOSRow.ToolTip(bIsMinOSSelectionAvailable ? MinOSProperty->GetToolTipText() : WindowsTargetSettingsDetailsConstants::DisabledTip);
-
 	// Next add the splash image customization
 	const FText EditorSplashDesc(LOCTEXT("EditorSplashLabel", "Editor Splash"));
 	IDetailCategoryBuilder& SplashCategoryBuilder = DetailBuilder.EditCategory(TEXT("Splash"));
diff --git a/Engine/Source/Programs/UnrealBuildTool/Platform/Windows/UEBuildWindows.cs b/Engine/Source/Programs/UnrealBuildTool/Platform/Windows/UEBuildWindows.cs
--- a/Engine/Source/Programs/UnrealBuildTool/Platform/Windows/UEBuildWindows.cs
+++ b/Engine/Source/Programs/UnrealBuildTool/Platform/Windows/UEBuildWindows.cs
@@ -1949,7 +1949,7 @@ namespace UnrealBuildTool
 			LinkEnvironment.SystemLibraries.Add("iphlpapi.lib");
 			LinkEnvironment.SystemLibraries.Add("setupapi.lib"); //  Required for access monitor device enumeration
 
-			// Windows Vista/7 Desktop Windows Manager API for Slate Windows Compliance
+			// Windows 7 Desktop Windows Manager API for Slate Windows Compliance
 			LinkEnvironment.SystemLibraries.Add("dwmapi.lib");
 
 			// IME
diff --git a/Engine/Source/Runtime/AVEncoder/Private/Microsoft/AVEncoderMicrosoftCommon.h b/Engine/Source/Runtime/AVEncoder/Private/Microsoft/AVEncoderMicrosoftCommon.h
--- a/Engine/Source/Runtime/AVEncoder/Private/Microsoft/AVEncoderMicrosoftCommon.h
+++ b/Engine/Source/Runtime/AVEncoder/Private/Microsoft/AVEncoderMicrosoftCommon.h
@@ -54,7 +54,7 @@
 
 #endif
 
-#define WMFMEDIA_SUPPORTED_PLATFORM (PLATFORM_WINDOWS && (WINVER >= 0x0600 /*Vista*/) && !UE_SERVER)
+#define WMFMEDIA_SUPPORTED_PLATFORM (PLATFORM_WINDOWS && !UE_SERVER)
 
 namespace AVEncoder
 {
diff --git a/Engine/Source/Runtime/ApplicationCore/Private/Windows/WindowsApplication.cpp b/Engine/Source/Runtime/ApplicationCore/Private/Windows/WindowsApplication.cpp
--- a/Engine/Source/Runtime/ApplicationCore/Private/Windows/WindowsApplication.cpp
+++ b/Engine/Source/Runtime/ApplicationCore/Private/Windows/WindowsApplication.cpp
@@ -48,15 +48,6 @@ THIRD_PARTY_INCLUDES_END
 #pragma push_macro("IsMaximized")
 #undef IsMaximized
 
-// This might not be defined by Windows when maintaining backwards-compatibility to pre-Vista builds
-#ifndef WM_MOUSEHWHEEL
-#define WM_MOUSEHWHEEL                  0x020E
-#endif
-
-#ifndef WM_DPICHANGED
-#define WM_DPICHANGED                   0x02E0
-#endif
-
 DEFINE_LOG_CATEGORY(LogWindowsDesktop);
 
 /**
@@ -222,8 +213,6 @@ FWindowsApplication::FWindowsApplication( const HINSTANCE HInstance, const HICON
 		bForceNoGamepads = true;
 	}
 
-#if WINVER >= 0x0601
-
 	if (FParse::Param(FCommandLine::Get(), TEXT("FilterLowLevelMouse")))
 	{
 		//Only check fake mouse inputs in game or digitizer devices will not work in the editor.
@@ -232,8 +221,6 @@ FWindowsApplication::FWindowsApplication( const HINSTANCE HInstance, const HICON
 			LowLevelMouseFilterHook = ::SetWindowsHookEx(WH_MOUSE_LL, HandleLowLevelMouseFilterHook, NULL, 0);
 		}
 	}
-
-#endif
 }
 
 void FWindowsApplication::AllowAccessibilityShortcutKeys(const bool bAllowKeys)
@@ -978,15 +965,11 @@ int32 FWindowsApplication::ProcessMessage( HWND hwnd, uint32 msg, WPARAM wParam,
 			}
 		}
 
-#if WINVER >= 0x0601
-
 		//Only check fake mouse inputs in game or digitizer devices will not work in the editor.
 		if (FApp::IsGame() && IsFakeMouseInputMessage(msg))
 		{
 			return 0;
 		}
-	
-#endif
 
 		switch(msg)
 		{
@@ -1061,9 +1044,7 @@ int32 FWindowsApplication::ProcessMessage( HWND hwnd, uint32 msg, WPARAM wParam,
 		case WM_NCMOUSEMOVE:
 		case WM_MOUSEMOVE:
 		case WM_MOUSEWHEEL:
-#if WINVER >= 0x0601
 		case WM_TOUCH:
-#endif
 			{
 				DeferMessage( CurrentNativeEventWindowPtr, hwnd, msg, wParam, lParam );
 				// Handled
@@ -1594,13 +1575,11 @@ int32 FWindowsApplication::ProcessMessage( HWND hwnd, uint32 msg, WPARAM wParam,
 			}
 			break;
 			
-#if WINVER > 0x502
 		case WM_DWMCOMPOSITIONCHANGED:
 			{
 				DeferMessage( CurrentNativeEventWindowPtr, hwnd, msg, wParam, lParam );
 			}
 			break;
-#endif
 
 			// Window focus and activation
 		case WM_MOUSEACTIVATE:
@@ -2278,7 +2257,6 @@ int32 FWindowsApplication::ProcessDeferredMessage( const FDeferredWindowsMessage
 			}
 			break;
 
-#if WINVER >= 0x0601
 		case WM_TOUCH:
 			{
 				UINT InputCount = LOWORD( wParam );
@@ -2338,7 +2316,6 @@ int32 FWindowsApplication::ProcessDeferredMessage( const FDeferredWindowsMessage
 				}
 				break;
 			}
-#endif
 
 			// Window focus and activation
 		case WM_MOUSEACTIVATE:
@@ -2500,13 +2477,11 @@ int32 FWindowsApplication::ProcessDeferredMessage( const FDeferredWindowsMessage
 			}
 			break;
 
-#if WINVER > 0x502
 		case WM_DWMCOMPOSITIONCHANGED:
 			{
 				CurrentNativeEventWindowPtr->OnTransparencySupportChanged(GetWindowTransparencySupport());
 			}
 			break;
-#endif
 
 		case WM_DPICHANGED:
 			{
@@ -2634,7 +2609,7 @@ bool FWindowsApplication::IsInputMessage( uint32 msg )
 
 bool FWindowsApplication::IsFakeMouseInputMessage(uint32 msg)
 {
-	const bool bShouldPrevent = IsWindowsVistaOrGreater() || bPreventDuplicateMouseEventsForTouch;
+	const bool bShouldPrevent = !!bPreventDuplicateMouseEventsForTouch;
 
 	if (bShouldPrevent && IsMouseInputMessage(msg))
 	{
@@ -3024,7 +2999,6 @@ void FWindowsApplication::QueryConnectedMice()
 	bIsMouseAttached = MouseCount > 0;
 }
 
-#if WINVER >= 0x0601
 uint32 FWindowsApplication::GetTouchIndexForID( int32 TouchID )
 {
 	for (int i = 0; i < TouchIDs.Num(); i++)
@@ -3049,7 +3023,6 @@ uint32 FWindowsApplication::GetFirstFreeTouchIndex()
 
 	return TouchIDs.Add(TOptional<int32>());
 }
-#endif
 
 void FTaskbarList::Initialize()
 {
diff --git a/Engine/Source/Runtime/ApplicationCore/Private/Windows/WindowsPlatformSurvey.cpp b/Engine/Source/Runtime/ApplicationCore/Private/Windows/WindowsPlatformSurvey.cpp
--- a/Engine/Source/Runtime/ApplicationCore/Private/Windows/WindowsPlatformSurvey.cpp
+++ b/Engine/Source/Runtime/ApplicationCore/Private/Windows/WindowsPlatformSurvey.cpp
@@ -51,9 +51,6 @@ void WriteFStringToResults(TCHAR* OutBuffer, const FString& InString);
 
 bool FWindowsPlatformSurvey::GetSurveyResults( FHardwareSurveyResults& OutResults, bool bWait )
 {
-	// Check that we're running on Vista or newer (version 6.0+).
-	bool bIsVistaOrNewer = FPlatformMisc::VerifyWindowsVersion(6, 0);
-
 	FMemory::Memset(&OutResults, 0, sizeof(FHardwareSurveyResults));
 	WriteFStringToResults(OutResults.Platform, TEXT("Windows"));
 	
@@ -201,126 +198,116 @@ bool FWindowsPlatformSurvey::GetSurveyResults( FHardwareSurveyResults& OutResult
 	}
 
 #if USING_WINSAT_API
-	// Use Windows System Assessment Tool?
-	if (bIsVistaOrNewer)
+	// Get an instance to the most recent formal WinSAT assessmenet.
+	IQueryRecentWinSATAssessment* Assessment;
+	HRESULT COMResult = CoCreateInstance(
+		__uuidof(CQueryWinSAT),
+		NULL,
+		CLSCTX_INPROC_SERVER,
+		__uuidof(IQueryRecentWinSATAssessment),
+		(void**)&Assessment);
+
+	if (FAILED(COMResult))
 	{
-		// Get an instance to the most recent formal WinSAT assessmenet.
-		IQueryRecentWinSATAssessment* Assessment;
-		HRESULT COMResult = CoCreateInstance(
-			__uuidof(CQueryWinSAT),
-			NULL,
-			CLSCTX_INPROC_SERVER,
-			__uuidof(IQueryRecentWinSATAssessment),
-			(void**)&Assessment);
+		UE_LOG(LogWindows, Warning, TEXT("FWindowsPlatformSurvey::TickSurveyHardware() failed to get query interface from WinSAT API") );
+		OutResults.ErrorCount++;
+		WriteFStringToResults(OutResults.LastSurveyError, TEXT("CoCreateInstance() failed to get WinSAT"));
+		WriteFStringToResults(OutResults.LastSurveyErrorDetail, FString::Printf(TEXT("HRESULT: 0x%0x"), COMResult));
 
+	}
+	else
+	{
+		// Get the summary information for the WinSAT assessment.
+		IProvideWinSATResultsInfo* WinSATResults = NULL;
+		COMResult = Assessment->get_Info(&WinSATResults);
 		if (FAILED(COMResult))
 		{
-			UE_LOG(LogWindows, Warning, TEXT("FWindowsPlatformSurvey::TickSurveyHardware() failed to get query interface from WinSAT API") );
+			UE_LOG(LogWindows, Error, TEXT("FWindowsPlatformSurvey::TickSurveyHardware() failed to get assessment results from WinSAT API") );
 			OutResults.ErrorCount++;
-			WriteFStringToResults(OutResults.LastSurveyError, TEXT("CoCreateInstance() failed to get WinSAT"));
+			WriteFStringToResults(OutResults.LastSurveyError, TEXT("get_Info() failed to get WinSAT assessment results"));
 			WriteFStringToResults(OutResults.LastSurveyErrorDetail, FString::Printf(TEXT("HRESULT: 0x%0x"), COMResult));
 
 		}
 		else
 		{
-			// Get the summary information for the WinSAT assessment.
-			IProvideWinSATResultsInfo* WinSATResults = NULL;
-			COMResult = Assessment->get_Info(&WinSATResults);
+			// Get the state of the assessment.
+			WINSAT_ASSESSMENT_STATE WinSATState;
+			COMResult = WinSATResults->get_AssessmentState(&WinSATState);
 			if (FAILED(COMResult))
 			{
-				UE_LOG(LogWindows, Error, TEXT("FWindowsPlatformSurvey::TickSurveyHardware() failed to get assessment results from WinSAT API") );
+				UE_LOG(LogWindows, Error, TEXT("FWindowsPlatformSurvey::TickSurveyHardware() failed to get assessment state from WinSAT API") );
 				OutResults.ErrorCount++;
-				WriteFStringToResults(OutResults.LastSurveyError, TEXT("get_Info() failed to get WinSAT assessment results"));
+				WriteFStringToResults(OutResults.LastSurveyError, TEXT("get_AssessmentState() failed to get WinSAT assessment state"));
 				WriteFStringToResults(OutResults.LastSurveyErrorDetail, FString::Printf(TEXT("HRESULT: 0x%0x"), COMResult));
-
 			}
 			else
 			{
-				// Get the state of the assessment.
-				WINSAT_ASSESSMENT_STATE WinSATState;
-				COMResult = WinSATResults->get_AssessmentState(&WinSATState);
-				if (FAILED(COMResult))
+				// Examine the assessment state
+				bool bAssessmentAvailable = false;
+				switch(WinSATState)
 				{
-					UE_LOG(LogWindows, Error, TEXT("FWindowsPlatformSurvey::TickSurveyHardware() failed to get assessment state from WinSAT API") );
+				case WINSAT_ASSESSMENT_STATE_VALID:
+					bAssessmentAvailable = true;
+					break;
+
+				case WINSAT_ASSESSMENT_STATE_INCOHERENT_WITH_HARDWARE:
+					UE_LOG(LogWindows, Log, TEXT("FWindowsPlatformSurvey::TickSurveyHardware() WinSAT assessment state is out-of-date. Unable to examine some hardware metrics. Run the Windows Experience Index Assessment.") );
 					OutResults.ErrorCount++;
-					WriteFStringToResults(OutResults.LastSurveyError, TEXT("get_AssessmentState() failed to get WinSAT assessment state"));
-					WriteFStringToResults(OutResults.LastSurveyErrorDetail, FString::Printf(TEXT("HRESULT: 0x%0x"), COMResult));
-				}
-				else
-				{
-					// Examine the assessment state
-					bool bAssessmentAvailable = false;
-					switch(WinSATState)
-					{
-					case WINSAT_ASSESSMENT_STATE_VALID:
-						bAssessmentAvailable = true;
-						break;
+					WriteFStringToResults(OutResults.LastSurveyError, TEXT("WinSAT assessment out-of-date. Using old results."));
+					WriteFStringToResults(OutResults.LastSurveyErrorDetail, TEXT(""));
+					bAssessmentAvailable = true;
+					break;
 
-					case WINSAT_ASSESSMENT_STATE_INCOHERENT_WITH_HARDWARE:
-						UE_LOG(LogWindows, Log, TEXT("FWindowsPlatformSurvey::TickSurveyHardware() WinSAT assessment state is out-of-date. Unable to examine some hardware metrics. Run the Windows Experience Index Assessment.") );
-						OutResults.ErrorCount++;
-						WriteFStringToResults(OutResults.LastSurveyError, TEXT("WinSAT assessment out-of-date. Using old results."));
-						WriteFStringToResults(OutResults.LastSurveyErrorDetail, TEXT(""));
-						bAssessmentAvailable = true;
-						break;
+				case WINSAT_ASSESSMENT_STATE_NOT_AVAILABLE:
+					UE_LOG(LogWindows, Log, TEXT("FWindowsPlatformSurvey::TickSurveyHardware() WinSAT assessment unavailable. Unable to examine some hardware metrics. Run the Windows Experience Index Assessment.") );
+					OutResults.ErrorCount++;
+					WriteFStringToResults(OutResults.LastSurveyError, TEXT("WinSAT assessment unavailable. User hasn't run Windows Experience Index Assessment."));
+					WriteFStringToResults(OutResults.LastSurveyErrorDetail, TEXT(""));
+					break;
 
-					case WINSAT_ASSESSMENT_STATE_NOT_AVAILABLE:
-						UE_LOG(LogWindows, Log, TEXT("FWindowsPlatformSurvey::TickSurveyHardware() WinSAT assessment unavailable. Unable to examine some hardware metrics. Run the Windows Experience Index Assessment.") );
-						OutResults.ErrorCount++;
-						WriteFStringToResults(OutResults.LastSurveyError, TEXT("WinSAT assessment unavailable. User hasn't run Windows Experience Index Assessment."));
-						WriteFStringToResults(OutResults.LastSurveyErrorDetail, TEXT(""));
-						break;
+				default:
+					UE_LOG(LogWindows, Warning, TEXT("FWindowsPlatformSurvey::TickSurveyHardware() WinSAT assessment data was invalid.") );
+					OutResults.ErrorCount++;
+					WriteFStringToResults(OutResults.LastSurveyError, TEXT("WinSAT assessment state unknown"));
+					WriteFStringToResults(OutResults.LastSurveyErrorDetail, FString::Printf(TEXT("WinSATState: %d"), (int32)WinSATState));
+				}
 
-					default:
-						UE_LOG(LogWindows, Warning, TEXT("FWindowsPlatformSurvey::TickSurveyHardware() WinSAT assessment data was invalid.") );
-						OutResults.ErrorCount++;
-						WriteFStringToResults(OutResults.LastSurveyError, TEXT("WinSAT assessment state unknown"));
-						WriteFStringToResults(OutResults.LastSurveyErrorDetail, FString::Printf(TEXT("WinSATState: %d"), (int32)WinSATState));
+				// Get the index scores from the results
+				if (bAssessmentAvailable)
+				{
+					if (!GetSubComponentIndex(WinSATResults, OutResults, WINSAT_ASSESSMENT_MEMORY, OutResults.RAMPerformanceIndex))
+					{
+						UE_LOG(LogWindows, Error, TEXT("FWindowsPlatformSurvey::TickSurveyHardware() failed to get MEMORY score from WinSAT API.") );
 					}
 
-					// Get the index scores from the results
-					if (bAssessmentAvailable)
+					if (!GetSubComponentIndex(WinSATResults, OutResults, WINSAT_ASSESSMENT_CPU, OutResults.CPUPerformanceIndex))
 					{
-						if (!GetSubComponentIndex(WinSATResults, OutResults, WINSAT_ASSESSMENT_MEMORY, OutResults.RAMPerformanceIndex))
-						{
-							UE_LOG(LogWindows, Error, TEXT("FWindowsPlatformSurvey::TickSurveyHardware() failed to get MEMORY score from WinSAT API.") );
-						}
-
-						if (!GetSubComponentIndex(WinSATResults, OutResults, WINSAT_ASSESSMENT_CPU, OutResults.CPUPerformanceIndex))
-						{
-							UE_LOG(LogWindows, Error, TEXT("FWindowsPlatformSurvey::TickSurveyHardware() failed to get CPU score from WinSAT API.") );
-						}
-
-						float GPU3DScoreIndex = 0.0f;
-						if (!GetSubComponentIndex(WinSATResults, OutResults, WINSAT_ASSESSMENT_D3D, GPU3DScoreIndex))
-						{
-							UE_LOG(LogWindows, Error, TEXT("FWindowsPlatformSurvey::TickSurveyHardware() failed to get D3D score from WinSAT API.") );
-						}
+						UE_LOG(LogWindows, Error, TEXT("FWindowsPlatformSurvey::TickSurveyHardware() failed to get CPU score from WinSAT API.") );
+					}
 
-						float GPUDesktopScoreIndex = 0.0f;
-						if (!GetSubComponentIndex(WinSATResults, OutResults, WINSAT_ASSESSMENT_D3D, GPUDesktopScoreIndex))
-						{
-							UE_LOG(LogWindows, Error, TEXT("FWindowsPlatformSurvey::TickSurveyHardware() failed to get GRAPHICS score from WinSAT API.") );
-						}
+					float GPU3DScoreIndex = 0.0f;
+					if (!GetSubComponentIndex(WinSATResults, OutResults, WINSAT_ASSESSMENT_D3D, GPU3DScoreIndex))
+					{
+						UE_LOG(LogWindows, Error, TEXT("FWindowsPlatformSurvey::TickSurveyHardware() failed to get D3D score from WinSAT API.") );
+					}
 
-						OutResults.GPUPerformanceIndex = 0.5f * (GPU3DScoreIndex + GPUDesktopScoreIndex);
+					float GPUDesktopScoreIndex = 0.0f;
+					if (!GetSubComponentIndex(WinSATResults, OutResults, WINSAT_ASSESSMENT_D3D, GPUDesktopScoreIndex))
+					{
+						UE_LOG(LogWindows, Error, TEXT("FWindowsPlatformSurvey::TickSurveyHardware() failed to get GRAPHICS score from WinSAT API.") );
 					}
+
+					OutResults.GPUPerformanceIndex = 0.5f * (GPU3DScoreIndex + GPUDesktopScoreIndex);
 				}
-				
-				WinSATResults->Release();
 			}
-		}
-
-		if (Assessment)
-		{
-			Assessment->Release();
+				
+			WinSATResults->Release();
 		}
 	}
-	else
+
+	if (Assessment)
 	{
-		OutResults.ErrorCount++;
-		WriteFStringToResults(OutResults.LastSurveyError, TEXT("WIE failed. Not supported on this version of Windows."));
-		WriteFStringToResults(OutResults.LastSurveyErrorDetail, TEXT(""));
+		Assessment->Release();
 	}
 #endif	// #if USING_WINSAT_API
 
diff --git a/Engine/Source/Runtime/ApplicationCore/Private/Windows/WindowsTextInputMethodSystem.cpp b/Engine/Source/Runtime/ApplicationCore/Private/Windows/WindowsTextInputMethodSystem.cpp
--- a/Engine/Source/Runtime/ApplicationCore/Private/Windows/WindowsTextInputMethodSystem.cpp
+++ b/Engine/Source/Runtime/ApplicationCore/Private/Windows/WindowsTextInputMethodSystem.cpp
@@ -465,43 +465,17 @@ bool FWindowsTextInputMethodSystem::InitializeTSF()
 #pragma warning(pop)
 #endif
 
-		const DWORD WindowsMajorVersion = LOBYTE(LOWORD(WindowsVersion));
-		const DWORD WindowsMinorVersion = HIBYTE(LOWORD(WindowsVersion));
-
-		static const DWORD WindowsVistaMajorVersion = 6;
-		static const DWORD WindowsVistaMinorVersion = 0;
-
-		// Install profile notification sink for versions of Windows Vista and after.
-		if(WindowsMajorVersion > WindowsVistaMajorVersion || (WindowsMajorVersion == WindowsVistaMajorVersion && WindowsMinorVersion >= WindowsVistaMinorVersion))
-		{
-			Result = TSFSource->AdviseSink(IID_ITfInputProcessorProfileActivationSink, static_cast<ITfInputProcessorProfileActivationSink*>(TSFActivationProxy), &(TSFActivationProxy->TSFProfileCookie));
-			if(FAILED(Result))
-			{
-				TCHAR ErrorMsg[1024];
-				FPlatformMisc::GetSystemErrorMessage(ErrorMsg, 1024, Result);
-				UE_LOG(LogWindowsTextInputMethodSystem, Error, TEXT("Initialization failed while advising the profile notification sink to the TSF source. %s (0x%08x)"), ErrorMsg, Result);
-				TSFInputProcessorProfiles.Reset();
-				TSFInputProcessorProfileManager.Reset();
-				TSFThreadManager.Reset();
-				TSFActivationProxy.Reset();
-				return false;
-			}
-		}
-		// Install language notification sink for versions before Windows Vista.
-		else
+		Result = TSFSource->AdviseSink(IID_ITfInputProcessorProfileActivationSink, static_cast<ITfInputProcessorProfileActivationSink*>(TSFActivationProxy), &(TSFActivationProxy->TSFProfileCookie));
+		if(FAILED(Result))
 		{
-			Result = TSFSource->AdviseSink(IID_ITfActiveLanguageProfileNotifySink, static_cast<ITfActiveLanguageProfileNotifySink*>(TSFActivationProxy), &(TSFActivationProxy->TSFLanguageCookie));
-			if(FAILED(Result))
-			{
-				TCHAR ErrorMsg[1024];
-				FPlatformMisc::GetSystemErrorMessage(ErrorMsg, 1024, Result);
-				UE_LOG(LogWindowsTextInputMethodSystem, Error, TEXT("Initialization failed while advising the language notification sink to the TSF source. %s (0x%08x)"), ErrorMsg, Result);
-				TSFInputProcessorProfiles.Reset();
-				TSFInputProcessorProfileManager.Reset();
-				TSFThreadManager.Reset();
-				TSFActivationProxy.Reset();
-				return false;
-			}
+			TCHAR ErrorMsg[1024];
+			FPlatformMisc::GetSystemErrorMessage(ErrorMsg, 1024, Result);
+			UE_LOG(LogWindowsTextInputMethodSystem, Error, TEXT("Initialization failed while advising the profile notification sink to the TSF source. %s (0x%08x)"), ErrorMsg, Result);
+			TSFInputProcessorProfiles.Reset();
+			TSFInputProcessorProfileManager.Reset();
+			TSFThreadManager.Reset();
+			TSFActivationProxy.Reset();
+			return false;
 		}
 	}
 
diff --git a/Engine/Source/Runtime/ApplicationCore/Private/Windows/WindowsWindow.cpp b/Engine/Source/Runtime/ApplicationCore/Private/Windows/WindowsWindow.cpp
--- a/Engine/Source/Runtime/ApplicationCore/Private/Windows/WindowsWindow.cpp
+++ b/Engine/Source/Runtime/ApplicationCore/Private/Windows/WindowsWindow.cpp
@@ -15,9 +15,7 @@
 #include "HAL/PlatformProcess.h"
 
 #include "Windows/AllowWindowsPlatformTypes.h"
-#if WINVER > 0x502	// Windows Vista or better required for DWM
-	#include <dwmapi.h>
-#endif
+#include <dwmapi.h>
 #include <ShlObj.h>
 #include "Windows/HideWindowsPlatformTypes.h"
 
@@ -197,14 +195,11 @@ void FWindowsWindow::Initialize( FWindowsApplication* const Application, const T
 		return;
 	}
 
-
-#if WINVER >= 0x0601
 	if ( RegisterTouchWindow( HWnd, 0 ) == false )
 	{
 		uint32 Error = GetLastError();
 		UE_LOG(LogWindows, Warning, TEXT("Register touch input failed!"));
 	}
-#endif
 
 	bool bDisableTouchFeedback;
 	GConfig->GetBool(TEXT("WindowsApplication.Accessibility"), TEXT("DisableTouchFeedback"), bDisableTouchFeedback, GEngineIni);
@@ -228,7 +223,6 @@ void FWindowsWindow::Initialize( FWindowsApplication* const Application, const T
 		SetOpacity( Definition->Opacity );
 	}
 
-#if WINVER > 0x502	// Windows Vista or better required for DWM
 	// Disable DWM Rendering and Nonclient Area painting if not showing the os window border
 	// This prevents the standard windows frame from ever being drawn
 	if( !Definition->HasOSWindowBorder )
@@ -248,8 +242,6 @@ void FWindowsWindow::Initialize( FWindowsApplication* const Application, const T
 	#endif
 	}
 
-#endif	// WINVER
-
 	// No region for non regular windows or windows displaying the os window border
 	if ( IsRegularWindow() && !Definition->HasOSWindowBorder )
 	{
diff --git a/Engine/Source/Runtime/ApplicationCore/Public/GenericPlatform/GenericPlatformApplicationMisc.h b/Engine/Source/Runtime/ApplicationCore/Public/GenericPlatform/GenericPlatformApplicationMisc.h
--- a/Engine/Source/Runtime/ApplicationCore/Public/GenericPlatform/GenericPlatformApplicationMisc.h
+++ b/Engine/Source/Runtime/ApplicationCore/Public/GenericPlatform/GenericPlatformApplicationMisc.h
@@ -83,8 +83,8 @@ struct APPLICATIONCORE_API FGenericPlatformApplicationMisc
 	}
 
 	/**
-	 * Prevents screen-saver from kicking in by moving the mouse by 0 pixels. This works even on
-	 * Vista in the presence of a group policy for password protected screen saver.
+	 * Prevents screen-saver from kicking in by moving the mouse by 0 pixels. This works even
+	 * in the presence of a group policy for password protected screen saver.
 	 */
 	static void PreventScreenSaver()
 	{
diff --git a/Engine/Source/Runtime/ApplicationCore/Public/Windows/WindowsApplication.h b/Engine/Source/Runtime/ApplicationCore/Public/Windows/WindowsApplication.h
--- a/Engine/Source/Runtime/ApplicationCore/Public/Windows/WindowsApplication.h
+++ b/Engine/Source/Runtime/ApplicationCore/Public/Windows/WindowsApplication.h
@@ -443,13 +443,11 @@ private:
 	/** Queries and caches the number of connected mouse devices. */
 	void QueryConnectedMice();
 
-#if WINVER >= 0x0601
 	/** Gets the touch index for a given windows touch ID. */
 	uint32 GetTouchIndexForID(int32 TouchID);
 
 	/** Searches for a free touch index. */
 	uint32 GetFirstFreeTouchIndex();
-#endif
 
 	/** Helper function to update the cached states of all modifier keys */
 	void UpdateAllModifierKeyStates();
@@ -529,10 +527,8 @@ private:
 	TOGGLEKEYS							StartupToggleKeys;
 	FILTERKEYS							StartupFilterKeys;
 
-#if WINVER >= 0x0601
 	/** Maps touch indexes to windows touch IDs. */
 	TArray<TOptional<int32>> TouchIDs;
-#endif
 
 	bool bSimulatingHighPrecisionMouseInputForRDP;
 
diff --git a/Engine/Source/Runtime/Core/Private/Misc/CommandLine.cpp b/Engine/Source/Runtime/Core/Private/Misc/CommandLine.cpp
--- a/Engine/Source/Runtime/Core/Private/Misc/CommandLine.cpp
+++ b/Engine/Source/Runtime/Core/Private/Misc/CommandLine.cpp
@@ -276,7 +276,7 @@ const TCHAR* FCommandLine::RemoveExeName(const TCHAR* InCmdLine)
 	{
 		InCmdLine++;
 	}
-	// skip over any spaces at the start, which Vista likes to toss in multiple
+	// skip over any spaces at the start
 	while (*InCmdLine == ' ')
 	{
 		InCmdLine++;
diff --git a/Engine/Source/Runtime/Core/Private/Windows/WindowsPlatformCrashContext.cpp b/Engine/Source/Runtime/Core/Private/Windows/WindowsPlatformCrashContext.cpp
--- a/Engine/Source/Runtime/Core/Private/Windows/WindowsPlatformCrashContext.cpp
+++ b/Engine/Source/Runtime/Core/Private/Windows/WindowsPlatformCrashContext.cpp
@@ -1422,13 +1422,11 @@ private:
 			CrashContext.CapturePortableCallStack(ErrorProgramCounter, ContextWrapper);
 
 			// First launch the crash reporter client.
-#if WINVER > 0x502	// Windows Error Reporting is not supported on Windows XP
 			if (GUseCrashReportClient)
 			{
 				ReportCrashUsingCrashReportClient(CrashContext, ExceptionInfo, EErrorReportUI::ShowDialog);
 			}
 			else
-#endif		// WINVER
 			{
 				CrashContext.SerializeContentToBuffer();
 				WriteMinidump(GetCurrentProcess(), GetCurrentThreadId(), CrashContext, MiniDumpFilenameW, ExceptionInfo);
@@ -1594,7 +1592,6 @@ int32 ReportCrash( LPEXCEPTION_POINTERS ExceptionInfo )
 static FCriticalSection EnsureLock;
 static bool bReentranceGuard = false;
 
-#if WINVER > 0x502	// Windows Error Reporting is not supported on Windows XP
 /**
  * A wrapper for ReportCrashUsingCrashReportClient that creates a new ensure crash context
  */
@@ -1606,11 +1603,9 @@ int32 ReportContinuableEventUsingCrashReportClient(ECrashContextType InType, EXC
 	return EXCEPTION_EXECUTE_HANDLER;
 #endif
 }
-#endif
 
 static void ReportEventOnCallingThread(ECrashContextType InType, const TCHAR* ErrorMessage, void* ProgramCounter)
 {
-#if WINVER > 0x502	// Windows Error Reporting is not supported on Windows XP
 #if !PLATFORM_SEH_EXCEPTIONS_DISABLED
 	__try
 #endif
@@ -1625,7 +1620,6 @@ static void ReportEventOnCallingThread(ECrashContextType InType, const TCHAR* Er
 	{
 	}
 #endif
-#endif	// WINVER
 }
 
 static void ReportEvent(ECrashContextType InType, const TCHAR* ErrorMessage, uint32 InThreadId, void* ProgramCounter)
diff --git a/Engine/Source/Runtime/Core/Private/Windows/WindowsPlatformMisc.cpp b/Engine/Source/Runtime/Core/Private/Windows/WindowsPlatformMisc.cpp
--- a/Engine/Source/Runtime/Core/Private/Windows/WindowsPlatformMisc.cpp
+++ b/Engine/Source/Runtime/Core/Private/Windows/WindowsPlatformMisc.cpp
@@ -129,82 +129,9 @@ int32 GetOSVersionsHelper( TCHAR* OutOSVersionLabel, int32 OSVersionLabelLength,
 
 		switch (OsVersionInfo.dwMajorVersion)
 		{
-		case 5:
-			switch (OsVersionInfo.dwMinorVersion)
-			{
-			case 0:
-				OSVersionLabel = TEXT("Windows 2000");
-				if (OsVersionInfo.wProductType == VER_NT_WORKSTATION)
-				{
-					OSSubVersionLabel = TEXT("Professional");
-				}
-				else
-				{
-					if (OsVersionInfo.wSuiteMask & VER_SUITE_DATACENTER)
-					{
-						OSSubVersionLabel = TEXT("Datacenter Server");
-					}
-					else if (OsVersionInfo.wSuiteMask & VER_SUITE_ENTERPRISE)
-					{
-						OSSubVersionLabel = TEXT("Advanced Server");
-					}
-					else
-					{
-						OSSubVersionLabel = TEXT("Server");
-					}
-				}
-				break;
-			case 1:
-				OSVersionLabel = TEXT("Windows XP");
-				if (OsVersionInfo.wSuiteMask & VER_SUITE_PERSONAL)
-				{
-					OSSubVersionLabel = TEXT("Home Edition");
-				}
-				else
-				{
-					OSSubVersionLabel = TEXT("Professional");
-				}
-				break;
-			case 2:
-				if (GetSystemMetrics(SM_SERVERR2))
-				{
-					OSVersionLabel = TEXT("Windows Server 2003 R2");
-				}
-				else if (OsVersionInfo.wSuiteMask & VER_SUITE_STORAGE_SERVER)
-				{
-					OSVersionLabel = TEXT("Windows Storage Server 2003");
-				}
-				else if (OsVersionInfo.wSuiteMask & VER_SUITE_WH_SERVER)
-				{
-					OSVersionLabel = TEXT("Windows Home Server");
-				}
-				else if (OsVersionInfo.wProductType == VER_NT_WORKSTATION && SystemInfo.wProcessorArchitecture == PROCESSOR_ARCHITECTURE_AMD64)
-				{
-					OSVersionLabel = TEXT("Windows XP");
-					OSSubVersionLabel = TEXT("Professional x64 Edition");
-				}
-				else
-				{
-					OSVersionLabel = TEXT("Windows Server 2003");
-				}
-				break;
-			default:
-				ErrorCode |= (int32)FWindowsOSVersionHelper::ERROR_UNKNOWNVERSION;
-			}
-			break;
 		case 6:
 			switch (OsVersionInfo.dwMinorVersion)
 			{
-			case 0:
-				if (OsVersionInfo.wProductType == VER_NT_WORKSTATION)
-				{
-					OSVersionLabel = TEXT("Windows Vista");
-				}
-				else
-				{
-					OSVersionLabel = TEXT("Windows Server 2008");
-				}
-				break;
 			case 1:
 				if (OsVersionInfo.wProductType == VER_NT_WORKSTATION)
 				{
diff --git a/Engine/Source/Runtime/Core/Private/Windows/WindowsPlatformProcess.cpp b/Engine/Source/Runtime/Core/Private/Windows/WindowsPlatformProcess.cpp
--- a/Engine/Source/Runtime/Core/Private/Windows/WindowsPlatformProcess.cpp
+++ b/Engine/Source/Runtime/Core/Private/Windows/WindowsPlatformProcess.cpp
@@ -679,13 +679,7 @@ FString FWindowsPlatformProcess::GetApplicationName( uint32 ProcessId )
 		int32 InOutSize = ProcessNameBufferSize;
 		static_assert(sizeof(::DWORD) == sizeof(int32), "DWORD size doesn't match int32. Is it the future or the past?");
 
-		if(
-#if WINVER == 0x0502
-		GetProcessImageFileName(ProcessHandle, ProcessNameBuffer, InOutSize)
-#else
-		QueryFullProcessImageName(ProcessHandle, 0, ProcessNameBuffer, (PDWORD)(&InOutSize))
-#endif
-			)
+		if (QueryFullProcessImageName(ProcessHandle, 0, ProcessNameBuffer, (PDWORD)(&InOutSize)))
 		{
 			// TODO no null termination guarantee on GetProcessImageFileName?  it returns size as well, whereas QueryFullProcessImageName just returns non-zero on success
 			Output = ProcessNameBuffer;
diff --git a/Engine/Source/Runtime/Core/Private/Windows/WindowsPlatformStackWalk.cpp b/Engine/Source/Runtime/Core/Private/Windows/WindowsPlatformStackWalk.cpp
--- a/Engine/Source/Runtime/Core/Private/Windows/WindowsPlatformStackWalk.cpp
+++ b/Engine/Source/Runtime/Core/Private/Windows/WindowsPlatformStackWalk.cpp
@@ -64,13 +64,8 @@ static const TCHAR* CrashReporterSettings = TEXT("/Script/UnrealEd.CrashReporter
 
 typedef bool  (WINAPI *TFEnumProcesses)( uint32* lpidProcess, uint32 cb, uint32* cbNeeded);
 typedef bool  (WINAPI *TFEnumProcessModules)(HANDLE hProcess, HMODULE *lphModule, uint32 cb, LPDWORD lpcbNeeded);
-#if WINVER > 0x502
 typedef uint32 (WINAPI *TFGetModuleBaseName)(HANDLE hProcess, HMODULE hModule, LPWSTR lpBaseName, uint32 nSize);
 typedef uint32 (WINAPI *TFGetModuleFileNameEx)(HANDLE hProcess, HMODULE hModule, LPWSTR lpFilename, uint32 nSize);
-#else
-typedef uint32 (WINAPI *TFGetModuleBaseName)(HANDLE hProcess, HMODULE hModule, LPSTR lpBaseName, uint32 nSize);
-typedef uint32 (WINAPI *TFGetModuleFileNameEx)(HANDLE hProcess, HMODULE hModule, LPSTR lpFilename, uint32 nSize);
-#endif
 typedef bool  (WINAPI *TFGetModuleInformation)(HANDLE hProcess, HMODULE hModule, LPMODULEINFO lpmodinfo, uint32 cb);
 
 static TFEnumProcesses			FEnumProcesses;
@@ -237,15 +232,7 @@ static bool GMaxCallstackDepthInitialized = false;
 /** Checks the current OS version and sets up the GMaxCallstackDepth variable. */
 void DetermineMaxCallstackDepth()
 {
-	// Check that we're running on Vista or newer (version 6.0+).
-	if ( FPlatformMisc::VerifyWindowsVersion(6, 0) )
-	{
-		GMaxCallstackDepth = MAX_CALLSTACK_DEPTH;
-	}
-	else
-	{
-		GMaxCallstackDepth = FMath::Min<ULONG>(62,MAX_CALLSTACK_DEPTH);
-	}
+	GMaxCallstackDepth = MAX_CALLSTACK_DEPTH;
 	GMaxCallstackDepthInitialized = true;
 }
 
@@ -628,7 +615,6 @@ bool FWindowsPlatformStackWalk::UploadLocalSymbols()
 {
 	InitStackWalking();
 
-#if WINVER > 0x502
 	// Upload locally compiled files to symbol storage.
 	FString SymbolStorage;
 	if (!GConfig->GetString( CrashReporterSettings, TEXT( "UploadSymbolsPath" ), SymbolStorage, GEditorPerProjectIni ) || SymbolStorage.IsEmpty())
@@ -711,9 +697,6 @@ bool FWindowsPlatformStackWalk::UploadLocalSymbols()
 			}
 		}
 	}
-#else
-	UE_LOG( LogWindows, Log, TEXT( "Symbol server not supported on Windows XP." ) );
-#endif
 	return true;
 }
 
@@ -724,13 +707,8 @@ void LoadSymbolsForModule(HMODULE ModuleHandle, const FString& RemoteStorage)
 	int32 ErrorCode = 0;
 
 	MODULEINFO ModuleInfo = { 0 };
-#if WINVER > 0x502
 	WCHAR ModuleName[FProgramCounterSymbolInfo::MAX_NAME_LENGTH] = { 0 };
 	WCHAR ImageName[FProgramCounterSymbolInfo::MAX_NAME_LENGTH] = { 0 };
-#else
-	ANSICHAR ModuleName[FProgramCounterSymbolInfo::MAX_NAME_LENGTH] = { 0 };
-	ANSICHAR ImageName[FProgramCounterSymbolInfo::MAX_NAME_LENGTH] = { 0 };
-#endif
 #if PLATFORM_64BITS
 	static_assert(sizeof(MODULEINFO) == 24, "Broken alignment for 64bit Windows include.");
 #else
@@ -741,25 +719,15 @@ void LoadSymbolsForModule(HMODULE ModuleHandle, const FString& RemoteStorage)
 	FGetModuleBaseName(ProcessHandle, ModuleHandle, ModuleName, FProgramCounterSymbolInfo::MAX_NAME_LENGTH);
 
 	// Set the search path to find PDBs in the same folder as the DLL.
-#if WINVER > 0x502
 	WCHAR SearchPath[MAX_PATH] = { 0 };
 	WCHAR* FileName = NULL;
 	const auto Result = GetFullPathNameW(ImageName, MAX_PATH, SearchPath, &FileName);
-#else
-	ANSICHAR SearchPath[MAX_PATH] = { 0 };
-	ANSICHAR* FileName = NULL;
-	const auto Result = GetFullPathNameA(ImageName, MAX_PATH, SearchPath, &FileName);
-#endif
 
 	FString SearchPathList;
 	if (Result != 0 && Result < MAX_PATH)
 	{
 		*FileName = 0;
-#if WINVER > 0x502
 		SearchPathList = SearchPath;
-#else
-		SearchPathList = ANSI_TO_TCHAR(SearchPath);
-#endif
 	}
 	if (!RemoteStorage.IsEmpty())
 	{
@@ -770,7 +738,6 @@ void LoadSymbolsForModule(HMODULE ModuleHandle, const FString& RemoteStorage)
 		SearchPathList.Append(RemoteStorage);
 	}
 
-#if WINVER > 0x502
 	SymSetSearchPathW(ProcessHandle, *SearchPathList);
 
 	// Load module.
@@ -785,22 +752,6 @@ void LoadSymbolsForModule(HMODULE ModuleHandle, const FString& RemoteStorage)
 			UE_LOG(LogWindows, Warning, TEXT("SymLoadModuleExW. Error: %d"), ErrorCode);
 		}
 	}
-#else
-	SymSetSearchPath(ProcessHandle, TCHAR_TO_ANSI(*SearchPathList));
-
-	// Load module.
-	const DWORD64 BaseAddress = SymLoadModuleEx(ProcessHandle, ModuleHandle, ImageName, ModuleName, (DWORD64)ModuleInfo.lpBaseOfDll, (uint32)ModuleInfo.SizeOfImage, NULL, 0);
-	if (!BaseAddress)
-	{
-		ErrorCode = GetLastError();
-
-		// If the module is already loaded, the return value is zero and GetLastError returns ERROR_SUCCESS.
-		if (ErrorCode != ERROR_SUCCESS)
-		{
-			UE_LOG(LogWindows, Warning, TEXT("SymLoadModuleEx. Error: %d"), ErrorCode);
-		}
-	}
-#endif
 }
 
 
@@ -883,13 +834,8 @@ int32 FWindowsPlatformStackWalk::GetProcessModuleSignatures(FStackWalkModuleInfo
 	for( int32 ModuleIndex = 0; ModuleHandlePointer[ModuleIndex] && SignatureIndex < ModuleSignaturesSize; ModuleIndex++ )
 	{
 		MODULEINFO ModuleInfo = {0};
-#if WINVER > 0x502
 		WCHAR ModuleName[MAX_PATH] = {0};
 		WCHAR ImageName[MAX_PATH] = {0};
-#else
-		ANSICHAR ModuleName[MAX_PATH] = { 0 };
-		ANSICHAR ImageName[MAX_PATH] = { 0 };
-#endif
 #if PLATFORM_64BITS
 		static_assert(sizeof( MODULEINFO ) == 24, "Broken alignment for 64bit Windows include.");
 #else
@@ -1093,13 +1039,8 @@ bool FWindowsPlatformStackWalk::InitStackWalkingInternal(void* Process, bool bFo
 		// Load dynamically linked PSAPI routines.
 		FEnumProcesses			= (TFEnumProcesses)			FPlatformProcess::GetDllExport( DllHandle,TEXT("EnumProcesses"));
 		FEnumProcessModules		= (TFEnumProcessModules)	FPlatformProcess::GetDllExport( DllHandle,TEXT("EnumProcessModules"));
-#if WINVER > 0x502
 		FGetModuleFileNameEx	= (TFGetModuleFileNameEx)	FPlatformProcess::GetDllExport( DllHandle,TEXT("GetModuleFileNameExW"));
 		FGetModuleBaseName		= (TFGetModuleBaseName)		FPlatformProcess::GetDllExport( DllHandle,TEXT("GetModuleBaseNameW"));
-#else
-		FGetModuleFileNameEx	= (TFGetModuleFileNameEx)	FPlatformProcess::GetDllExport( DllHandle,TEXT("GetModuleFileNameExA"));
-		FGetModuleBaseName		= (TFGetModuleBaseName)		FPlatformProcess::GetDllExport( DllHandle,TEXT("GetModuleBaseNameA"));
-#endif
 		FGetModuleInformation	= (TFGetModuleInformation)	FPlatformProcess::GetDllExport( DllHandle,TEXT("GetModuleInformation"));
 	}
 
@@ -1132,12 +1073,8 @@ bool FWindowsPlatformStackWalk::InitStackWalkingInternal(void* Process, bool bFo
 		SymSetOptions( SymOpts );
 
 		// Initialize the symbol engine.
-#if WINVER > 0x502
 		FString SymbolSearchPath = GetSymbolSearchPath();
 		SymInitializeW( GProcessHandle, SymbolSearchPath.IsEmpty() ? nullptr : *SymbolSearchPath, true );
-#else
-		SymInitialize( GProcessHandle, nullptr, true );
-#endif
 
 		GNeedToRefreshSymbols = false;
 		GStackWalkingInitialized = true;
@@ -1150,7 +1087,6 @@ bool FWindowsPlatformStackWalk::InitStackWalkingInternal(void* Process, bool bFo
 		}
 #endif
 	}
-#if WINVER > 0x502
 	else if (GNeedToRefreshSymbols)
 	{
 		// Refresh and reload symbols
@@ -1166,7 +1102,6 @@ bool FWindowsPlatformStackWalk::InitStackWalkingInternal(void* Process, bool bFo
 			LoadSymbolsForProcessModules( RemoteStorage );
 		}
 	}
-#endif
 
 	return GStackWalkingInitialized;
 }
diff --git a/Engine/Source/Runtime/Core/Public/Windows/WindowsPlatform.h b/Engine/Source/Runtime/Core/Public/Windows/WindowsPlatform.h
--- a/Engine/Source/Runtime/Core/Public/Windows/WindowsPlatform.h
+++ b/Engine/Source/Runtime/Core/Public/Windows/WindowsPlatform.h
@@ -10,6 +10,9 @@
 	#include "MSVC/MSVCPlatform.h"
 #endif
 
+#if WINVER <= 0x600
+	#error "Windows Vista and earlier are no longer supported"
+#endif
 
 /**
 * Windows specific types
diff --git a/Engine/Source/Runtime/Core/Public/Windows/WindowsPlatformCompilerSetup.h b/Engine/Source/Runtime/Core/Public/Windows/WindowsPlatformCompilerSetup.h
--- a/Engine/Source/Runtime/Core/Public/Windows/WindowsPlatformCompilerSetup.h
+++ b/Engine/Source/Runtime/Core/Public/Windows/WindowsPlatformCompilerSetup.h
@@ -187,10 +187,6 @@ static_assert(_MSC_VER >= 1924, "Visual Studio 2019 v16.4 or greater is required
 
 // NOTE: ocid.h breaks this
 #pragma warning(disable: 4917) // 'declarator' : a GUID can only be associated with a class, interface or namespace 				https://docs.microsoft.com/en-us/cpp/error-messages/compiler-warnings/compiler-warning-level-1-c4917
-#if WINVER == 0x0502
-// NOTE: WinXP hits deprecated versions of stdio across the board
-#pragma warning(disable: 4995) // 'function': name was marked as #pragma deprecated													https://docs.microsoft.com/en-us/cpp/error-messages/compiler-warnings/compiler-warning-level-3-c4995
-#endif
 
 //
 // All of the /Wall warnings that we are able to enable
diff --git a/Engine/Source/Runtime/D3D12RHI/Private/Windows/WindowsD3D12Device.cpp b/Engine/Source/Runtime/D3D12RHI/Private/Windows/WindowsD3D12Device.cpp
--- a/Engine/Source/Runtime/D3D12RHI/Private/Windows/WindowsD3D12Device.cpp
+++ b/Engine/Source/Runtime/D3D12RHI/Private/Windows/WindowsD3D12Device.cpp
@@ -162,7 +162,6 @@ static bool bIsQuadBufferStereoEnabled = false;
 /** This function is used as a SEH filter to catch only delay load exceptions. */
 static bool IsDelayLoadException(PEXCEPTION_POINTERS ExceptionPointers)
 {
-#if WINVER > 0x502	// Windows SDK 7.1 doesn't define VcppException
 	switch (ExceptionPointers->ExceptionRecord->ExceptionCode)
 	{
 	case VcppException(ERROR_SEVERITY_ERROR, ERROR_MOD_NOT_FOUND):
@@ -171,9 +170,6 @@ static bool IsDelayLoadException(PEXCEPTION_POINTERS ExceptionPointers)
 	default:
 		return EXCEPTION_CONTINUE_SEARCH;
 	}
-#else
-	return EXCEPTION_EXECUTE_HANDLER;
-#endif
 }
 
 /**
diff --git a/Engine/Source/Runtime/D3D12RHI/Public/D3D12Util.h b/Engine/Source/Runtime/D3D12RHI/Public/D3D12Util.h
--- a/Engine/Source/Runtime/D3D12RHI/Public/D3D12Util.h
+++ b/Engine/Source/Runtime/D3D12RHI/Public/D3D12Util.h
@@ -10,28 +10,6 @@
 #include "Containers/Queue.h"
 #include "RenderUtils.h"
 
-#if WINVER == 0x0502
-// Windows XP uses Win7 sdk, and in that one winerror.h doesn't include them
-#ifndef DXGI_ERROR_INVALID_CALL
-#define DXGI_ERROR_INVALID_CALL                 MAKE_DXGI_HRESULT(1)
-#define DXGI_ERROR_NOT_FOUND                    MAKE_DXGI_HRESULT(2)
-#define DXGI_ERROR_MORE_DATA                    MAKE_DXGI_HRESULT(3)
-#define DXGI_ERROR_UNSUPPORTED                  MAKE_DXGI_HRESULT(4)
-#define DXGI_ERROR_DEVICE_REMOVED               MAKE_DXGI_HRESULT(5)
-#define DXGI_ERROR_DEVICE_HUNG                  MAKE_DXGI_HRESULT(6)
-#define DXGI_ERROR_DEVICE_RESET                 MAKE_DXGI_HRESULT(7)
-#define DXGI_ERROR_WAS_STILL_DRAWING            MAKE_DXGI_HRESULT(10)
-#define DXGI_ERROR_FRAME_STATISTICS_DISJOINT    MAKE_DXGI_HRESULT(11)
-#define DXGI_ERROR_GRAPHICS_VIDPN_SOURCE_IN_USE MAKE_DXGI_HRESULT(12)
-#define DXGI_ERROR_DRIVER_INTERNAL_ERROR        MAKE_DXGI_HRESULT(32)
-#define DXGI_ERROR_NONEXCLUSIVE                 MAKE_DXGI_HRESULT(33)
-#define DXGI_ERROR_NOT_CURRENTLY_AVAILABLE      MAKE_DXGI_HRESULT(34)
-#define DXGI_ERROR_REMOTE_CLIENT_DISCONNECTED   MAKE_DXGI_HRESULT(35)
-#define DXGI_ERROR_REMOTE_OUTOFMEMORY           MAKE_DXGI_HRESULT(36)
-#endif
-
-#endif
-
 namespace D3D12RHI
 {
 	/**
diff --git a/Engine/Source/Runtime/GameplayMediaEncoder/Private/GameplayMediaEncoderCommon.h b/Engine/Source/Runtime/GameplayMediaEncoder/Private/GameplayMediaEncoderCommon.h
--- a/Engine/Source/Runtime/GameplayMediaEncoder/Private/GameplayMediaEncoderCommon.h
+++ b/Engine/Source/Runtime/GameplayMediaEncoder/Private/GameplayMediaEncoderCommon.h
@@ -28,7 +28,7 @@
 	#include "Templates/RefCounting.h"
 #endif
 
-#define WMFMEDIA_SUPPORTED_PLATFORM (PLATFORM_WINDOWS && (WINVER >= 0x0600 /*Vista*/) && !UE_SERVER)
+#define WMFMEDIA_SUPPORTED_PLATFORM (PLATFORM_WINDOWS && !UE_SERVER)
 
 DECLARE_LOG_CATEGORY_EXTERN(GameplayMediaEncoder, Log, VeryVerbose);
 
diff --git a/Engine/Source/Runtime/RHI/Public/GPUDefragAllocator.h b/Engine/Source/Runtime/RHI/Public/GPUDefragAllocator.h
--- a/Engine/Source/Runtime/RHI/Public/GPUDefragAllocator.h
+++ b/Engine/Source/Runtime/RHI/Public/GPUDefragAllocator.h
@@ -1059,12 +1059,7 @@ protected:
 	/** Cumulative time spent in allocator.							*/
 	double			TimeSpentInAllocator;
 	/** Allocated memory in uint8s.									*/
-#if PLATFORM_WINDOWS && (WINVER < 0x0600)
-	// Interlock...64 functions are only available from Vista onwards
-	typedef int32 memsize_t;
-#else
 	typedef int64 memsize_t;
-#endif
 
 	volatile memsize_t	PaddingWasteSize;
 
diff --git a/Engine/Source/Runtime/Windows/D3D11RHI/Private/D3D11Viewport.cpp b/Engine/Source/Runtime/Windows/D3D11RHI/Private/D3D11Viewport.cpp
--- a/Engine/Source/Runtime/Windows/D3D11RHI/Private/D3D11Viewport.cpp
+++ b/Engine/Source/Runtime/Windows/D3D11RHI/Private/D3D11Viewport.cpp
@@ -10,11 +10,7 @@
 #include "HAL/ThreadHeartBeat.h"
 
 #ifndef D3D11_WITH_DWMAPI
-#if WINVER > 0x502		// Windows XP doesn't support DWM
-	#define D3D11_WITH_DWMAPI	1
-#else
-	#define D3D11_WITH_DWMAPI	0
-#endif
+#define D3D11_WITH_DWMAPI	1
 #endif
 
 #if D3D11_WITH_DWMAPI
diff --git a/Engine/Source/Runtime/Windows/D3D11RHI/Private/Windows/WindowsD3D11Device.cpp b/Engine/Source/Runtime/Windows/D3D11RHI/Private/Windows/WindowsD3D11Device.cpp
--- a/Engine/Source/Runtime/Windows/D3D11RHI/Private/Windows/WindowsD3D11Device.cpp
+++ b/Engine/Source/Runtime/Windows/D3D11RHI/Private/Windows/WindowsD3D11Device.cpp
@@ -172,7 +172,6 @@ FAutoConsoleCommand FD3DDumpLiveObjectsCommand
 /** This function is used as a SEH filter to catch only delay load exceptions. */
 static bool IsDelayLoadException(PEXCEPTION_POINTERS ExceptionPointers)
 {
-#if WINVER > 0x502	// Windows SDK 7.1 doesn't define VcppException
 	switch(ExceptionPointers->ExceptionRecord->ExceptionCode)
 	{
 	case VcppException(ERROR_SEVERITY_ERROR, ERROR_MOD_NOT_FOUND):
@@ -181,9 +180,6 @@ static bool IsDelayLoadException(PEXCEPTION_POINTERS ExceptionPointers)
 	default:
 		return EXCEPTION_CONTINUE_SEARCH;
 	}
-#else
-	return EXCEPTION_EXECUTE_HANDLER;
-#endif
 }
 
 static bool bIsQuadBufferStereoEnabled = false;
diff --git a/Engine/Source/Runtime/Windows/D3D11RHI/Public/D3D11Util.h b/Engine/Source/Runtime/Windows/D3D11RHI/Public/D3D11Util.h
--- a/Engine/Source/Runtime/Windows/D3D11RHI/Public/D3D11Util.h
+++ b/Engine/Source/Runtime/Windows/D3D11RHI/Public/D3D11Util.h
@@ -6,27 +6,6 @@
 
 #pragma once
 
-#if WINVER == 0x0502
-// Windows XP uses Win7 sdk, and in that one winerror.h doesn't include them
-
-#define DXGI_ERROR_INVALID_CALL                 MAKE_DXGI_HRESULT(1)
-#define DXGI_ERROR_NOT_FOUND                    MAKE_DXGI_HRESULT(2)
-#define DXGI_ERROR_MORE_DATA                    MAKE_DXGI_HRESULT(3)
-#define DXGI_ERROR_UNSUPPORTED                  MAKE_DXGI_HRESULT(4)
-#define DXGI_ERROR_DEVICE_REMOVED               MAKE_DXGI_HRESULT(5)
-#define DXGI_ERROR_DEVICE_HUNG                  MAKE_DXGI_HRESULT(6)
-#define DXGI_ERROR_DEVICE_RESET                 MAKE_DXGI_HRESULT(7)
-#define DXGI_ERROR_WAS_STILL_DRAWING            MAKE_DXGI_HRESULT(10)
-#define DXGI_ERROR_FRAME_STATISTICS_DISJOINT    MAKE_DXGI_HRESULT(11)
-#define DXGI_ERROR_GRAPHICS_VIDPN_SOURCE_IN_USE MAKE_DXGI_HRESULT(12)
-#define DXGI_ERROR_DRIVER_INTERNAL_ERROR        MAKE_DXGI_HRESULT(32)
-#define DXGI_ERROR_NONEXCLUSIVE                 MAKE_DXGI_HRESULT(33)
-#define DXGI_ERROR_NOT_CURRENTLY_AVAILABLE      MAKE_DXGI_HRESULT(34)
-#define DXGI_ERROR_REMOTE_CLIENT_DISCONNECTED   MAKE_DXGI_HRESULT(35)
-#define DXGI_ERROR_REMOTE_OUTOFMEMORY           MAKE_DXGI_HRESULT(36)
-
-#endif
-
 #define D3D11RHI_IMMEDIATE_CONTEXT	(GD3D11RHI->GetDeviceContext())
 #define D3D11RHI_DEVICE				(GD3D11RHI->GetDevice())
 
diff --git a/Engine/Source/Runtime/Windows/WindowsPlatformFeatures/Private/WmfPrivate.h b/Engine/Source/Runtime/Windows/WindowsPlatformFeatures/Private/WmfPrivate.h
--- a/Engine/Source/Runtime/Windows/WindowsPlatformFeatures/Private/WmfPrivate.h
+++ b/Engine/Source/Runtime/Windows/WindowsPlatformFeatures/Private/WmfPrivate.h
@@ -5,7 +5,7 @@
 #include "Logging/LogMacros.h"
 #include "CoreMinimal.h"
 
-#define WMFMEDIA_SUPPORTED_PLATFORM (PLATFORM_WINDOWS && (WINVER >= 0x0600 /*Vista*/) && !UE_SERVER)
+#define WMFMEDIA_SUPPORTED_PLATFORM (PLATFORM_WINDOWS && !UE_SERVER)
 
 DECLARE_LOG_CATEGORY_EXTERN(WMF, Log, VeryVerbose);
 
diff --git a/Templates/TP_InCamVFXBP/Config/DefaultEngine.ini b/Templates/TP_InCamVFXBP/Config/DefaultEngine.ini
--- a/Templates/TP_InCamVFXBP/Config/DefaultEngine.ini
+++ b/Templates/TP_InCamVFXBP/Config/DefaultEngine.ini
@@ -24,7 +24,6 @@ Compiler=VisualStudio2019
 -TargetedRHIs=PCD3D_SM5
 +TargetedRHIs=PCD3D_SM5
 DefaultGraphicsRHI=DefaultGraphicsRHI_DX12
-MinimumOSVersion=MSOS_Vista
 bTarget32Bit=False
 AudioSampleRate=48000
 AudioCallbackBufferFrameSize=1024
