Removed references to UE4 in WindowsApplication source

diff --git a/Engine/Build/Windows/Resources/Default.rc2 b/Engine/Build/Windows/Resources/Default.rc2
--- a/Engine/Build/Windows/Resources/Default.rc2
+++ b/Engine/Build/Windows/Resources/Default.rc2
@@ -92,7 +92,7 @@ END
 // remains consistent on all systems.
 
 #ifdef BUILD_ICON_FILE_NAME
-IDICON_UE4Game          ICON                    BUILD_ICON_FILE_NAME
+IDICON_UEGame          ICON                    BUILD_ICON_FILE_NAME
 #endif
 
 #endif    // English (United States) resources
diff --git a/Engine/Source/Programs/SwitchboardListener/Private/Windows/Resources/Resource.rc2 b/Engine/Source/Programs/SwitchboardListener/Private/Windows/Resources/Resource.rc2
--- a/Engine/Source/Programs/SwitchboardListener/Private/Windows/Resources/Resource.rc2
+++ b/Engine/Source/Programs/SwitchboardListener/Private/Windows/Resources/Resource.rc2
@@ -49,4 +49,4 @@ END
 
 // Icon with lowest ID value placed first to ensure application icon
 // remains consistent on all systems.
-IDICON_UE4Game			ICON                    "Icon.ico"
+IDICON_UEGame			ICON                    "Icon.ico"
diff --git a/Engine/Source/Programs/UnrealVersionSelector/Private/Windows/Resources/Resource.rc2 b/Engine/Source/Programs/UnrealVersionSelector/Private/Windows/Resources/Resource.rc2
--- a/Engine/Source/Programs/UnrealVersionSelector/Private/Windows/Resources/Resource.rc2
+++ b/Engine/Source/Programs/UnrealVersionSelector/Private/Windows/Resources/Resource.rc2
@@ -49,4 +49,4 @@ END
 
 // Icon with lowest ID value placed first to ensure application icon
 // remains consistent on all systems.
-IDICON_UE4Game			ICON                    "Icon.ico"
+IDICON_UEGame			ICON                    "Icon.ico"
diff --git a/Engine/Source/Runtime/ApplicationCore/Private/Windows/WindowsPlatformApplicationMisc.cpp b/Engine/Source/Runtime/ApplicationCore/Private/Windows/WindowsPlatformApplicationMisc.cpp
--- a/Engine/Source/Runtime/ApplicationCore/Private/Windows/WindowsPlatformApplicationMisc.cpp
+++ b/Engine/Source/Runtime/ApplicationCore/Private/Windows/WindowsPlatformApplicationMisc.cpp
@@ -90,7 +90,7 @@ bool FWindowsPlatformApplicationMisc::IsThisApplicationForeground()
 
 int32 FWindowsPlatformApplicationMisc::GetAppIcon()
 {
-	return IDICON_UE4Game;
+	return IDICON_UEGame;
 }
 
 static void WinPumpMessages()
diff --git a/Engine/Source/Runtime/ApplicationCore/Private/Windows/WindowsPlatformSurvey.cpp b/Engine/Source/Runtime/ApplicationCore/Private/Windows/WindowsPlatformSurvey.cpp
--- a/Engine/Source/Runtime/ApplicationCore/Private/Windows/WindowsPlatformSurvey.cpp
+++ b/Engine/Source/Runtime/ApplicationCore/Private/Windows/WindowsPlatformSurvey.cpp
@@ -166,7 +166,7 @@ bool FWindowsPlatformSurvey::GetSurveyResults( FHardwareSurveyResults& OutResult
 	else
 	{
 		uint32 ErrorCode = FPlatformMisc::GetLastError();
-		UE_LOG(LogWindows, Warning, TEXT("FWindowsPlatformSurvey::TickSurveyHardware() failed to get UE4 root-folder drive size from Win32") );
+		UE_LOG(LogWindows, Warning, TEXT("FWindowsPlatformSurvey::TickSurveyHardware() failed to get Unreal Engine root-folder drive size from Win32") );
 		OutResults.ErrorCount++;
 		WriteFStringToResults(OutResults.LastSurveyError, TEXT("GetDiskFreeSpaceEx() failed"));
 		WriteFStringToResults(OutResults.LastSurveyErrorDetail, FString::Printf(TEXT("ErrorCode: 0x%0x"), ErrorCode));
diff --git a/Engine/Source/Runtime/ApplicationCore/Public/Windows/Accessibility/WindowsUIAManager.h b/Engine/Source/Runtime/ApplicationCore/Public/Windows/Accessibility/WindowsUIAManager.h
--- a/Engine/Source/Runtime/ApplicationCore/Public/Windows/Accessibility/WindowsUIAManager.h
+++ b/Engine/Source/Runtime/ApplicationCore/Public/Windows/Accessibility/WindowsUIAManager.h
@@ -16,7 +16,7 @@ class IAccessibleWidget;
 class FVariant;
 
 /**
- * Manager for Windows implementation of UE4's accessibility API, utilizing Microsoft's UI Automation API. It provides a
+ * Manager for Windows implementation of UE's accessibility API, utilizing Microsoft's UI Automation API. It provides a
  * central location for message passing to/from the WindowsApplication, AccessibleMessageHandler, and Windows UIA Providers.
  *
  * This class only accepts and returns references, so callers should ensure their objects are valid before using it.
diff --git a/Engine/Source/Runtime/Launch/Resources/Windows/resource.h b/Engine/Source/Runtime/Launch/Resources/Windows/resource.h
--- a/Engine/Source/Runtime/Launch/Resources/Windows/resource.h
+++ b/Engine/Source/Runtime/Launch/Resources/Windows/resource.h
@@ -7,8 +7,7 @@
 #define IDD_YESNO2ALL                   101
 #define IDD_YESNO2ALLCANCEL             102
 #define IDD_YESNOYESTOALL               103
-// If IDICON_UE4Game is modified, please also modify IconResourceId in Engine\Source\Programs\UnrealBuildTool\Windows\WindowsDeploy.cs
-#define IDICON_UE4Game                  123
+#define IDICON_UEGame                   123
 #define IDC_YESTOALL                    1000
 #define IDC_NOTOALL                     1001
 #define IDC_YES                         1002
