UE5EA - Remove ue4 reference

diff --git a/Engine/Source/Runtime/ApplicationCore/Private/Linux/LinuxWindow.cpp b/Engine/Source/Runtime/ApplicationCore/Private/Linux/LinuxWindow.cpp
--- a/Engine/Source/Runtime/ApplicationCore/Private/Linux/LinuxWindow.cpp
+++ b/Engine/Source/Runtime/ApplicationCore/Private/Linux/LinuxWindow.cpp
@@ -20,7 +20,7 @@ DEFINE_LOG_CATEGORY( LogLinuxWindowEvent );
 
 // SDL 2.0.4 as of 10374:dccf51aee79b will account for border width/height automatically (see SDL_x11window.c)
 // might need to be a function in case SDL gets overriden at runtime
-#define UE4_USING_BORDERS_AWARE_SDL					1
+#define UE_USING_BORDERS_AWARE_SDL					1
 
 FLinuxWindow::~FLinuxWindow()
 {
@@ -317,7 +317,7 @@ void FLinuxWindow::Initialize( FLinuxApplication* const Application, const TShar
 	/* 
 		Do not set for Notification Windows the transient flag because the WM's usually raise the the parent window
 		if the Notificaton Window gets raised. That behaviour is to aggresive and disturbs users doing other things 
-		while UE4 calculates lights and other things and pop ups notifications. Notifications will be handled so that 
+		while UnrealEngine calculates lights and other things and pop ups notifications. Notifications will be handled so that 
 		they are some sort of independend but will be raised if the TopLevel Window gets focused or activated.
 	*/
 	// Make the Window modal for it's parent.
@@ -384,7 +384,7 @@ SDL_HitTestResult FLinuxWindow::HitTest( SDL_Window *SDLwin, const SDL_Point *po
 /** Native windows should implement MoveWindowTo by relocating the platform-specific window to (X,Y). */
 void FLinuxWindow::MoveWindowTo( int32 X, int32 Y )
 {
-	if (UE4_USING_BORDERS_AWARE_SDL)
+	if (UE_USING_BORDERS_AWARE_SDL)
 	{
 		SDL_SetWindowPosition( HWnd, X, Y );
 	}
@@ -545,7 +545,7 @@ void FLinuxWindow::ReshapeWindow( int32 NewX, int32 NewY, int32 NewWidth, int32
 
 		case EWindowMode::Windowed:
 		{
-			if (UE4_USING_BORDERS_AWARE_SDL == 0 && Definition->HasOSWindowBorder)
+			if (UE_USING_BORDERS_AWARE_SDL == 0 && Definition->HasOSWindowBorder)
 			{
 				// we are passed coordinates of a client area, so account for decorations
 				checkf(bValidNativePropertiesCache, TEXT("Attempted to use border sizes too early, native properties aren't yet cached. Review the flow"));
