Updated windows splash screen. Mac and Linux may not look correct but will be fixed by engineers familar with those platforms

diff --git a/Engine/Source/Editor/UnrealEd/Private/FeedbackContextEditor.cpp b/Engine/Source/Editor/UnrealEd/Private/FeedbackContextEditor.cpp
--- a/Engine/Source/Editor/UnrealEd/Private/FeedbackContextEditor.cpp
+++ b/Engine/Source/Editor/UnrealEd/Private/FeedbackContextEditor.cpp
@@ -594,7 +594,7 @@ void FFeedbackContextEditor::ProgressReported( const float TotalProgressInterp,
 				}
 			}
 
-			DisplayMessage = FText::FromString(FString::Printf(TEXT("%3i%% - %s"), int(TotalProgressInterp * 100.f), *NewDisplayMessage));
+			DisplayMessage = FText::FromString(FString::Printf(TEXT("%i%% - %s"), int(TotalProgressInterp * 100.f), *NewDisplayMessage));
 		}
 
 		FPlatformSplash::SetSplashText(SplashTextType::StartupProgress, *DisplayMessage.ToString());
diff --git a/Engine/Source/Runtime/ApplicationCore/Private/Linux/LinuxPlatformSplash.cpp b/Engine/Source/Runtime/ApplicationCore/Private/Linux/LinuxPlatformSplash.cpp
--- a/Engine/Source/Runtime/ApplicationCore/Private/Linux/LinuxPlatformSplash.cpp
+++ b/Engine/Source/Runtime/ApplicationCore/Private/Linux/LinuxPlatformSplash.cpp
@@ -684,7 +684,7 @@ void FLinuxPlatformSplash::Show( )
 	// first look for the splash, do not init anything if not found
 	FString SplashPath;
 	{
-		const TCHAR* SplashImage = GIsEditor ? (GameName.IsEmpty() ? TEXT("EdSplashDefault") : TEXT("EdSplash")) : (GameName.IsEmpty() ? TEXT("SplashDefault") : TEXT("Splash"));
+		const TCHAR* SplashImage = GIsEditor ? TEXT("EdSplash") : TEXT("Splash");
 		if (!GetSplashPath(SplashImage, SplashPath, IsCustom))
 		{
 			UE_LOG(LogHAL, Warning, TEXT("Splash screen image not found."));
diff --git a/Engine/Source/Runtime/ApplicationCore/Private/Mac/MacPlatformSplash.cpp b/Engine/Source/Runtime/ApplicationCore/Private/Mac/MacPlatformSplash.cpp
--- a/Engine/Source/Runtime/ApplicationCore/Private/Mac/MacPlatformSplash.cpp
+++ b/Engine/Source/Runtime/ApplicationCore/Private/Mac/MacPlatformSplash.cpp
@@ -186,7 +186,7 @@ void FMacPlatformSplash::Show()
 
 		const FText GameName = FText::FromString(FApp::GetProjectName());
 
-		const TCHAR* SplashImage = GIsEditor ? ( GameName.IsEmpty() ? TEXT("EdSplashDefault") : TEXT("EdSplash") ) : ( GameName.IsEmpty() ? TEXT("SplashDefault") : TEXT("Splash") );
+		const TCHAR* SplashImage = GIsEditor ? TEXT("EdSplash") : TEXT("Splash");
 
 		// make sure a splash was found
 		FString SplashPath;
diff --git a/Engine/Source/Runtime/ApplicationCore/Private/Windows/WindowsPlatformSplash.cpp b/Engine/Source/Runtime/ApplicationCore/Private/Windows/WindowsPlatformSplash.cpp
--- a/Engine/Source/Runtime/ApplicationCore/Private/Windows/WindowsPlatformSplash.cpp
+++ b/Engine/Source/Runtime/ApplicationCore/Private/Windows/WindowsPlatformSplash.cpp
@@ -95,11 +95,6 @@ LRESULT CALLBACK SplashScreenWindowProc(HWND hWnd, uint32 message, WPARAM wParam
 							}
 
 							// Alignment
-							if ( CurTypeIndex == SplashTextType::GameName )
-							{
-								SetTextAlign( hdc, TA_RIGHT | TA_TOP | TA_NOUPDATECP );
-							}
-							else
 							{
 								SetTextAlign( hdc, TA_LEFT | TA_TOP | TA_NOUPDATECP );
 							}
@@ -111,7 +106,7 @@ LRESULT CALLBACK SplashScreenWindowProc(HWND hWnd, uint32 message, WPARAM wParam
 							GetClientRect( hWnd, &ClientRect );
 
 							// Draw background text passes
-							const int32 NumBGPasses = 8;
+					/*		const int32 NumBGPasses = 1;
 							for( int32 CurBGPass = 0; CurBGPass < NumBGPasses; ++CurBGPass )
 							{
 								int32 BGXOffset, BGYOffset;
@@ -135,24 +130,24 @@ LRESULT CALLBACK SplashScreenWindowProc(HWND hWnd, uint32 message, WPARAM wParam
 									TextRect.top + BGYOffset,
 									*SplashText.ToString(),
 									SplashText.ToString().Len() );
-							}
+							}*/
 							
 							// Draw foreground text pass
 							if( CurTypeIndex == SplashTextType::StartupProgress )
 							{
-								SetTextColor( hdc, RGB( 200, 200, 200 ) );
+								SetTextColor( hdc, RGB(160, 160, 160) );
 							}
 							else if( CurTypeIndex == SplashTextType::VersionInfo1 )
 							{
-								SetTextColor( hdc, RGB( 240, 240, 240 ) );
+								SetTextColor( hdc, RGB(160, 160, 160) );
 							}
 							else if ( CurTypeIndex == SplashTextType::GameName )
 							{
-								SetTextColor(hdc, RGB(240, 240, 240));
+								SetTextColor(hdc, RGB(255, 255, 255));
 							}
 							else
 							{
-								SetTextColor( hdc, RGB( 160, 160, 160 ) );
+								SetTextColor( hdc, RGB(160, 160, 160) );
 							}
 
 							TextOut(
@@ -412,15 +407,12 @@ uint32 WINAPI StartSplashScreenThread( LPVOID unused )
 	{
 		BITMAP bm;
 		GetObjectW(GSplashScreenBitmap, sizeof(bm), &bm);
-
-		const int32 BorderWidth = GetSystemMetrics(SM_CXBORDER);
-		const int32 BorderHeight = GetSystemMetrics(SM_CYBORDER);
-		const int32 WindowWidth = bm.bmWidth + BorderWidth;
-		const int32 WindowHeight = bm.bmHeight + BorderHeight;
+		const int32 WindowWidth = bm.bmWidth;
+		const int32 WindowHeight = bm.bmHeight;
 		int32 ScreenPosX = (GetSystemMetrics(SM_CXSCREEN) - WindowWidth) / 2;
 		int32 ScreenPosY = (GetSystemMetrics(SM_CYSCREEN) - WindowHeight) / 2;
 
-		const bool bAllowFading = true;
+		const bool bAllowFading = false;
 
 		// Force the editor splash screen to show up in the taskbar and alt-tab lists
 		uint32 dwWindowStyle = (GIsEditor ? WS_EX_APPWINDOW : 0) | WS_EX_TOOLWINDOW;
@@ -433,7 +425,7 @@ uint32 WINAPI StartSplashScreenThread( LPVOID unused )
 			dwWindowStyle,
 			wc.lpszClassName, 
 			TEXT("SplashScreen"),
-			WS_BORDER|WS_POPUP,
+			WS_POPUP,
 			ScreenPosX,
 			ScreenPosY,
 			WindowWidth,
@@ -443,6 +435,13 @@ uint32 WINAPI StartSplashScreenThread( LPVOID unused )
 			hInstance,
 			(LPVOID) NULL); 
 
+		//Round the corners
+		RECT rcWnd;
+		GetWindowRect(GSplashScreenWnd, &rcWnd);
+		HRGN Region = CreateRoundRectRgn(0, 0, rcWnd.right - rcWnd.left, rcWnd.bottom - rcWnd.top, 10, 10);
+		SetWindowRgn(GSplashScreenWnd, Region, false);
+
+
 		if( bAllowFading )
 		{
 			// Set window to fully transparent to start out
@@ -459,8 +458,9 @@ uint32 WINAPI StartSplashScreenThread( LPVOID unused )
 				LOGFONT MyFont;
 				FMemory::Memzero( &MyFont, sizeof( MyFont ) );
 				GetObjectW( SystemFontHandle, sizeof( MyFont ), &MyFont );
-				MyFont.lfHeight = 10;
-				// MyFont.lfQuality = ANTIALIASED_QUALITY;
+				MyFont.lfHeight = 11;
+				MyFont.lfQuality = CLEARTYPE_QUALITY;
+				//StringCchCopy(MyFont.lfFaceName, LF_FACESIZE, TEXT("Roboto"));
 				GSplashScreenSmallTextFontHandle = CreateFontIndirect( &MyFont );
 				if( GSplashScreenSmallTextFontHandle == NULL )
 				{
@@ -475,7 +475,8 @@ uint32 WINAPI StartSplashScreenThread( LPVOID unused )
 				FMemory::Memzero( &MyFont, sizeof( MyFont ) );
 				GetObjectW( SystemFontHandle, sizeof( MyFont ), &MyFont );
 				MyFont.lfHeight = 12;
-				// MyFont.lfQuality = ANTIALIASED_QUALITY;
+				MyFont.lfQuality = CLEARTYPE_QUALITY;
+			//	StringCchCopy(MyFont.lfFaceName, LF_FACESIZE, TEXT("Roboto"));
 				GSplashScreenNormalTextFontHandle = CreateFontIndirect( &MyFont );
 				if( GSplashScreenNormalTextFontHandle == NULL )
 				{
@@ -489,10 +490,10 @@ uint32 WINAPI StartSplashScreenThread( LPVOID unused )
 				LOGFONT MyFont;
 				FMemory::Memzero(&MyFont, sizeof( MyFont ));
 				GetObjectW(SystemFontHandle, sizeof( MyFont ), &MyFont);
-				MyFont.lfHeight = 40;
+				MyFont.lfHeight = GIsEditor ? 18 : 28;
 				MyFont.lfWeight = FW_BOLD;
-				MyFont.lfQuality = ANTIALIASED_QUALITY;
-				StringCchCopy(MyFont.lfFaceName, LF_FACESIZE, TEXT("Verdana"));
+				MyFont.lfQuality = CLEARTYPE_QUALITY;
+			//	StringCchCopy(MyFont.lfFaceName, LF_FACESIZE, TEXT("Roboto"));
 				GSplashScreenTitleTextFontHandle = CreateFontIndirect(&MyFont);
 				if ( GSplashScreenTitleTextFontHandle == NULL )
 				{
@@ -503,29 +504,40 @@ uint32 WINAPI StartSplashScreenThread( LPVOID unused )
 		}
 
 		// Setup bounds for game name
-		GSplashScreenTextRects[ SplashTextType::GameName ].top = 10;
-		GSplashScreenTextRects[ SplashTextType::GameName ].bottom = 60;
-		GSplashScreenTextRects[ SplashTextType::GameName ].left = bm.bmWidth - 12;
-		GSplashScreenTextRects[ SplashTextType::GameName ].right = 12;
+
+		if (GIsEditor)
+		{
+			GSplashScreenTextRects[SplashTextType::GameName].top = bm.bmHeight - 60;
+			GSplashScreenTextRects[SplashTextType::GameName].bottom = bm.bmHeight - 45;
+			GSplashScreenTextRects[SplashTextType::GameName].left = 10;
+			GSplashScreenTextRects[SplashTextType::GameName].right = bm.bmWidth - 20;
+		}
+		else
+		{
+			GSplashScreenTextRects[SplashTextType::GameName].top = bm.bmHeight - 45;
+			GSplashScreenTextRects[SplashTextType::GameName].bottom = bm.bmHeight - 6;
+			GSplashScreenTextRects[SplashTextType::GameName].left = 10;
+			GSplashScreenTextRects[SplashTextType::GameName].right = bm.bmWidth - 20;
+		}
 		
 		// Setup bounds for version info text 1
-		GSplashScreenTextRects[ SplashTextType::VersionInfo1 ].top = bm.bmHeight - 60;
-		GSplashScreenTextRects[ SplashTextType::VersionInfo1 ].bottom = bm.bmHeight - 40;
+		GSplashScreenTextRects[ SplashTextType::VersionInfo1 ].top = bm.bmHeight - 40;
+		GSplashScreenTextRects[ SplashTextType::VersionInfo1 ].bottom = bm.bmHeight - 20;
 		GSplashScreenTextRects[ SplashTextType::VersionInfo1 ].left = 10;
 		GSplashScreenTextRects[ SplashTextType::VersionInfo1 ].right = bm.bmWidth - 20;
 
 		// Setup bounds for copyright info text
 		if( GIsEditor )
 		{
-			GSplashScreenTextRects[ SplashTextType::CopyrightInfo ].top = bm.bmHeight - 44;
-			GSplashScreenTextRects[ SplashTextType::CopyrightInfo ].bottom = bm.bmHeight - 34;
+			GSplashScreenTextRects[ SplashTextType::CopyrightInfo ].top = bm.bmHeight - 16;
+			GSplashScreenTextRects[ SplashTextType::CopyrightInfo ].bottom = bm.bmHeight - 6;
 		}
 		else
 		{
-			GSplashScreenTextRects[ SplashTextType::CopyrightInfo ].top = bm.bmHeight - 16;
-			GSplashScreenTextRects[ SplashTextType::CopyrightInfo ].bottom = bm.bmHeight - 6;
+			GSplashScreenTextRects[SplashTextType::CopyrightInfo].top = bm.bmHeight - 16;
+			GSplashScreenTextRects[SplashTextType::CopyrightInfo].bottom = bm.bmHeight - 6;
 		}
-		GSplashScreenTextRects[ SplashTextType::CopyrightInfo ].left = 10;
+		GSplashScreenTextRects[ SplashTextType::CopyrightInfo ].left = bm.bmWidth - 180;
 		GSplashScreenTextRects[ SplashTextType::CopyrightInfo ].right = bm.bmWidth - 20;
 
 		// Setup bounds for startup progress text
@@ -610,7 +622,7 @@ void FWindowsPlatformSplash::Show()
 	{
 		const FText GameName = FText::FromString( FApp::GetProjectName() );
 
-		const TCHAR* SplashImage = GIsEditor ? ( GameName.IsEmpty() ? TEXT("EdSplashDefault") : TEXT("EdSplash") ) : ( GameName.IsEmpty() ? TEXT("SplashDefault") : TEXT("Splash") );
+		const TCHAR* SplashImage = GIsEditor ?  TEXT("EdSplash") : TEXT("Splash");
 
 		// make sure a splash was found
 		FString SplashPath;
@@ -618,11 +630,7 @@ void FWindowsPlatformSplash::Show()
 		if ( GetSplashPath(SplashImage, SplashPath, IsCustom ) == true )
 		{
 			// Don't set the game name if the splash screen is custom.
-			if ( !IsCustom )
-			{
-				StartSetSplashText(SplashTextType::GameName, *GameName.ToString());
-			}
-
+		
 			// In the editor, we'll display loading info
 			if( GIsEditor )
 			{
@@ -636,20 +644,19 @@ void FWindowsPlatformSplash::Show()
 				{
 					const FText Version = FText::FromString( FEngineVersion::Current().ToString( FEngineBuildSettings::IsPerforceBuild() ? EVersionComponent::Branch : EVersionComponent::Patch ) );
 
-					FText VersionInfo;
 					FText AppName;
+					FText VersionInfo = FText::Format(NSLOCTEXT("UnrealEd", "UnrealEdTitleWithVersion_F", "Unreal Editor {0}"), Version);
 					if( GameName.IsEmpty() )
 					{
-						VersionInfo = FText::Format( NSLOCTEXT( "UnrealEd", "UnrealEdTitleWithVersionNoGameName_F", "Unreal Editor {0}" ), Version );
 						AppName = NSLOCTEXT( "UnrealEd", "UnrealEdTitleNoGameName_F", "Unreal Editor" );
 					}
 					else
 					{
-						VersionInfo = FText::Format( NSLOCTEXT( "UnrealEd", "UnrealEdTitleWithVersion_F", "Unreal Editor {0}  -  {1}" ), Version, GameName );
 						AppName = FText::Format( NSLOCTEXT( "UnrealEd", "UnrealEdTitle_F", "Unreal Editor - {0}" ), GameName );
 					}
 
-					StartSetSplashText( SplashTextType::VersionInfo1, *VersionInfo.ToString() );
+					StartSetSplashText(SplashTextType::GameName, *AppName.ToString());
+					StartSetSplashText(SplashTextType::VersionInfo1, *VersionInfo.ToString());
 
 					// Change the window text (which will be displayed in the taskbar)
 					GSplashScreenAppName = AppName;
@@ -661,7 +668,10 @@ void FWindowsPlatformSplash::Show()
 					StartSetSplashText( SplashTextType::CopyrightInfo, *CopyrightInfo );
 				}
 			}
-
+			else if(!IsCustom)
+			{
+				StartSetSplashText(SplashTextType::GameName, *GameName.ToString());
+			}
 			// Spawn a window to receive the Z-order swap when the splashscreen is destroyed.
 			// This will prevent the main window from being sent to the background when the splash window closes.
 			GSplashScreenGuard = CreateWindow(
