Fix static analysis

diff --git a/Engine/Source/Runtime/ApplicationCore/Private/Windows/WindowsPlatformSplash.cpp b/Engine/Source/Runtime/ApplicationCore/Private/Windows/WindowsPlatformSplash.cpp
--- a/Engine/Source/Runtime/ApplicationCore/Private/Windows/WindowsPlatformSplash.cpp
+++ b/Engine/Source/Runtime/ApplicationCore/Private/Windows/WindowsPlatformSplash.cpp
@@ -527,16 +527,8 @@ uint32 WINAPI StartSplashScreenThread( LPVOID unused )
 		GSplashScreenTextRects[ SplashTextType::VersionInfo1 ].right = bm.bmWidth - 20;
 
 		// Setup bounds for copyright info text
-		if( GIsEditor )
-		{
-			GSplashScreenTextRects[ SplashTextType::CopyrightInfo ].top = bm.bmHeight - 16;
-			GSplashScreenTextRects[ SplashTextType::CopyrightInfo ].bottom = bm.bmHeight - 6;
-		}
-		else
-		{
-			GSplashScreenTextRects[SplashTextType::CopyrightInfo].top = bm.bmHeight - 16;
-			GSplashScreenTextRects[SplashTextType::CopyrightInfo].bottom = bm.bmHeight - 6;
-		}
+		GSplashScreenTextRects[ SplashTextType::CopyrightInfo ].top = bm.bmHeight - 16;
+		GSplashScreenTextRects[ SplashTextType::CopyrightInfo ].bottom = bm.bmHeight - 6;
 		GSplashScreenTextRects[ SplashTextType::CopyrightInfo ].left = bm.bmWidth - 180;
 		GSplashScreenTextRects[ SplashTextType::CopyrightInfo ].right = bm.bmWidth - 20;
 
