Fix up LinuxSplashScreen changes

diff --git a/Engine/Source/Runtime/ApplicationCore/Private/Linux/LinuxPlatformSplash.cpp b/Engine/Source/Runtime/ApplicationCore/Private/Linux/LinuxPlatformSplash.cpp
--- a/Engine/Source/Runtime/ApplicationCore/Private/Linux/LinuxPlatformSplash.cpp
+++ b/Engine/Source/Runtime/ApplicationCore/Private/Linux/LinuxPlatformSplash.cpp
@@ -185,7 +185,7 @@ bool FLinuxSplashState::OpenFonts()
 	}
 	else
 	{
-		FT_Set_Pixel_Sizes(FontNormal, 0, 12);
+		FT_Set_Pixel_Sizes(FontNormal, 0, 16);
 	}
 
 	// large font face
@@ -197,7 +197,7 @@ bool FLinuxSplashState::OpenFonts()
 	}
 	else
 	{
-		FT_Set_Pixel_Sizes(FontLarge, 0, 40);
+		FT_Set_Pixel_Sizes(FontLarge, 0, 32);
 	}
 
 	return true;
@@ -292,6 +292,7 @@ void FLinuxSplashState::RenderStrings()
 		uint32 PenX = SplashTextRects[ CurTypeIndex].Left;
 		uint32 PenY = SplashTextRects[ CurTypeIndex].Bottom;
 		Kerning.x = 0;
+		bRightJustify = false;
 
 		// Set font color and text position
 		if (CurTypeIndex == SplashTextType::StartupProgress)
@@ -302,12 +303,24 @@ void FLinuxSplashState::RenderStrings()
 		else if (CurTypeIndex == SplashTextType::VersionInfo1)
 		{
 			Red = Green = Blue = 240.0f;
-			Font = FontNormal;
+			Font = FontSmall;
 		}
 		else if (CurTypeIndex == SplashTextType::GameName)
 		{
 			Red = Green = Blue = 240.0f;
-			Font = FontLarge;
+			if (GIsEditor)
+			{
+				Font = FontNormal;
+			}
+			else
+			{
+				Font = FontLarge;
+			}
+		}
+		else if (CurTypeIndex == SplashTextType::CopyrightInfo)
+		{
+			Red = Green = Blue = 160.0f;
+			Font = FontSmall;
 			PenX = SplashTextRects[ CurTypeIndex].Right;
 			bRightJustify = true;
 		}
@@ -533,31 +546,31 @@ bool FLinuxSplashState::InitSplashResources(const FText &AppName, const FString
 	ScratchSpace = reinterpret_cast<unsigned char *>(FMemory::Malloc(SplashHeight*SplashWidth*SplashBPP));
 
 	// Setup bounds for game name
-	SplashTextRects[ SplashTextType::GameName ].Top = 0;
-	SplashTextRects[ SplashTextType::GameName ].Bottom = 50;
-	SplashTextRects[ SplashTextType::GameName ].Left = 12;
-	SplashTextRects[ SplashTextType::GameName ].Right = SplashWidth - 12;
-
-	// Setup bounds for version info text 1
-	SplashTextRects[ SplashTextType::VersionInfo1 ].Top = SplashHeight - 60;
-	SplashTextRects[ SplashTextType::VersionInfo1 ].Bottom = SplashHeight - 40;
-	SplashTextRects[ SplashTextType::VersionInfo1 ].Left = 10;
-	SplashTextRects[ SplashTextType::VersionInfo1 ].Right = SplashWidth - 10;
-
-	// Setup bounds for copyright info text
 	if (GIsEditor)
 	{
-		SplashTextRects[ SplashTextType::CopyrightInfo ].Top = SplashHeight - 44;
-		SplashTextRects[ SplashTextType::CopyrightInfo ].Bottom = SplashHeight - 24;
+		SplashTextRects[ SplashTextType::GameName ].Top = SplashHeight - 60;
+		SplashTextRects[ SplashTextType::GameName ].Bottom = SplashHeight - 40;
 	}
 	else
 	{
-		SplashTextRects[ SplashTextType::CopyrightInfo ].Top = SplashHeight - 16;
-		SplashTextRects[ SplashTextType::CopyrightInfo ].Bottom = SplashHeight - 6;
+		SplashTextRects[ SplashTextType::GameName ].Top = SplashHeight - 60;
+		SplashTextRects[ SplashTextType::GameName ].Bottom = SplashHeight - 20;
 	}
 
-	SplashTextRects[ SplashTextType::CopyrightInfo ].Left = 10;
-	SplashTextRects[ SplashTextType::CopyrightInfo ].Right = SplashWidth - 20;
+	SplashTextRects[ SplashTextType::GameName ].Left = 10;
+	SplashTextRects[ SplashTextType::GameName ].Right = SplashWidth - 10;
+
+	// Setup bounds for version info text 1
+	SplashTextRects[ SplashTextType::VersionInfo1 ].Top = SplashHeight - 35;
+	SplashTextRects[ SplashTextType::VersionInfo1 ].Bottom = SplashHeight - 20;
+	SplashTextRects[ SplashTextType::VersionInfo1 ].Left = 10;
+	SplashTextRects[ SplashTextType::VersionInfo1 ].Right = SplashWidth - 10;
+
+	// Setup bounds for copyright info text
+	SplashTextRects[ SplashTextType::CopyrightInfo ].Top = SplashHeight - 16;
+	SplashTextRects[ SplashTextType::CopyrightInfo ].Bottom = SplashHeight;
+	SplashTextRects[ SplashTextType::CopyrightInfo ].Left = 0;
+	SplashTextRects[ SplashTextType::CopyrightInfo ].Right = SplashWidth - 12;
 
 	// Setup bounds for startup progress text
 	SplashTextRects[ SplashTextType::StartupProgress ].Top = SplashHeight - 20;
@@ -677,7 +690,20 @@ void FLinuxPlatformSplash::Show( )
 	}
 
 	// decide on which splash screen to show
-	const FText GameName = FText::FromString(FApp::GetProjectName());
+	FString FinalGameName = FApp::GetProjectName();
+	if (GIsEditor)
+	{
+		if (!FApp::HasProjectName())
+		{
+			FinalGameName = FString("Unreal Editor");
+		}
+		else
+		{
+			FinalGameName = FString("Unreal Editor - ") + FinalGameName;
+		}
+	}
+
+	const FText GameName = FText::FromString(FinalGameName);
 
 	bool IsCustom = false;
 
@@ -729,12 +755,6 @@ void FLinuxPlatformSplash::Show( )
 
 	GSplashState = new FLinuxSplashState;
 
-	// Don't set the game name if the splash screen is custom.
-	if ( !IsCustom )
-	{
-		GSplashState->SetSplashText( SplashTextType::GameName, GameName );
-	}
-
 	// In the editor, we'll display loading info
 	FText AppName;
 	if ( GIsEditor )
@@ -749,17 +769,8 @@ void FLinuxPlatformSplash::Show( )
 		{
 			const FText Version = FText::FromString( FEngineVersion::Current().ToString( FEngineBuildSettings::IsPerforceBuild() ? EVersionComponent::Branch : EVersionComponent::Patch ) );
 
-			FText VersionInfo;
-			if( GameName.IsEmpty() )
-			{
-				VersionInfo = FText::Format( NSLOCTEXT( "UnrealEd", "UnrealEdTitleWithVersionNoGameName_F", "Unreal Editor {0}" ), Version );
-				AppName = NSLOCTEXT( "UnrealEd", "UnrealEdTitleNoGameName_F", "Unreal Editor" );
-			}
-			else
-			{
-				VersionInfo = FText::Format( NSLOCTEXT( "UnrealEd", "UnrealEdTitleWithVersion_F", "Unreal Editor {0}  -  {1}" ), Version, GameName );
-				AppName = FText::Format( NSLOCTEXT( "UnrealEd", "UnrealEdTitle_F", "Unreal Editor - {0}" ), GameName );
-			}
+			FText VersionInfo = FText::Format( NSLOCTEXT( "UnrealEd", "UnrealEdTitleWithVersionNoGameName_F", "Unreal Editor {0}" ), Version );
+			AppName = NSLOCTEXT( "UnrealEd", "UnrealEdTitleNoGameName_F", "Unreal Editor" );
 
 			GSplashState->SetSplashText( SplashTextType::VersionInfo1, VersionInfo );
 		}
@@ -769,6 +780,16 @@ void FLinuxPlatformSplash::Show( )
 			const FText CopyrightInfo = NSLOCTEXT( "UnrealEd", "SplashScreen_CopyrightInfo", "Copyright \x00a9   Epic Games, Inc.   All rights reserved." );
 			GSplashState->SetSplashText( SplashTextType::CopyrightInfo, CopyrightInfo );
 		}
+
+		// Set GameName for Editor always
+		{
+			GSplashState->SetSplashText( SplashTextType::GameName, GameName );
+		}
+	}
+	// Don't set the game name if the splash screen is custom.
+	else if ( !IsCustom )
+	{
+		GSplashState->SetSplashText( SplashTextType::GameName, GameName );
 	}
 
 	if (!GSplashState->InitSplashResources(AppName, SplashPath, IconPath))
