Update splash screen images

diff --git a/Engine/Source/Runtime/ApplicationCore/Private/Windows/WindowsPlatformSplash.cpp b/Engine/Source/Runtime/ApplicationCore/Private/Windows/WindowsPlatformSplash.cpp
--- a/Engine/Source/Runtime/ApplicationCore/Private/Windows/WindowsPlatformSplash.cpp
+++ b/Engine/Source/Runtime/ApplicationCore/Private/Windows/WindowsPlatformSplash.cpp
@@ -435,13 +435,6 @@ uint32 WINAPI StartSplashScreenThread( LPVOID unused )
 			hInstance,
 			(LPVOID) NULL); 
 
-		//Round the corners
-		RECT rcWnd;
-		GetWindowRect(GSplashScreenWnd, &rcWnd);
-		HRGN Region = CreateRoundRectRgn(0, 0, rcWnd.right - rcWnd.left, rcWnd.bottom - rcWnd.top, 10, 10);
-		SetWindowRgn(GSplashScreenWnd, Region, false);
-
-
 		if( bAllowFading )
 		{
 			// Set window to fully transparent to start out
