Clone of CL 15874941.:  Call out to update isOnActiveSpace per FCocoaWindow when the app becomes active.  Works around unrespostive crash reporter on Catalina when coming from a full screen app crash.   First call to isOnActiveSpace in OnActiveSpaceDidChange() during startup returns false on Catalina but true on BigSur.

diff --git a/Engine/Source/Runtime/ApplicationCore/Private/Mac/MacApplication.cpp b/Engine/Source/Runtime/ApplicationCore/Private/Mac/MacApplication.cpp
--- a/Engine/Source/Runtime/ApplicationCore/Private/Mac/MacApplication.cpp
+++ b/Engine/Source/Runtime/ApplicationCore/Private/Mac/MacApplication.cpp
@@ -1397,6 +1397,10 @@ void FMacApplication::OnApplicationDidBecomeActive()
 		FGenericCrashContext::SetEngineData(TEXT("Platform.AppHasFocus"), TEXT("true"));
 #endif
 	}, @[ NSDefaultRunLoopMode ], false);
+	
+	// Call out to update isOnActiveSpace per FCocoaWindow when the app becomes active.  Works round unrespostive crash reporter on Catalina
+	// First call to isOnActiveSpace in OnActiveSpaceDidChange() during crash reporter startup returns false on Catalina but true on BigSur.
+	OnActiveSpaceDidChange();
 }
 
 void FMacApplication::OnApplicationWillResignActive()
