Fix Mac Splash Screen

diff --git a/Engine/Source/Runtime/ApplicationCore/Private/Mac/MacPlatformSplash.cpp b/Engine/Source/Runtime/ApplicationCore/Private/Mac/MacPlatformSplash.cpp
--- a/Engine/Source/Runtime/ApplicationCore/Private/Mac/MacPlatformSplash.cpp
+++ b/Engine/Source/Runtime/ApplicationCore/Private/Mac/MacPlatformSplash.cpp
@@ -80,7 +80,7 @@ static void StartSetSplashText( const SplashTextType::Type InType, const FText&
 
 		if( SplashText.ToString().Len() > 0 )
 		{
-			int32 FontSize = 11;
+			int32 FontSize = 9;
 			switch ( CurTypeIndex )
 			{
 				case SplashTextType::StartupProgress:
@@ -88,7 +88,7 @@ static void StartSetSplashText( const SplashTextType::Type InType, const FText&
 					FontSize = 12;
 					break;
 				case SplashTextType::GameName:
-					FontSize = 34;
+					FontSize = 18;
 					break;
 			}
 
@@ -121,7 +121,7 @@ static void StartSetSplashText( const SplashTextType::Type InType, const FText&
 			NSTextAlignment align = NSTextAlignmentLeft;
 			switch ( CurTypeIndex )
 			{
-			case SplashTextType::GameName:
+			case SplashTextType::CopyrightInfo:
 				align = NSTextAlignmentRight;
 				break;
 			}
@@ -184,8 +184,8 @@ void FMacPlatformSplash::Show()
 	{
 		SCOPED_AUTORELEASE_POOL;
 
-		const FText GameName = FText::FromString(FApp::GetProjectName());
-
+        const FText GameName = FText::FromString(FApp::GetProjectName());
+        
 		const TCHAR* SplashImage = GIsEditor ? TEXT("EdSplash") : TEXT("Splash");
 
 		// make sure a splash was found
@@ -218,6 +218,7 @@ void FMacPlatformSplash::Show()
 					{
 						VersionInfo = FText::Format( NSLOCTEXT( "UnrealEd", "UnrealEdTitleWithVersionNoGameName_F", "Unreal Editor {0}" ), Version );
 						AppName = NSLOCTEXT( "UnrealEd", "UnrealEdTitleNoGameName_F", "Unreal Editor" );
+						StartSetSplashText( SplashTextType::GameName, AppName);
 					}
 					else
 					{
@@ -262,10 +263,10 @@ void FMacPlatformSplash::Show()
 			int32 FontHeight = 14;
 			
 			// Setup bounds for game name
-			GSplashScreenTextRects[ SplashTextType::GameName ].origin.x = 10;
-			GSplashScreenTextRects[ SplashTextType::GameName ].origin.y = 0;
-			GSplashScreenTextRects[ SplashTextType::GameName ].size.width = ImageWidth - 2 * 10;
-			GSplashScreenTextRects[ SplashTextType::GameName ].size.height = ImageHeight;
+			GSplashScreenTextRects[ SplashTextType::GameName ].origin.x = OriginX;
+			GSplashScreenTextRects[ SplashTextType::GameName ].origin.y = GIsEditor ? OriginY + 34 : OriginY + 18;
+			GSplashScreenTextRects[ SplashTextType::GameName ].size.width = ImageWidth - 2 * OriginX;
+			GSplashScreenTextRects[ SplashTextType::GameName ].size.height = 2 * FontHeight;
 			
 			// Setup bounds for texts
 			GSplashScreenTextRects[ SplashTextType::VersionInfo1 ].origin.x =
@@ -277,20 +278,12 @@ void FMacPlatformSplash::Show()
 			GSplashScreenTextRects[ SplashTextType::StartupProgress ].size.width = ImageWidth - 2 * OriginX;
 			
 			GSplashScreenTextRects[ SplashTextType::VersionInfo1 ].size.height =
-			GSplashScreenTextRects[ SplashTextType::CopyrightInfo ].size.height =
 			GSplashScreenTextRects[ SplashTextType::StartupProgress ].size.height = FontHeight;
+			GSplashScreenTextRects[ SplashTextType::CopyrightInfo ].size.height = FontHeight;
 			
-			GSplashScreenTextRects[ SplashTextType::VersionInfo1 ].origin.y = OriginY + 3 * FontHeight;
-			GSplashScreenTextRects[ SplashTextType::StartupProgress ].origin.y = OriginY;
-			
-			if( GIsEditor )
-			{
-				GSplashScreenTextRects[ SplashTextType::CopyrightInfo ].origin.y = OriginY + 2 * FontHeight;
-			}
-			else
-			{
-				GSplashScreenTextRects[ SplashTextType::CopyrightInfo ].origin.y = OriginY;
-			}
+			GSplashScreenTextRects[ SplashTextType::VersionInfo1 ].origin.y = OriginY + 22;
+			GSplashScreenTextRects[ SplashTextType::StartupProgress ].origin.y = OriginY + 3;
+			GSplashScreenTextRects[ SplashTextType::CopyrightInfo ].origin.y = OriginY - 2;
 		}
 
 		MainThreadCall(^{
