Remove UE4 references

This removes a variety of low-risk UE4 references.

diff --git a/Engine/Source/Developer/DesktopPlatform/Private/Mac/MacNativeFeedbackContext.cpp b/Engine/Source/Developer/DesktopPlatform/Private/Mac/MacNativeFeedbackContext.cpp
--- a/Engine/Source/Developer/DesktopPlatform/Private/Mac/MacNativeFeedbackContext.cpp
+++ b/Engine/Source/Developer/DesktopPlatform/Private/Mac/MacNativeFeedbackContext.cpp
@@ -62,7 +62,7 @@
 		NSRect WindowRect = NSMakeRect(ConsolePosX, ConsolePosY, ConsoleWidth, ConsoleHeight);
 		
 		Window = [[NSWindow alloc] initWithContentRect:WindowRect styleMask:NSWindowStyleMaskTitled|NSWindowStyleMaskMiniaturizable|NSWindowStyleMaskResizable|NSWindowStyleMaskClosable backing:NSBackingStoreBuffered defer:NO];
-		[Window setTitle:@"Unreal Engine 4"];
+		[Window setTitle:@"Unreal Engine"];
 		[Window setReleasedWhenClosed:NO];
 		[Window setMinSize:NSMakeSize(400, 100)];
 		[Window setRestorable:NO];
@@ -266,7 +266,7 @@ FMacNativeFeedbackContext::FMacNativeFeedbackContext()
 {
 	MainThreadCall(^{
 		WindowController = [[FMacNativeFeedbackContextWindowController alloc] init];
-	}, UE4NilEventMode, true);
+	}, UnrealNilEventMode, true);
 	SetDefaultTextColor();
 }
 
@@ -279,7 +279,7 @@ FMacNativeFeedbackContext::~FMacNativeFeedbackContext()
 	
 	MainThreadCall(^{
 		[WindowController release];
-	}, UE4NilEventMode, true);
+	}, UnrealNilEventMode, true);
 }
 
 void FMacNativeFeedbackContext::Serialize( const TCHAR* Data, ELogVerbosity::Type Verbosity, const class FName& Category )
diff --git a/Engine/Source/Runtime/ApplicationCore/Private/Mac/Accessibility/MacAccessibilityElement.cpp b/Engine/Source/Runtime/ApplicationCore/Private/Mac/Accessibility/MacAccessibilityElement.cpp
--- a/Engine/Source/Runtime/ApplicationCore/Private/Mac/Accessibility/MacAccessibilityElement.cpp
+++ b/Engine/Source/Runtime/ApplicationCore/Private/Mac/Accessibility/MacAccessibilityElement.cpp
@@ -268,7 +268,7 @@
 			{
 				Widget->SetFocus();
 			}
-		}, @[ NSDefaultRunLoopMode, UE4IMEEventMode, UE4ShowEventMode, UE4ResizeEventMode, UE4FullscreenEventMode], false); // GameThreadCall
+		}, @[ NSDefaultRunLoopMode, UnrealIMEEventMode, UnrealShowEventMode, UnrealResizeEventMode, UnrealFullscreenEventMode], false); // GameThreadCall
 	}
 	[super setAccessibilityFocused:InIsFocused];
 }
@@ -435,7 +435,7 @@
 				}
 			}
 		} // if widget valid and enabled
-	}, @[ NSDefaultRunLoopMode, UE4IMEEventMode, UE4ShowEventMode, UE4ResizeEventMode, UE4FullscreenEventMode], false); // GameThreadCall
+	}, @[ NSDefaultRunLoopMode, UnrealIMEEventMode, UnrealShowEventMode, UnrealResizeEventMode, UnrealFullscreenEventMode], false); // GameThreadCall
 	return YES;
 }
 
diff --git a/Engine/Source/Runtime/ApplicationCore/Private/Mac/Accessibility/MacAccessibilityManager.cpp b/Engine/Source/Runtime/ApplicationCore/Private/Mac/Accessibility/MacAccessibilityManager.cpp
--- a/Engine/Source/Runtime/ApplicationCore/Private/Mac/Accessibility/MacAccessibilityManager.cpp
+++ b/Engine/Source/Runtime/ApplicationCore/Private/Mac/Accessibility/MacAccessibilityManager.cpp
@@ -227,7 +227,7 @@ static void DumpAccessibilityStatsForwarder()
 				}, NSDefaultRunLoopMode, true);
 			} // if Widget valid
 		} // for Ids loop
-	}, @[ NSDefaultRunLoopMode, UE4IMEEventMode, UE4ShowEventMode, UE4ResizeEventMode, UE4FullscreenEventMode, UE4NilEventMode ], true);
+	}, @[ NSDefaultRunLoopMode, UnrealIMEEventMode, UnrealShowEventMode, UnrealResizeEventMode, UnrealFullscreenEventMode, UnrealNilEventMode ], true);
 }
 
 #if !UE_BUILD_SHIPPING
diff --git a/Engine/Source/Runtime/ApplicationCore/Private/Mac/CocoaTextView.cpp b/Engine/Source/Runtime/ApplicationCore/Private/Mac/CocoaTextView.cpp
--- a/Engine/Source/Runtime/ApplicationCore/Private/Mac/CocoaTextView.cpp
+++ b/Engine/Source/Runtime/ApplicationCore/Private/Mac/CocoaTextView.cpp
@@ -202,7 +202,7 @@
 						ITextInputMethodContext::ECaretPosition CaretPosition;
 						IMMContext->GetSelectionRange(SelectionLocation, SelectionLength, CaretPosition);
 					}
-				}, @[ NSDefaultRunLoopMode, UE4IMEEventMode ]);
+				}, @[ NSDefaultRunLoopMode, UnrealIMEEventMode ]);
 			}
 		}
 		else
@@ -229,7 +229,7 @@
 				IMMContext->SetTextInRange(SelectionLocation, SelectionLength, TheFString);
 				IMMContext->SetSelectionRange(SelectionLocation+TheFString.Len(), 0, ITextInputMethodContext::ECaretPosition::Ending);
 			}
-		}, @[ NSDefaultRunLoopMode, UE4IMEEventMode ]);
+		}, @[ NSDefaultRunLoopMode, UnrealIMEEventMode ]);
 
 		[self unmarkText];
 		[[self inputContext] invalidateCharacterCoordinates]; // recentering
@@ -271,7 +271,7 @@
 						ITextInputMethodContext::ECaretPosition CaretPosition;
 						IMMContext->GetSelectionRange(SelectionLocation, SelectionLength, CaretPosition);
 					}
-				}, @[ NSDefaultRunLoopMode, UE4IMEEventMode ]);
+				}, @[ NSDefaultRunLoopMode, UnrealIMEEventMode ]);
 			}
 		}
 		else
@@ -287,7 +287,7 @@
 				{
 					IMMContext->SetTextInRange(SelectionLocation, SelectionLength, FString());
 				}
-			}, @[ NSDefaultRunLoopMode, UE4IMEEventMode ]);
+			}, @[ NSDefaultRunLoopMode, UnrealIMEEventMode ]);
 			[self unmarkText];
 		}
 		else
@@ -299,7 +299,7 @@
 					{
 						IMMContext->BeginComposition();
 					}
-				}, @[ NSDefaultRunLoopMode, UE4IMEEventMode ]);
+				}, @[ NSDefaultRunLoopMode, UnrealIMEEventMode ]);
 			}
 			markedRange = NSMakeRange(SelectionLocation, [aString length]);
 
@@ -344,7 +344,7 @@
 					IMMContext->UpdateCompositionRange(CompositionRange.location, CompositionRange.length);
 					IMMContext->SetSelectionRange(markedRange.location + selectedRange.location, 0, ITextInputMethodContext::ECaretPosition::Ending);
 				}
-			}, @[ NSDefaultRunLoopMode, UE4IMEEventMode ]);
+			}, @[ NSDefaultRunLoopMode, UnrealIMEEventMode ]);
 		}
 		[[self inputContext] invalidateCharacterCoordinates]; // recentering
 	}
@@ -369,7 +369,7 @@
 					IMMContext->UpdateCompositionRange(0, 0);
 					IMMContext->EndComposition();
 				}
-			}, @[ NSDefaultRunLoopMode, UE4IMEEventMode ]);
+			}, @[ NSDefaultRunLoopMode, UnrealIMEEventMode ]);
 		}
 	}
 }
@@ -394,7 +394,7 @@
 			{
 				IMMContext->GetSelectionRange(SelectionLocation, SelectionLength, CaretPosition);
 			}
-		}, @[ NSDefaultRunLoopMode, UE4IMEEventMode ]);
+		}, @[ NSDefaultRunLoopMode, UnrealIMEEventMode ]);
 		return {SelectionLocation, SelectionLength};
 	}
 	else
@@ -438,7 +438,7 @@
 			{
 				IMMContext->GetTextInRange(aRange.location, aRange.length, String);
 			}
-		}, @[ NSDefaultRunLoopMode, UE4IMEEventMode ]);
+		}, @[ NSDefaultRunLoopMode, UnrealIMEEventMode ]);
 		CFStringRef CFString = FPlatformString::TCHARToCFString(*String);
 		if(CFString)
 		{
@@ -479,7 +479,7 @@
 			{
 				IMMContext->GetTextBounds(aRange.location, aRange.length, Position, Size);
 			}
-		}, @[ NSDefaultRunLoopMode, UE4IMEEventMode ]);
+		}, @[ NSDefaultRunLoopMode, UnrealIMEEventMode ]);
 
 		if(actualRange)
 		{
@@ -507,7 +507,7 @@
 	if (IMMContext.IsValid())
 	{
 		FVector2D Pos(aPoint.x, aPoint.y);
-		int32 Index = GameThreadReturn(^{ return IMMContext.IsValid() ? IMMContext->GetCharacterIndexFromPoint(Pos) : INDEX_NONE; }, @[ NSDefaultRunLoopMode, UE4IMEEventMode ]);
+		int32 Index = GameThreadReturn(^{ return IMMContext.IsValid() ? IMMContext->GetCharacterIndexFromPoint(Pos) : INDEX_NONE; }, @[ NSDefaultRunLoopMode, UnrealIMEEventMode ]);
 		NSUInteger Idx = Index == INDEX_NONE ? NSNotFound : (NSUInteger)Index;
 		return Idx;
 	}
diff --git a/Engine/Source/Runtime/ApplicationCore/Private/Mac/CocoaWindow.cpp b/Engine/Source/Runtime/ApplicationCore/Private/Mac/CocoaWindow.cpp
--- a/Engine/Source/Runtime/ApplicationCore/Private/Mac/CocoaWindow.cpp
+++ b/Engine/Source/Runtime/ApplicationCore/Private/Mac/CocoaWindow.cpp
@@ -151,7 +151,7 @@ NSString* NSPerformDragOperation = @"NSPerformDragOperation";
 	{
 		MainThreadCall(^{
 			[super setAlphaValue:WindowAlpha];
-		}, UE4ShowEventMode, false);
+		}, UnrealShowEventMode, false);
 
 	}
 }
@@ -163,7 +163,7 @@ NSString* NSPerformDragOperation = @"NSPerformDragOperation";
 		bRenderInitialized = true;
 		MainThreadCall(^{
 			[super setAlphaValue:Opacity];
-		}, UE4ShowEventMode, false);
+		}, UnrealShowEventMode, false);
 	}
 }
 
@@ -300,7 +300,7 @@ NSString* NSPerformDragOperation = @"NSPerformDragOperation";
 					MacApplication->OnWindowActivationChanged(Window.ToSharedRef(), EWindowActivation::Activate);
 				}
 			}
-		}, @[ NSDefaultRunLoopMode, UE4ResizeEventMode, UE4ShowEventMode, UE4FullscreenEventMode, UE4CloseEventMode ], true);
+		}, @[ NSDefaultRunLoopMode, UnrealResizeEventMode, UnrealShowEventMode, UnrealFullscreenEventMode, UnrealCloseEventMode ], true);
 	}
 }
 
@@ -321,7 +321,7 @@ NSString* NSPerformDragOperation = @"NSPerformDragOperation";
 					MacApplication->OnWindowActivationChanged(Window.ToSharedRef(), EWindowActivation::Deactivate);
 				}
 			}
-		}, @[ NSDefaultRunLoopMode, UE4ResizeEventMode, UE4ShowEventMode, UE4FullscreenEventMode, UE4CloseEventMode ], true);
+		}, @[ NSDefaultRunLoopMode, UnrealResizeEventMode, UnrealShowEventMode, UnrealFullscreenEventMode, UnrealCloseEventMode ], true);
 	}
 }
 
@@ -480,7 +480,7 @@ NSString* NSPerformDragOperation = @"NSPerformDragOperation";
 					MacApplication->OnWindowWillResize(Window.ToSharedRef());
 				}
 			}
-		}, @[ NSDefaultRunLoopMode, UE4ResizeEventMode, UE4ShowEventMode, UE4FullscreenEventMode ], true);
+		}, @[ NSDefaultRunLoopMode, UnrealResizeEventMode, UnrealShowEventMode, UnrealFullscreenEventMode ], true);
 	}
 	return frameSize;
 }
diff --git a/Engine/Source/Runtime/ApplicationCore/Private/Mac/MacApplication.cpp b/Engine/Source/Runtime/ApplicationCore/Private/Mac/MacApplication.cpp
--- a/Engine/Source/Runtime/ApplicationCore/Private/Mac/MacApplication.cpp
+++ b/Engine/Source/Runtime/ApplicationCore/Private/Mac/MacApplication.cpp
@@ -444,7 +444,7 @@ void FMacApplication::DeferEvent(NSObject* Object)
 
 		if (DeferredEvent.Type == NSEventTypeKeyDown)
 		{
-			// In UE4 the main window rather than key window is the current active window in Slate, so the main window may be the one we want to send immKeyDown to,
+			// In Unreal the main window rather than key window is the current active window in Slate, so the main window may be the one we want to send immKeyDown to,
 			// for example in case of search text edit fields in context menus.
 			NSWindow* MainWindow = [NSApp mainWindow];
 			FCocoaWindow* IMMWindow = [MainWindow isKindOfClass:[FCocoaWindow class]] ? (FCocoaWindow*)MainWindow : DeferredEvent.Window;
@@ -543,7 +543,7 @@ void FMacApplication::DeferEvent(NSObject* Object)
 						{
 							OnWindowDidResize(Window.ToSharedRef());
 						}
-					}, @[ NSDefaultRunLoopMode, UE4ResizeEventMode, UE4ShowEventMode, UE4FullscreenEventMode ], true);
+					}, @[ NSDefaultRunLoopMode, UnrealResizeEventMode, UnrealShowEventMode, UnrealFullscreenEventMode ], true);
 				}
 				return;
 			}
@@ -1568,7 +1568,7 @@ void FMacApplication::ConditionallyUpdateModifierKeys(const FDeferredMacEvent& E
 	}
 }
 
-void FMacApplication::HandleModifierChange(NSUInteger NewModifierFlags, NSUInteger FlagsShift, NSUInteger UE4Shift, EMacModifierKeys TranslatedCode)
+void FMacApplication::HandleModifierChange(NSUInteger NewModifierFlags, NSUInteger FlagsShift, NSUInteger UEShift, EMacModifierKeys TranslatedCode)
 {
 	const bool CurrentPressed = (CurrentModifierFlags & FlagsShift) != 0;
 	const bool NewPressed = (NewModifierFlags & FlagsShift) != 0;
@@ -1576,12 +1576,12 @@ void FMacApplication::HandleModifierChange(NSUInteger NewModifierFlags, NSUInteg
 	{
 		if (NewPressed)
 		{
-			ModifierKeysFlags |= 1 << UE4Shift;
+			ModifierKeysFlags |= 1 << UEShift;
 			MessageHandler->OnKeyDown(TranslatedCode, 0, false);
 		}
 		else
 		{
-			ModifierKeysFlags &= ~(1 << UE4Shift);
+			ModifierKeysFlags &= ~(1 << UEShift);
 			MessageHandler->OnKeyUp(TranslatedCode, 0, false);
 		}
 	}
@@ -2084,7 +2084,7 @@ void FMacApplication::CloseQueuedWindows()
 				[Window close];
 				[Window release];
 			}
-		}, UE4CloseEventMode, true);
+		}, UnrealCloseEventMode, true);
 
 		CocoaWindowsToClose.Empty();
 	}
@@ -2106,7 +2106,7 @@ void FMacApplication::InvalidateTextLayouts()
 				}
 			}
 
-		}, UE4IMEEventMode, true);
+		}, UnrealIMEEventMode, true);
 
 		WindowsRequiringTextInvalidation.Empty();
 	}
diff --git a/Engine/Source/Runtime/ApplicationCore/Private/Mac/MacConsoleOutputDevice.cpp b/Engine/Source/Runtime/ApplicationCore/Private/Mac/MacConsoleOutputDevice.cpp
--- a/Engine/Source/Runtime/ApplicationCore/Private/Mac/MacConsoleOutputDevice.cpp
+++ b/Engine/Source/Runtime/ApplicationCore/Private/Mac/MacConsoleOutputDevice.cpp
@@ -110,7 +110,7 @@ void FMacConsoleOutputDevice::CreateConsole()
 		}
 		
 		SetDefaultTextColor();
-	}, UE4NilEventMode, true);
+	}, UnrealNilEventMode, true);
 }
 
 void FMacConsoleOutputDevice::DestroyConsole()
@@ -137,7 +137,7 @@ void FMacConsoleOutputDevice::DestroyConsole()
 			
 			[ConsoleWindow close];
 			TextViewTextColor = NULL;
-		}, UE4NilEventMode, true);
+		}, UnrealNilEventMode, true);
 	}
 }
 
diff --git a/Engine/Source/Runtime/ApplicationCore/Private/Mac/MacTextInputMethodSystem.cpp b/Engine/Source/Runtime/ApplicationCore/Private/Mac/MacTextInputMethodSystem.cpp
--- a/Engine/Source/Runtime/ApplicationCore/Private/Mac/MacTextInputMethodSystem.cpp
+++ b/Engine/Source/Runtime/ApplicationCore/Private/Mac/MacTextInputMethodSystem.cpp
@@ -69,7 +69,7 @@ namespace
 					FCocoaTextView* TextView = (FCocoaTextView*)[CocoaWindow openGLView];
 					[[TextView inputContext] invalidateCharacterCoordinates];
 				}
-			}, UE4IMEEventMode, true);
+			}, UnrealIMEEventMode, true);
 		}
 	}
 	
@@ -85,7 +85,7 @@ namespace
 					FCocoaTextView* TextView = (FCocoaTextView*)[CocoaWindow openGLView];
 					[[TextView inputContext] invalidateCharacterCoordinates];
 				}
-			}, UE4IMEEventMode, true);
+			}, UnrealIMEEventMode, true);
 		}
 	}
 	
@@ -102,7 +102,7 @@ namespace
 					[[TextView inputContext] discardMarkedText];
 					[TextView unmarkText];
 				}
-			}, UE4IMEEventMode, true);
+			}, UnrealIMEEventMode, true);
 		}
 	}
 }
@@ -182,7 +182,7 @@ void FMacTextInputMethodSystem::ActivateContext(const TSharedRef<ITextInputMetho
 				}
 			}
 			return bSuccess;
-		}, UE4IMEEventMode);
+		}, UnrealIMEEventMode);
 	}
 	if(!bActivatedContext)
 	{
@@ -228,7 +228,7 @@ void FMacTextInputMethodSystem::DeactivateContext(const TSharedRef<ITextInputMet
 				// If the window is no longer open, the context is already inactive (handled by MessageHandler->OnWindowActivationChanged() called in FMacApplication::OnWindowDestroyed())
 				return true;
 			}
-		}, UE4IMEEventMode);
+		}, UnrealIMEEventMode);
 	}
 	if(!bDeactivatedContext)
 	{
diff --git a/Engine/Source/Runtime/ApplicationCore/Private/Mac/MacWindow.cpp b/Engine/Source/Runtime/ApplicationCore/Private/Mac/MacWindow.cpp
--- a/Engine/Source/Runtime/ApplicationCore/Private/Mac/MacWindow.cpp
+++ b/Engine/Source/Runtime/ApplicationCore/Private/Mac/MacWindow.cpp
@@ -213,7 +213,7 @@ PRAGMA_DISABLE_DEPRECATION_WARNINGS
 PRAGMA_ENABLE_DEPRECATION_WARNINGS
 			check(0);
 		}
-	}, UE4ShowEventMode, true);
+	}, UnrealShowEventMode, true);
 }
 
 FCocoaWindow* FMacWindow::GetWindowHandle() const
@@ -249,7 +249,7 @@ void FMacWindow::MoveWindowTo( int32 X, int32 Y )
 		SCOPED_AUTORELEASE_POOL;
 		const FVector2D Point = FMacApplication::ConvertSlatePositionToCocoa(X, Y);
 		[WindowHandle setFrameOrigin:NSMakePoint(Point.X, Point.Y - [WindowHandle openGLFrame].size.height + 1)];
-	}, UE4ResizeEventMode, true);
+	}, UnrealResizeEventMode, true);
 }
 
 void FMacWindow::BringToFront( bool bForce )
@@ -266,7 +266,7 @@ void FMacWindow::BringToFront( bool bForce )
 		MainThreadCall(^{
 			SCOPED_AUTORELEASE_POOL;
 			[WindowHandleCopy orderFrontAndMakeMain:bOrderAndKey andKey:bOrderAndKey];
-		}, UE4ShowEventMode, true);
+		}, UnrealShowEventMode, true);
 
 		MacApplication->OnWindowOrderedFront(SharedThis(this));
 	}
@@ -284,7 +284,7 @@ void FMacWindow::Destroy()
 			SCOPED_AUTORELEASE_POOL;
 			[WindowHandleCopy setAlphaValue:0.0f];
 			[WindowHandleCopy setBackgroundColor:[NSColor clearColor]];
-		}, UE4ShowEventMode, false);
+		}, UnrealShowEventMode, false);
 
 		MacApplication->OnWindowDestroyed(SharedThis(this));
 		WindowHandle = nullptr;
@@ -296,7 +296,7 @@ void FMacWindow::Minimize()
 	MainThreadCall(^{
 		SCOPED_AUTORELEASE_POOL;
 		[WindowHandle miniaturize:nil];
-	}, UE4ResizeEventMode, true);
+	}, UnrealResizeEventMode, true);
 }
 
 void FMacWindow::Maximize()
@@ -308,7 +308,7 @@ void FMacWindow::Maximize()
 			WindowHandle->bZoomed = true;
 			[WindowHandle zoom:nil];
 		}
-	}, UE4ResizeEventMode, true);
+	}, UnrealResizeEventMode, true);
 }
 
 void FMacWindow::Restore()
@@ -323,7 +323,7 @@ void FMacWindow::Restore()
 		{
 			[WindowHandle zoom:nil];
 		}
-	}, UE4ResizeEventMode, true);
+	}, UnrealResizeEventMode, true);
 
 	WindowHandle->bZoomed = WindowHandle.zoomed;
 }
@@ -360,7 +360,7 @@ void FMacWindow::Show()
 		MainThreadCall(^{
 			SCOPED_AUTORELEASE_POOL;
 			[WindowHandleCopy orderFrontAndMakeMain:bShouldActivate andKey:bShouldActivate];
-		}, UE4ShowEventMode, true);
+		}, UnrealShowEventMode, true);
 
 		bIsVisible = true;
 	}
@@ -375,7 +375,7 @@ void FMacWindow::Hide()
 		MainThreadCall(^{
 			SCOPED_AUTORELEASE_POOL;
 			[WindowHandle orderOut:nil];
-		}, UE4CloseEventMode, true);
+		}, UnrealCloseEventMode, true);
 	}
 }
 
@@ -439,7 +439,7 @@ void FMacWindow::SetWindowFocus()
 	MainThreadCall(^{
 		SCOPED_AUTORELEASE_POOL;
 		[WindowHandle orderFrontAndMakeMain:true andKey:true];
-	}, UE4ShowEventMode, true);
+	}, UnrealShowEventMode, true);
 
 	MacApplication->OnWindowOrderedFront(SharedThis(this));
 }
@@ -450,7 +450,7 @@ void FMacWindow::SetOpacity( const float InOpacity )
 		SCOPED_AUTORELEASE_POOL;
         CachedOpacity = InOpacity;
 		[WindowHandle setAlphaValue:InOpacity];
-	}, UE4NilEventMode, true);
+	}, UnrealNilEventMode, true);
 }
 
 bool FMacWindow::IsPointInWindow( int32 X, int32 Y ) const
@@ -521,7 +521,7 @@ void FMacWindow::SetText(const TCHAR* const Text)
 				[NSApp changeWindowsItem: WindowHandle title: (NSString*)CFName filename: NO];
 			}
 			CFRelease( CFName );
-		}, UE4NilEventMode, true);
+		}, UnrealNilEventMode, true);
 	}
 }
 
@@ -611,7 +611,7 @@ void FMacWindow::ApplySizeAndModeChanges(int32 X, int32 Y, int32 Width, int32 He
 			{
 				CGDisplayFade(FadeReservationToken, 0.3, kCGDisplayBlendNormal, kCGDisplayBlendSolidColor, 0, 0, 0, true);
 			}
-		}, UE4ResizeEventMode, true);
+		}, UnrealResizeEventMode, true);
 	}
 
 	if (WindowMode == EWindowMode::Windowed || WindowMode == EWindowMode::WindowedFullscreen)
@@ -620,7 +620,7 @@ void FMacWindow::ApplySizeAndModeChanges(int32 X, int32 Y, int32 Width, int32 He
 		{
 			MainThreadCall(^{
 				CGDisplaySetDisplayMode(WindowedModeSavedState.CapturedDisplayID, WindowedModeSavedState.DesktopDisplayMode, nullptr);
-			}, UE4ResizeEventMode, true);
+			}, UnrealResizeEventMode, true);
 
 			CGDisplayModeRelease(WindowedModeSavedState.DesktopDisplayMode);
 			WindowedModeSavedState.DesktopDisplayMode = nullptr;
@@ -666,7 +666,7 @@ void FMacWindow::ApplySizeAndModeChanges(int32 X, int32 Y, int32 Width, int32 He
 					{
 						[WindowHandle setContentAspectRatio:NSMakeSize((float)Width / (float)Height, 1.0f)];
 					}
-				}, UE4ResizeEventMode, true);
+				}, UnrealResizeEventMode, true);
 			}
 		}
 	}
@@ -689,7 +689,7 @@ void FMacWindow::ApplySizeAndModeChanges(int32 X, int32 Y, int32 Width, int32 He
 			CGDisplayModeRef DisplayMode = FPlatformApplicationMisc::GetSupportedDisplayMode(WindowedModeSavedState.CapturedDisplayID, Width, Height);
 			MainThreadCall(^{
 				CGDisplaySetDisplayMode(WindowedModeSavedState.CapturedDisplayID, DisplayMode, nullptr);
-			}, UE4ResizeEventMode, true);
+			}, UnrealResizeEventMode, true);
 
 			UpdateFullScreenState(bIsFullScreen != bWantsFullScreen);
 
@@ -707,7 +707,7 @@ void FMacWindow::ApplySizeAndModeChanges(int32 X, int32 Y, int32 Width, int32 He
 		MainThreadCall(^{
 			CGDisplayFade(FadeReservationToken, 0.5, kCGDisplayBlendNormal, kCGDisplayBlendSolidColor, 0, 0, 0, false);
 			CGReleaseDisplayFadeReservation(FadeReservationToken);
-		}, UE4ResizeEventMode, false);
+		}, UnrealResizeEventMode, false);
 	}
 }
 
@@ -738,7 +738,7 @@ void FMacWindow::UpdateFullScreenState(bool bToggleFullScreen)
 			[WindowHandle setLevel:WindowedModeSavedState.WindowLevel];
 			[NSApp setPresentationOptions:NSApplicationPresentationDefault];
 		}
-	}, UE4FullscreenEventMode, true);
+	}, UnrealFullscreenEventMode, true);
 
 	// If we toggle fullscreen, ensure that the window has transitioned BEFORE leaving this function.
 	// This prevents problems with failure to correctly update mouse locks and rendering contexts due to bad event ordering.
diff --git a/Engine/Source/Runtime/ApplicationCore/Public/Mac/MacApplication.h b/Engine/Source/Runtime/ApplicationCore/Public/Mac/MacApplication.h
--- a/Engine/Source/Runtime/ApplicationCore/Public/Mac/MacApplication.h
+++ b/Engine/Source/Runtime/ApplicationCore/Public/Mac/MacApplication.h
@@ -303,7 +303,7 @@ private:
 	void CacheKeyboardInputSource();
 
 	void ConditionallyUpdateModifierKeys(const FDeferredMacEvent& Event);
-	void HandleModifierChange(NSUInteger NewModifierFlags, NSUInteger FlagsShift, NSUInteger UE4Shift, EMacModifierKeys TranslatedCode);
+	void HandleModifierChange(NSUInteger NewModifierFlags, NSUInteger FlagsShift, NSUInteger UEShift, EMacModifierKeys TranslatedCode);
 
 	FCocoaWindow* FindEventWindow(NSEvent* CocoaEvent) const;
 	FCocoaWindow* FindSlateWindowUnderCursor() const;
diff --git a/Engine/Source/Runtime/Core/Private/Mac/CocoaThread.cpp b/Engine/Source/Runtime/Core/Private/Mac/CocoaThread.cpp
--- a/Engine/Source/Runtime/Core/Private/Mac/CocoaThread.cpp
+++ b/Engine/Source/Runtime/Core/Private/Mac/CocoaThread.cpp
@@ -5,7 +5,7 @@
 #include "Templates/RefCounting.h"
 #include "Misc/OutputDeviceRedirector.h"
 
-#define MAC_SEPARATE_GAME_THREAD 1 // Separate the main & game threads so that we better handle the interaction between the Cocoa's event delegates and UE4's event polling.
+#define MAC_SEPARATE_GAME_THREAD 1 // Separate the main & game threads so that we better handle the interaction between the Cocoa's event delegates and UE's event polling.
 
 // this is the size of the game thread stack, it must be a multiple of 4k
 #if (UE_BUILD_DEBUG)
@@ -14,12 +14,12 @@
 	#define GAME_THREAD_STACK_SIZE 128 * 1024 * 1024
 #endif
 
-NSString* UE4NilEventMode = @"UE4NilEventMode";
-NSString* UE4ShowEventMode = @"UE4ShowEventMode";
-NSString* UE4ResizeEventMode = @"UE4ResizeEventMode";
-NSString* UE4FullscreenEventMode = @"UE4FullscreenEventMode";
-NSString* UE4CloseEventMode = @"UE4CloseEventMode";
-NSString* UE4IMEEventMode = @"UE4IMEEventMode";
+NSString* UnrealNilEventMode = @"UnrealNilEventMode";
+NSString* UnrealShowEventMode = @"UnrealShowEventMode";
+NSString* UnrealResizeEventMode = @"UnrealResizeEventMode";
+NSString* UnrealFullscreenEventMode = @"UnrealFullscreenEventMode";
+NSString* UnrealCloseEventMode = @"UnrealCloseEventMode";
+NSString* UnrealIMEEventMode = @"UnrealIMEEventMode";
 
 static FCocoaGameThread* GCocoaGameThread = nil;
 static uint32 GMainThreadId = 0;
@@ -148,12 +148,12 @@ private:
 		// Register for default modes
 		Register(kCFRunLoopDefaultMode);
 		Register((CFStringRef)NSModalPanelRunLoopMode);
-		Register((CFStringRef)UE4NilEventMode);
-		Register((CFStringRef)UE4ShowEventMode);
-		Register((CFStringRef)UE4ResizeEventMode);
-		Register((CFStringRef)UE4FullscreenEventMode);
-		Register((CFStringRef)UE4CloseEventMode);
-		Register((CFStringRef)UE4IMEEventMode);
+		Register((CFStringRef)UnrealNilEventMode);
+		Register((CFStringRef)UnrealShowEventMode);
+		Register((CFStringRef)UnrealResizeEventMode);
+		Register((CFStringRef)UnrealFullscreenEventMode);
+		Register((CFStringRef)UnrealCloseEventMode);
+		Register((CFStringRef)UnrealIMEEventMode);
 	}
 	
 	virtual ~FCocoaRunLoopSource()
@@ -382,7 +382,7 @@ FCocoaRunLoopSource* FCocoaRunLoopSource::GameRunLoopSource = nullptr;
 
 	[super main];
 	
-	// We have exited the game thread, so any UE4 code running now should treat the Main thread
+	// We have exited the game thread, so any UE code running now should treat the Main thread
 	// as the game thread, so we don't crash in static destructors.
 	GGameThreadId = GMainThreadId;
 	
diff --git a/Engine/Source/Runtime/Core/Public/Mac/CocoaThread.h b/Engine/Source/Runtime/Core/Public/Mac/CocoaThread.h
--- a/Engine/Source/Runtime/Core/Public/Mac/CocoaThread.h
+++ b/Engine/Source/Runtime/Core/Public/Mac/CocoaThread.h
@@ -3,13 +3,13 @@
 
 #include "CoreMinimal.h"
 
-/* Custom run-loop modes for UE4 that process only certain kinds of events to simulate Windows event ordering. */
-CORE_API extern NSString* UE4NilEventMode; /* Process only mandatory events */
-CORE_API extern NSString* UE4ShowEventMode; /* Process only show window events */
-CORE_API extern NSString* UE4ResizeEventMode; /* Process only resize/move window events */
-CORE_API extern NSString* UE4FullscreenEventMode; /* Process only fullscreen mode events */
-CORE_API extern NSString* UE4CloseEventMode; /* Process only close window events */
-CORE_API extern NSString* UE4IMEEventMode; /* Process only input method events */
+/* Custom run-loop modes for Unreal that process only certain kinds of events to simulate Windows event ordering. */
+CORE_API extern NSString* UnrealNilEventMode; /* Process only mandatory events */
+CORE_API extern NSString* UnrealShowEventMode; /* Process only show window events */
+CORE_API extern NSString* UnrealResizeEventMode; /* Process only resize/move window events */
+CORE_API extern NSString* UnrealFullscreenEventMode; /* Process only fullscreen mode events */
+CORE_API extern NSString* UnrealCloseEventMode; /* Process only close window events */
+CORE_API extern NSString* UnrealIMEEventMode; /* Process only input method events */
 
 @interface NSThread (FCocoaThread)
 + (NSThread*) gameThread; // Returns the main game thread, or nil if has yet to be constructed.
