Remove UE4 references

This removes a variety of low-risk UE4 references.

diff --git a/Engine/Source/Runtime/ApplicationCore/Private/Mac/MacCursor.cpp b/Engine/Source/Runtime/ApplicationCore/Private/Mac/MacCursor.cpp
--- a/Engine/Source/Runtime/ApplicationCore/Private/Mac/MacCursor.cpp
+++ b/Engine/Source/Runtime/ApplicationCore/Private/Mac/MacCursor.cpp
@@ -16,7 +16,7 @@ int32 GMacDisableMouseCoalescing = 1;
 static FAutoConsoleVariableRef CVarMacDisableMouseCoalescing(
 	TEXT("io.Mac.HighPrecisionDisablesMouseCoalescing"),
 	GMacDisableMouseCoalescing,
-	TEXT("If set to true then OS X mouse event coalescing will be disabled while using high-precision mouse mode, to send all mouse events to UE4's event handling routines to reduce apparent mouse lag. (Default: True)"));
+	TEXT("If set to true then OS X mouse event coalescing will be disabled while using high-precision mouse mode, to send all mouse events to Unreal's event handling routines to reduce apparent mouse lag. (Default: True)"));
 
 int32 GMacDisableMouseAcceleration = 0;
 static FAutoConsoleVariableRef CVarMacDisableMouseAcceleration(
diff --git a/Engine/Source/Runtime/ApplicationCore/Private/Mac/MacPlatformSplash.cpp b/Engine/Source/Runtime/ApplicationCore/Private/Mac/MacPlatformSplash.cpp
--- a/Engine/Source/Runtime/ApplicationCore/Private/Mac/MacPlatformSplash.cpp
+++ b/Engine/Source/Runtime/ApplicationCore/Private/Mac/MacPlatformSplash.cpp
@@ -54,7 +54,7 @@ static void StartSetSplashText( const SplashTextType::Type InType, const FText&
 	GSplashScreenText[ InType ] = InText;
 }
 
-@interface UE4SplashView : NSView
+@interface UnrealSplashView : NSView
 {
 }
 
@@ -63,7 +63,7 @@ static void StartSetSplashText( const SplashTextType::Type InType, const FText&
 
 @end
 
-@implementation UE4SplashView
+@implementation UnrealSplashView
 
 - (void)drawRect: (NSRect)DirtyRect
 {
@@ -291,7 +291,7 @@ void FMacPlatformSplash::Show()
 
 			// Create bordeless window with size from NSImage
 			GSplashWindow = [[FSplashWindow alloc] initWithContentRect: ContentRect styleMask: 0 backing: NSBackingStoreBuffered defer: NO];
-			[GSplashWindow setContentView: [[UE4SplashView alloc] initWithFrame: ContentRect]];
+			[GSplashWindow setContentView: [[UnrealSplashView alloc] initWithFrame: ContentRect]];
 
 			if( GSplashWindow )
 			{
