Remove ~15% of allocations during game startup of an internal demo

diff --git a/Engine/Source/Runtime/ApplicationCore/Private/Windows/WindowsPlatformApplicationMisc.cpp b/Engine/Source/Runtime/ApplicationCore/Private/Windows/WindowsPlatformApplicationMisc.cpp
--- a/Engine/Source/Runtime/ApplicationCore/Private/Windows/WindowsPlatformApplicationMisc.cpp
+++ b/Engine/Source/Runtime/ApplicationCore/Private/Windows/WindowsPlatformApplicationMisc.cpp
@@ -109,13 +109,20 @@ static void WinPumpMessages()
 
 void FWindowsPlatformApplicationMisc::PumpMessages(bool bFromMainLoop)
 {
-	TSharedPtr<void> RevertGlobalFlag;
-	if (!GPumpingMessages)
+	const bool bSetPumpingMessages = !GPumpingMessages;
+	if (bSetPumpingMessages)
 	{
 		GPumpingMessages = true;
-		RevertGlobalFlag = MakeShareable<void>(nullptr, [](auto) {GPumpingMessages = false; });
 	}
 
+	ON_SCOPE_EXIT
+	{
+		if (bSetPumpingMessages)
+		{
+			GPumpingMessages = false;
+		}
+	};
+
 	if (!bFromMainLoop)
 	{
 		FPlatformMisc::PumpMessagesOutsideMainLoop();
diff --git a/Engine/Source/Runtime/MoviePlayer/Private/MoviePlayerThreading.cpp b/Engine/Source/Runtime/MoviePlayer/Private/MoviePlayerThreading.cpp
--- a/Engine/Source/Runtime/MoviePlayer/Private/MoviePlayerThreading.cpp
+++ b/Engine/Source/Runtime/MoviePlayer/Private/MoviePlayerThreading.cpp
@@ -75,7 +75,7 @@ void FSlateLoadingSynchronizationMechanism::DestroySlateThread()
 		{
 			FPlatformApplicationMisc::PumpMessages(false);
 
-			FPlatformProcess::Sleep(0.f);
+			FPlatformProcess::Sleep(0.001f);
 		}
 
 		delete SlateLoadingThread;
