Add missing include

#rnx

diff --git a/Engine/Source/Runtime/ApplicationCore/Private/Windows/WindowsPlatformApplicationMisc.cpp b/Engine/Source/Runtime/ApplicationCore/Private/Windows/WindowsPlatformApplicationMisc.cpp
--- a/Engine/Source/Runtime/ApplicationCore/Private/Windows/WindowsPlatformApplicationMisc.cpp
+++ b/Engine/Source/Runtime/ApplicationCore/Private/Windows/WindowsPlatformApplicationMisc.cpp
@@ -8,6 +8,7 @@
 #include "HAL/FeedbackContextAnsi.h"
 #include "Misc/App.h"
 #include "Math/Color.h"
+#include "Misc/ScopeExit.h"
 #include "Windows/WindowsHWrapper.h"
 #include "Modules/ModuleManager.h"
 #include "Misc/CoreDelegates.h"
