fix for errors when compiling against newer Windows SDKs.

#rnx

diff --git a/Engine/Source/Runtime/ApplicationCore/Private/Windows/WindowsWindow.cpp b/Engine/Source/Runtime/ApplicationCore/Private/Windows/WindowsWindow.cpp
--- a/Engine/Source/Runtime/ApplicationCore/Private/Windows/WindowsWindow.cpp
+++ b/Engine/Source/Runtime/ApplicationCore/Private/Windows/WindowsWindow.cpp
@@ -424,18 +424,18 @@ void FWindowsWindow::DisableTouchFeedback()
 	{
 		typedef enum tagWINVER602FEEDBACK_TYPE
 		{
-			FEEDBACK_TOUCH_CONTACTVISUALIZATION = 1,
-			FEEDBACK_PEN_BARRELVISUALIZATION = 2,
-			FEEDBACK_PEN_TAP = 3,
-			FEEDBACK_PEN_DOUBLETAP = 4,
-			FEEDBACK_PEN_PRESSANDHOLD = 5,
-			FEEDBACK_PEN_RIGHTTAP = 6,
-			FEEDBACK_TOUCH_TAP = 7,
-			FEEDBACK_TOUCH_DOUBLETAP = 8,
-			FEEDBACK_TOUCH_PRESSANDHOLD = 9,
-			FEEDBACK_TOUCH_RIGHTTAP = 10,
-			FEEDBACK_GESTURE_PRESSANDTAP = 11,
-			FEEDBACK_MAX = 0xFFFFFFFF
+			WINVER602_FEEDBACK_TOUCH_CONTACTVISUALIZATION = 1,
+			WINVER602_FEEDBACK_PEN_BARRELVISUALIZATION = 2,
+			WINVER602_FEEDBACK_PEN_TAP = 3,
+			WINVER602_FEEDBACK_PEN_DOUBLETAP = 4,
+			WINVER602_FEEDBACK_PEN_PRESSANDHOLD = 5,
+			WINVER602_FEEDBACK_PEN_RIGHTTAP = 6,
+			WINVER602_FEEDBACK_TOUCH_TAP = 7,
+			WINVER602_FEEDBACK_TOUCH_DOUBLETAP = 8,
+			WINVER602_FEEDBACK_TOUCH_PRESSANDHOLD = 9,
+			WINVER602_FEEDBACK_TOUCH_RIGHTTAP = 10,
+			WINVER602_FEEDBACK_GESTURE_PRESSANDTAP = 11,
+			WINVER602_FEEDBACK_MAX = 0xFFFFFFFF
 		} WINVER602FEEDBACK_TYPE;
 
 		typedef BOOL(*SetWindowFeedbackSettingProc)(_In_ HWND hwnd,
@@ -449,9 +449,9 @@ void FWindowsWindow::DisableTouchFeedback()
 		if (SetWindowFeedbackSetting)
 		{
 			BOOL enabled = 0;
-			SetWindowFeedbackSetting(HWnd, FEEDBACK_TOUCH_CONTACTVISUALIZATION, 0, sizeof(enabled), &enabled);
-			SetWindowFeedbackSetting(HWnd, FEEDBACK_TOUCH_TAP, 0, sizeof(enabled), &enabled);
-			SetWindowFeedbackSetting(HWnd, FEEDBACK_TOUCH_PRESSANDHOLD, 0, sizeof(enabled), &enabled);
+			SetWindowFeedbackSetting(HWnd, WINVER602_FEEDBACK_TOUCH_CONTACTVISUALIZATION, 0, sizeof(enabled), &enabled);
+			SetWindowFeedbackSetting(HWnd, WINVER602_FEEDBACK_TOUCH_TAP, 0, sizeof(enabled), &enabled);
+			SetWindowFeedbackSetting(HWnd, WINVER602_FEEDBACK_TOUCH_PRESSANDHOLD, 0, sizeof(enabled), &enabled);
 		}
 	}
 }
diff --git a/Engine/Source/Runtime/Windows/XAudio2/Private/XAudio2Support.h b/Engine/Source/Runtime/Windows/XAudio2/Private/XAudio2Support.h
--- a/Engine/Source/Runtime/Windows/XAudio2/Private/XAudio2Support.h
+++ b/Engine/Source/Runtime/Windows/XAudio2/Private/XAudio2Support.h
@@ -872,25 +872,25 @@ struct FXAudioDeviceProperties final : public IDeviceChangedListener
 		return bChanged;
 	}
 
-	bool Validate(const TCHAR* Function, uint32 ErrorCode) const
+	bool Validate(const TCHAR* Function, HRESULT ErrorCode) const
 	{
 		if (ErrorCode != S_OK)
 		{
 			switch (ErrorCode)
 			{
-				case XAUDIO2_E_INVALID_CALL:
+				case (HRESULT)XAUDIO2_E_INVALID_CALL:
 				UE_LOG(LogAudio, Error, TEXT("%s error: Invalid Call"), Function);
 				break;
 
-				case XAUDIO2_E_XMA_DECODER_ERROR:
+				case (HRESULT)XAUDIO2_E_XMA_DECODER_ERROR:
 				UE_LOG(LogAudio, Error, TEXT("%s error: XMA Decoder Error"), Function);
 				break;
 
-				case XAUDIO2_E_XAPO_CREATION_FAILED:
+				case (HRESULT)XAUDIO2_E_XAPO_CREATION_FAILED:
 				UE_LOG(LogAudio, Error, TEXT("%s error: XAPO Creation Failed"), Function);
 				break;
 
-				case XAUDIO2_E_DEVICE_INVALIDATED:
+				case (HRESULT)XAUDIO2_E_DEVICE_INVALIDATED:
 				UE_LOG(LogAudio, Error, TEXT("%s error: Device Invalidated"), Function);
 				break;
 
