Move UI call to main thread.

#rnx

diff --git a/Engine/Source/Runtime/ApplicationCore/Private/Mac/MacPlatformSplash.cpp b/Engine/Source/Runtime/ApplicationCore/Private/Mac/MacPlatformSplash.cpp
--- a/Engine/Source/Runtime/ApplicationCore/Private/Mac/MacPlatformSplash.cpp
+++ b/Engine/Source/Runtime/ApplicationCore/Private/Mac/MacPlatformSplash.cpp
@@ -359,10 +359,9 @@ void FMacPlatformSplash::SetSplashText(const SplashTextType::Type InType, const
 
 			if( bWasUpdated )
 			{
-				SCOPED_AUTORELEASE_POOL;
-
-				// Repaint the window
-				[[GSplashWindow contentView] setNeedsDisplayInRect: GSplashScreenTextRects[InType]];
+				dispatch_async(dispatch_get_main_queue(), ^{
+					[[GSplashWindow contentView] setNeedsDisplayInRect: GSplashScreenTextRects[InType]];
+				});
 
 				FMacPlatformApplicationMisc::PumpMessages(true);
 			}
