Accommodated SEditableText assumptions about Num Lock and Numpad Keys to match Windows
Fixed Blueprint node movement bug caused by similar assumption in SGraphPanel

diff --git a/Engine/Source/Runtime/ApplicationCore/Private/Linux/LinuxApplication.cpp b/Engine/Source/Runtime/ApplicationCore/Private/Linux/LinuxApplication.cpp
--- a/Engine/Source/Runtime/ApplicationCore/Private/Linux/LinuxApplication.cpp
+++ b/Engine/Source/Runtime/ApplicationCore/Private/Linux/LinuxApplication.cpp
@@ -219,9 +219,36 @@ bool FLinuxApplication::GeneratesKeyCharMessage(const SDL_KeyboardEvent & KeyDow
 		(Sym != SDLK_DOWN && Sym != SDLK_LEFT && Sym != SDLK_RIGHT && Sym != SDLK_UP && Sym != SDLK_DELETE);
 }
 
+// Windows handles translating numpad numbers to arrow keys, but we have to do it manually
+static SDL_Keycode TranslateNumLockKeySyms(const SDL_Keycode KeySym)
+{
+	if ((SDL_GetModState() & KMOD_NUM) == 0)
+	{
+		switch (KeySym)
+		{
+		case SDLK_KP_2:
+			return SDLK_DOWN;
+		case SDLK_KP_4:
+			return SDLK_LEFT;
+		case SDLK_KP_6:
+			return SDLK_RIGHT;
+		case SDLK_KP_8:
+			return SDLK_UP;
+		default:
+			break;
+		}
+	}
+	return KeySym;
+}
+
 static inline uint32 CharCodeFromSDLKeySym(const SDL_Keycode KeySym)
 {
-	if ((KeySym & SDLK_SCANCODE_MASK) != 0)
+	// Mirrors windows returning nonzero char codes for numpad keys
+	if (KeySym == SDLK_KP_2 || KeySym == SDLK_KP_4 || KeySym == SDLK_KP_6 || KeySym == SDLK_KP_8)
+	{
+		return (uint32) KeySym;
+	}
+	else if ((KeySym & SDLK_SCANCODE_MASK) != 0)
 	{
 		return 0;
 	}
@@ -256,8 +283,9 @@ void FLinuxApplication::ProcessDeferredMessage( SDL_Event Event )
 	case SDL_KEYDOWN:
 		{
 			const SDL_KeyboardEvent &KeyEvent = Event.key;
-			const SDL_Keycode KeySym = KeyEvent.keysym.sym;
+			SDL_Keycode KeySym = KeyEvent.keysym.sym;
 			const uint32 CharCode = CharCodeFromSDLKeySym(KeySym);
+			KeySym = TranslateNumLockKeySyms(KeySym);
 			const bool bIsRepeated = KeyEvent.repeat != 0;
 
 			// Text input is now handled in SDL_TEXTINPUT: see below
