Linux now dispatches OnControllerConnectionChange when gamepad is connected/disconnected.

diff --git a/Engine/Source/Runtime/ApplicationCore/Private/Linux/LinuxApplication.cpp b/Engine/Source/Runtime/ApplicationCore/Private/Linux/LinuxApplication.cpp
--- a/Engine/Source/Runtime/ApplicationCore/Private/Linux/LinuxApplication.cpp
+++ b/Engine/Source/Runtime/ApplicationCore/Private/Linux/LinuxApplication.cpp
@@ -2078,6 +2078,8 @@ void FLinuxApplication::AddGameController(int Index)
 			ControllerState.Haptic = nullptr;
 		}
 	}
+
+	FCoreDelegates::OnControllerConnectionChange.Broadcast(true, -1, Id);
 }
 
 void FLinuxApplication::RemoveGameController(SDL_JoystickID Id)
@@ -2097,4 +2099,6 @@ void FLinuxApplication::RemoveGameController(SDL_JoystickID Id)
 
 	SDL_GameControllerClose(ControllerState.Controller);
 	ControllerStates.Remove(Id);
+
+	FCoreDelegates::OnControllerConnectionChange.Broadcast(false, -1, Id);			
 }
