fix some WindowsHWrapper includes for RECT definitions on Microsoft platforms.

#rnx

diff --git a/Engine/Source/Runtime/ApplicationCore/Private/GenericPlatform/GenericPlatformSoftwareCursor.cpp b/Engine/Source/Runtime/ApplicationCore/Private/GenericPlatform/GenericPlatformSoftwareCursor.cpp
--- a/Engine/Source/Runtime/ApplicationCore/Private/GenericPlatform/GenericPlatformSoftwareCursor.cpp
+++ b/Engine/Source/Runtime/ApplicationCore/Private/GenericPlatform/GenericPlatformSoftwareCursor.cpp
@@ -4,7 +4,7 @@
 #include "GenericPlatform/GenericApplication.h"
 
 // Windows has special needs with RECT, and we don't need this class on Windows anyway, so just skip it
-#if !PLATFORM_WINDOWS && !PLATFORM_HOLOLENS
+#if !PLATFORM_MICROSOFT
 
 FGenericPlatformSoftwareCursor::FGenericPlatformSoftwareCursor() 
 	: CurrentType(EMouseCursor::None)
diff --git a/Engine/Source/Runtime/Slate/Private/Framework/Application/SlateUser.cpp b/Engine/Source/Runtime/Slate/Private/Framework/Application/SlateUser.cpp
--- a/Engine/Source/Runtime/Slate/Private/Framework/Application/SlateUser.cpp
+++ b/Engine/Source/Runtime/Slate/Private/Framework/Application/SlateUser.cpp
@@ -6,9 +6,9 @@
 #include "Widgets/SWindow.h"
 #include "Widgets/SWeakWidget.h"
 
-#if PLATFORM_WINDOWS || PLATFORM_HOLOLENS
+#if PLATFORM_MICROSOFT
 // Needed to be able to use RECT
-#include "Windows/WindowsHWrapper.h"
+#include "Microsoft/WindowsHWrapper.h"
 #endif
 
 DECLARE_CYCLE_STAT(TEXT("QueryCursor"), STAT_SlateQueryCursor, STATGROUP_Slate);
