Add insight tag when editor decides to sleep because it is in background

#rnx

diff --git a/Engine/Source/Runtime/ApplicationCore/Private/Windows/WindowsPlatformApplicationMisc.cpp b/Engine/Source/Runtime/ApplicationCore/Private/Windows/WindowsPlatformApplicationMisc.cpp
--- a/Engine/Source/Runtime/ApplicationCore/Private/Windows/WindowsPlatformApplicationMisc.cpp
+++ b/Engine/Source/Runtime/ApplicationCore/Private/Windows/WindowsPlatformApplicationMisc.cpp
@@ -143,6 +143,7 @@ void FWindowsPlatformApplicationMisc::PumpMessages(bool bFromMainLoop)
 	{
 		if( !HasFocus )
 		{
+			TRACE_CPUPROFILER_EVENT_SCOPE("EditorIsInBackgroundSleep");
 			// Sleep for a bit to not eat up all CPU time.
 			FPlatformProcess::Sleep(0.005f);
 		}
