Fix mouse wrapping prematurely over RDP when clicking 10% away from screen edge.

diff --git a/Engine/Source/Runtime/ApplicationCore/Private/Windows/WindowsApplication.cpp b/Engine/Source/Runtime/ApplicationCore/Private/Windows/WindowsApplication.cpp
--- a/Engine/Source/Runtime/ApplicationCore/Private/Windows/WindowsApplication.cpp
+++ b/Engine/Source/Runtime/ApplicationCore/Private/Windows/WindowsApplication.cpp
@@ -1141,8 +1141,10 @@ int32 FWindowsApplication::ProcessMessage( HWND hwnd, uint32 msg, WPARAM wParam,
 								int32 DeltaX = CursorPoint.x - LastCursorPoint.X;
 								int32 DeltaY = CursorPoint.y - LastCursorPoint.Y;
 
-								bool bAcceptingPreWrapDelta = false;
-								if (bCanAcceptPreWrapMsg)
+								// Always skip wrapping if no movement has occured.
+								bool bAcceptingPreWrapDelta = DeltaX == 0 && DeltaY == 0;
+
+								if (bCanAcceptPreWrapMsg && !bAcceptingPreWrapDelta)
 								{
 									const int32 DeltaXPreWrap = CursorPoint.x - LastCursorPointPreWrap.X;
 									const int32 DeltaYPreWrap = CursorPoint.y - LastCursorPointPreWrap.Y;
