Fixed compilation

diff --git a/Engine/Source/Runtime/ApplicationCore/Private/GenericPlatform/GenericPlatformSoftwareCursor.cpp b/Engine/Source/Runtime/ApplicationCore/Private/GenericPlatform/GenericPlatformSoftwareCursor.cpp
--- a/Engine/Source/Runtime/ApplicationCore/Private/GenericPlatform/GenericPlatformSoftwareCursor.cpp
+++ b/Engine/Source/Runtime/ApplicationCore/Private/GenericPlatform/GenericPlatformSoftwareCursor.cpp
@@ -52,10 +52,10 @@ void FGenericPlatformSoftwareCursor::Lock(const RECT* const Bounds)
 	}
 	else
 	{
-		CursorClipRect.Min.X = FMath::TruncToInt(Bounds->left);
-		CursorClipRect.Min.Y = FMath::TruncToInt(Bounds->top);
-		CursorClipRect.Max.X = FMath::TruncToInt(Bounds->right) - 1;
-		CursorClipRect.Max.Y = FMath::TruncToInt(Bounds->bottom) - 1;
+		CursorClipRect.Min.X = Bounds->left;
+		CursorClipRect.Min.Y = Bounds->top;
+		CursorClipRect.Max.X = Bounds->right - 1;
+		CursorClipRect.Max.Y = Bounds->bottom - 1;
 	}
 
 	FVector2D Position = GetPosition();
diff --git a/Engine/Source/Runtime/ApplicationCore/Private/Linux/LinuxCursor.cpp b/Engine/Source/Runtime/ApplicationCore/Private/Linux/LinuxCursor.cpp
--- a/Engine/Source/Runtime/ApplicationCore/Private/Linux/LinuxCursor.cpp
+++ b/Engine/Source/Runtime/ApplicationCore/Private/Linux/LinuxCursor.cpp
@@ -311,11 +311,11 @@ static SDL_Rect GetConfinementRect(const RECT* const Bounds, const TSharedPtr<FL
 	int Bottom = 0;
 	SDL_GetWindowBordersSize(CurrentFocusWindow->GetHWnd(), &Top, &Left, &Bottom, &Right);
 
-	ConfineRect.x = FMath::TruncToInt(Bounds->left) + Left;
-	ConfineRect.y = FMath::TruncToInt(Bounds->top) - Top;
+	ConfineRect.x = Bounds->left + Left;
+	ConfineRect.y = Bounds->top - Top;
 
-	ConfineRect.w = FMath::TruncToInt(Bounds->right) - FMath::TruncToInt(Bounds->left) - 1;
-	ConfineRect.h = FMath::TruncToInt(Bounds->bottom) - FMath::TruncToInt(Bounds->top) - 1;
+	ConfineRect.w = Bounds->right - Bounds->left - 1;
+	ConfineRect.h = Bounds->bottom - Bounds->top - 1;
 
 	// TODO. For some reason the values from Bounds seem not to cover the SDL's window area.
 	// The cursor exceeds the right and bottom border of the window. This can be troublesome
diff --git a/Engine/Source/Runtime/ApplicationCore/Private/Mac/MacCursor.cpp b/Engine/Source/Runtime/ApplicationCore/Private/Mac/MacCursor.cpp
--- a/Engine/Source/Runtime/ApplicationCore/Private/Mac/MacCursor.cpp
+++ b/Engine/Source/Runtime/ApplicationCore/Private/Mac/MacCursor.cpp
@@ -363,10 +363,10 @@ void FMacCursor::Lock(const RECT* const Bounds)
 	}
 	else
 	{
-		CursorClipRect.Min.X = FMath::TruncToInt(Bounds->left);
-		CursorClipRect.Min.Y = FMath::TruncToInt(Bounds->top);
-		CursorClipRect.Max.X = Bounds->right > Bounds->left ? FMath::TruncToInt(Bounds->right) - 1 : Bounds->left;
-		CursorClipRect.Max.Y = Bounds->bottom > Bounds->top ? FMath::TruncToInt(Bounds->bottom) - 1 : Bounds->top;
+		CursorClipRect.Min.X = Bounds->left;
+		CursorClipRect.Min.Y = Bounds->top;
+		CursorClipRect.Max.X = Bounds->right > Bounds->left ? Bounds->right - 1 : Bounds->left;
+		CursorClipRect.Max.Y = Bounds->bottom > Bounds->top ? Bounds->bottom - 1 : Bounds->top;
 	}
 
 	MacApplication->OnCursorLock();
