Fixed the cpu scoped timer name for EditorIsInBackgroundSleep to not include quotes in name (i.e. it does not need quotes when using TRACE_CPUPROFILER_EVENT_SCOPE macro; only the _STR variant of the macro accepts a static string).

diff --git a/Engine/Source/Runtime/ApplicationCore/Private/Windows/WindowsPlatformApplicationMisc.cpp b/Engine/Source/Runtime/ApplicationCore/Private/Windows/WindowsPlatformApplicationMisc.cpp
--- a/Engine/Source/Runtime/ApplicationCore/Private/Windows/WindowsPlatformApplicationMisc.cpp
+++ b/Engine/Source/Runtime/ApplicationCore/Private/Windows/WindowsPlatformApplicationMisc.cpp
@@ -144,7 +144,7 @@ void FWindowsPlatformApplicationMisc::PumpMessages(bool bFromMainLoop)
 		bool bBenchmarking = FApp::IsBenchmarking();
 		if( !bHasFocus && !bBenchmarking )
 		{
-			TRACE_CPUPROFILER_EVENT_SCOPE("EditorIsInBackgroundSleep");
+			TRACE_CPUPROFILER_EVENT_SCOPE(EditorIsInBackgroundSleep);
 			// Sleep for a bit to not eat up all CPU time.
 			FPlatformProcess::Sleep(0.005f);
 		}
