Prevent accidiental closing of application windows by double clicking on the tile bar icon (default win32 behavior only)

diff --git a/Engine/Source/Runtime/ApplicationCore/Private/Windows/WindowsWindow.cpp b/Engine/Source/Runtime/ApplicationCore/Private/Windows/WindowsWindow.cpp
--- a/Engine/Source/Runtime/ApplicationCore/Private/Windows/WindowsWindow.cpp
+++ b/Engine/Source/Runtime/ApplicationCore/Private/Windows/WindowsWindow.cpp
@@ -290,6 +290,10 @@ void FWindowsWindow::Initialize( FWindowsApplication* const Application, const T
 		}
 
 		::SetWindowPos(HWnd, nullptr, 0, 0, 0, 0, SetWindowPositionFlags);
+		
+		// For regular non-game windows delete the close menu from the default system menu. This prevents accidental closing of win32 apps by double clicking by accident on the application icon
+		// The overwhelming majority of feedback is that is confusing behavior and we want to prevent this.
+		::DeleteMenu(GetSystemMenu(HWnd, false), SC_CLOSE, MF_BYCOMMAND);
 
 		AdjustWindowRegion( ClientWidth, ClientHeight );
 	}
