Fix for crash in FMacApplication::FindScreenByCocoaPosition().   FindScreenByCocoaPosition/FindScreenBySlatePosition can be accessed externally from any thread.   Mouse move (and other event) updates in the Engine Game Thread were causing a race condition with an external plugin Main Thread event processing.   Update FMacScreen references to ESPMode::ThreadSafe and use C++ references where is makes sense - even though UE5 has now since defaulted to thread safe - still changing this just in case it goes back to non thread safe and we want to be explicit about it in the mac screen handling.

diff --git a/Engine/Source/Runtime/ApplicationCore/Private/Mac/MacApplication.cpp b/Engine/Source/Runtime/ApplicationCore/Private/Mac/MacApplication.cpp
--- a/Engine/Source/Runtime/ApplicationCore/Private/Mac/MacApplication.cpp
+++ b/Engine/Source/Runtime/ApplicationCore/Private/Mac/MacApplication.cpp
@@ -34,7 +34,7 @@
 FMacApplication* MacApplication = nullptr;
 
 static FCriticalSection GAllScreensMutex;
-TArray<TSharedRef<FMacScreen>> FMacApplication::AllScreens;
+static TArray<FMacScreenRef> GAllScreens;
 
 const uint32 RESET_EVENT_SUBTYPE = 0x0f00;
 
@@ -412,7 +412,7 @@ FPlatformRect FMacApplication::GetWorkArea(const FPlatformRect& CurrentWindow) c
 {
 	SCOPED_AUTORELEASE_POOL;
 
-	TSharedRef<FMacScreen> Screen = FindScreenBySlatePosition(CurrentWindow.Left, CurrentWindow.Top);
+	FMacScreenRef Screen = FindScreenBySlatePosition(CurrentWindow.Left, CurrentWindow.Top);
 
 	const NSRect VisibleFrame = Screen->VisibleFramePixels;
 
@@ -905,7 +905,7 @@ void FMacApplication::ProcessMouseMovedEvent(const FDeferredMacEvent& Event, TSh
 		FVector2D HighPrecisionMousePos = MacCursor->GetPosition();
 
 		// Find the visible frame of the screen the cursor is currently on.
-		TSharedRef<FMacScreen> Screen = FindScreenBySlatePosition(HighPrecisionMousePos.X, HighPrecisionMousePos.Y);
+		FMacScreenRef Screen = FindScreenBySlatePosition(HighPrecisionMousePos.X, HighPrecisionMousePos.Y);
 		NSRect VisibleFrame = Screen->VisibleFramePixels;
 
 		// Under OS X we disassociate the cursor and mouse position during hi-precision mouse input.
@@ -1692,7 +1692,7 @@ void FMacApplication::SetForceFeedbackChannelValues(int32 ControllerId, const FF
 	{
 		if (InputDevice.IsValid())
 		{
-			// Mirrored from the Window's impl: "Ideally, we would want to use 
+			// Mirrored from the Window's impl: "Ideally, we would want to use
 			// GetHapticDevice instead but they're not implemented for SteamController"
 			if (InputDevice->IsGamepadAttached())
 			{
@@ -1727,26 +1727,26 @@ void FMacApplication::UpdateScreensArray()
 	MainThreadCall(^{
 		SCOPED_AUTORELEASE_POOL;
 		FScopeLock Lock(&GAllScreensMutex);
-		AllScreens.Empty();
+		GAllScreens.Empty();
 		NSArray* Screens = [NSScreen screens];
 		for (NSScreen* Screen in Screens)
 		{
-			AllScreens.Add(MakeShareable(new FMacScreen(Screen)));
+			GAllScreens.Add(MakeShareable(new FMacScreen(Screen)));
 		}
 	}, NSDefaultRunLoopMode, true);
 
 	FScopeLock Lock(&GAllScreensMutex);
 
 	NSRect WholeWorkspace = {{0, 0}, {0, 0}};
-	for (TSharedRef<FMacScreen> CurScreen : AllScreens)
+	for (FMacScreenRef& CurScreen : GAllScreens)
 	{
 		WholeWorkspace = NSUnionRect(WholeWorkspace, CurScreen->Frame);
 	}
 
 	const bool bUseHighDPIMode = FPlatformApplicationMisc::IsHighDPIModeEnabled();
 
-	TArray<TSharedRef<FMacScreen>> SortedScreens;
-	for (TSharedRef<FMacScreen> CurScreen : AllScreens)
+	TArray<FMacScreenRef> SortedScreens;
+	for (FMacScreenRef& CurScreen : GAllScreens)
 	{
 		CurScreen->Frame.origin.y = CurScreen->FramePixels.origin.y = WholeWorkspace.origin.y + WholeWorkspace.size.height - CurScreen->Frame.size.height - CurScreen->Frame.origin.y;
 		CurScreen->VisibleFrame.origin.y = CurScreen->VisibleFramePixels.origin.y = WholeWorkspace.origin.y + WholeWorkspace.size.height - CurScreen->VisibleFrame.size.height - CurScreen->VisibleFrame.origin.y;
@@ -1754,11 +1754,11 @@ void FMacApplication::UpdateScreensArray()
 		SortedScreens.Add(CurScreen);
 	}
 
-	SortedScreens.Sort([](const TSharedRef<FMacScreen>& A, const TSharedRef<FMacScreen>& B) -> bool { return A->Frame.origin.x < B->Frame.origin.x; });
+	SortedScreens.Sort([](const FMacScreenRef& A, const FMacScreenRef& B) -> bool { return A->Frame.origin.x < B->Frame.origin.x; });
 
 	for (int32 Index = 0; Index < SortedScreens.Num(); ++Index)
 	{
-		TSharedRef<FMacScreen> CurScreen = SortedScreens[Index];
+		FMacScreenRef& CurScreen = SortedScreens[Index];
 		const float DPIScaleFactor = bUseHighDPIMode ? CurScreen->Screen.backingScaleFactor : 1.0f;
 		if (DPIScaleFactor != 1.0f)
 		{
@@ -1769,7 +1769,7 @@ void FMacApplication::UpdateScreensArray()
 
 			for (int32 OtherIndex = Index + 1; OtherIndex < SortedScreens.Num(); ++OtherIndex)
 			{
-				TSharedRef<FMacScreen> OtherScreen = SortedScreens[OtherIndex];
+				FMacScreenRef& OtherScreen = SortedScreens[OtherIndex];
 				const float DiffFrame = (OtherScreen->Frame.origin.x - CurScreen->Frame.origin.x) * DPIScaleFactor;
 				const float DiffVisibleFrame = (OtherScreen->VisibleFrame.origin.x - CurScreen->VisibleFrame.origin.x) * DPIScaleFactor;
 				OtherScreen->FramePixels.origin.x = CurScreen->FramePixels.origin.x + DiffFrame;
@@ -1778,17 +1778,17 @@ void FMacApplication::UpdateScreensArray()
 		}
 	}
 
-	SortedScreens.Sort([](const TSharedRef<FMacScreen>& A, const TSharedRef<FMacScreen>& B) -> bool { return A->Frame.origin.y < B->Frame.origin.y; });
+	SortedScreens.Sort([](const FMacScreenRef& A, const FMacScreenRef& B) -> bool { return A->Frame.origin.y < B->Frame.origin.y; });
 
 	for (int32 Index = 0; Index < SortedScreens.Num(); ++Index)
 	{
-		TSharedRef<FMacScreen> CurScreen = SortedScreens[Index];
+		FMacScreenRef& CurScreen = SortedScreens[Index];
 		const float DPIScaleFactor = bUseHighDPIMode ? CurScreen->Screen.backingScaleFactor : 1.0f;
 		if (DPIScaleFactor != 1.0f)
 		{
 			for (int32 OtherIndex = Index + 1; OtherIndex < SortedScreens.Num(); ++OtherIndex)
 			{
-				TSharedRef<FMacScreen> OtherScreen = SortedScreens[OtherIndex];
+				FMacScreenRef& OtherScreen = SortedScreens[OtherIndex];
 				const float DiffFrame = (OtherScreen->Frame.origin.y - CurScreen->Frame.origin.y) * DPIScaleFactor;
 				const float DiffVisibleFrame = (OtherScreen->VisibleFrame.origin.y - CurScreen->VisibleFrame.origin.y) * DPIScaleFactor;
 				OtherScreen->FramePixels.origin.y = CurScreen->FramePixels.origin.y + DiffFrame;
@@ -1798,10 +1798,10 @@ void FMacApplication::UpdateScreensArray()
 	}
 
 	// The primary screen needs to be at (0,0), so we need to offset all screen origins by its position
-	TSharedRef<FMacScreen> PrimaryScreen = AllScreens[0];
+	FMacScreenRef& PrimaryScreen = GAllScreens[0];
 	const FVector2D FrameOffset(PrimaryScreen->Frame.origin.x, PrimaryScreen->Frame.origin.y);
 	const FVector2D FramePixelsOffset(PrimaryScreen->FramePixels.origin.x, PrimaryScreen->FramePixels.origin.y);
-	for (TSharedRef<FMacScreen> CurScreen : AllScreens)
+	for (FMacScreenRef& CurScreen : GAllScreens)
 	{
 		CurScreen->Frame.origin.x -= FrameOffset.X;
 		CurScreen->Frame.origin.y -= FrameOffset.Y;
@@ -1818,8 +1818,8 @@ FVector2D FMacApplication::CalculateScreenOrigin(NSScreen* Screen)
 {
 	NSRect WholeWorkspace = {{0, 0}, {0, 0}};
 	NSRect ScreenFrame = {{0, 0}, {0, 0}};
-	GAllScreensMutex.Lock();
-	for (TSharedRef<FMacScreen> CurScreen : AllScreens)
+	FScopeLock Lock(&GAllScreensMutex);
+	for (FMacScreenRef& CurScreen : GAllScreens)
 	{
 		WholeWorkspace = NSUnionRect(WholeWorkspace, CurScreen->FramePixels);
 		if (Screen == CurScreen->Screen)
@@ -1827,7 +1827,6 @@ FVector2D FMacApplication::CalculateScreenOrigin(NSScreen* Screen)
 			ScreenFrame = CurScreen->FramePixels;
 		}
 	}
-	GAllScreensMutex.Unlock();
 	return FVector2D(ScreenFrame.origin.x, WholeWorkspace.size.height - ScreenFrame.size.height - ScreenFrame.origin.y);
 }
 
@@ -1835,17 +1834,17 @@ float FMacApplication::GetPrimaryScreenBackingScaleFactor()
 {
 	FScopeLock Lock(&GAllScreensMutex);
 	const bool bUseHighDPIMode = FPlatformApplicationMisc::IsHighDPIModeEnabled();
-	return bUseHighDPIMode ? AllScreens[0]->Screen.backingScaleFactor : 1.0f;
+	return bUseHighDPIMode ? GAllScreens[0]->Screen.backingScaleFactor : 1.0f;
 }
 
-TSharedRef<FMacScreen> FMacApplication::FindScreenBySlatePosition(float X, float Y)
+FMacScreenRef FMacApplication::FindScreenBySlatePosition(float X, float Y)
 {
 	NSPoint Point = NSMakePoint(X, Y);
 
 	FScopeLock Lock(&GAllScreensMutex);
 
-	TSharedRef<FMacScreen> TargetScreen = AllScreens[0];
-	for (TSharedRef<FMacScreen> Screen : AllScreens)
+	FMacScreenRef TargetScreen = GAllScreens[0];
+	for (FMacScreenRef& Screen : GAllScreens)
 	{
 		if (NSPointInRect(Point, Screen->FramePixels))
 		{
@@ -1857,14 +1856,14 @@ TSharedRef<FMacScreen> FMacApplication::FindScreenBySlatePosition(float X, float
 	return TargetScreen;
 }
 
-TSharedRef<FMacScreen> FMacApplication::FindScreenByCocoaPosition(float X, float Y)
+FMacScreenRef FMacApplication::FindScreenByCocoaPosition(float X, float Y)
 {
 	NSPoint Point = NSMakePoint(X, Y);
 
 	FScopeLock Lock(&GAllScreensMutex);
 
-	TSharedRef<FMacScreen> TargetScreen = AllScreens[0];
-	for (TSharedRef<FMacScreen> Screen : AllScreens)
+	FMacScreenRef TargetScreen = GAllScreens[0];
+	for (FMacScreenRef& Screen : GAllScreens)
 	{
 		if (NSPointInRect(Point, Screen->Screen.frame))
 		{
@@ -1878,7 +1877,7 @@ TSharedRef<FMacScreen> FMacApplication::FindScreenByCocoaPosition(float X, float
 
 FVector2D FMacApplication::ConvertSlatePositionToCocoa(float X, float Y)
 {
-	TSharedRef<FMacScreen> Screen = FindScreenBySlatePosition(X, Y);
+	FMacScreenRef Screen = FindScreenBySlatePosition(X, Y);
 	const bool bUseHighDPIMode = FPlatformApplicationMisc::IsHighDPIModeEnabled();
 	const float DPIScaleFactor = bUseHighDPIMode ? Screen->Screen.backingScaleFactor : 1.0f;
 	const FVector2D OffsetOnScreen = FVector2D(X - Screen->FramePixels.origin.x, Screen->FramePixels.origin.y + Screen->FramePixels.size.height - Y) / DPIScaleFactor;
@@ -1887,7 +1886,7 @@ FVector2D FMacApplication::ConvertSlatePositionToCocoa(float X, float Y)
 
 FVector2D FMacApplication::ConvertCocoaPositionToSlate(float X, float Y)
 {
-	TSharedRef<FMacScreen> Screen = FindScreenByCocoaPosition(X, Y);
+	FMacScreenRef Screen = FindScreenByCocoaPosition(X, Y);
 	const bool bUseHighDPIMode = FPlatformApplicationMisc::IsHighDPIModeEnabled();
 	const float DPIScaleFactor = bUseHighDPIMode ? Screen->Screen.backingScaleFactor : 1.0f;
 	const FVector2D OffsetOnScreen = FVector2D(X - Screen->Screen.frame.origin.x, Screen->Screen.frame.origin.y + Screen->Screen.frame.size.height - Y) * DPIScaleFactor;
@@ -1896,7 +1895,7 @@ FVector2D FMacApplication::ConvertCocoaPositionToSlate(float X, float Y)
 
 CGPoint FMacApplication::ConvertSlatePositionToCGPoint(float X, float Y)
 {
-	TSharedRef<FMacScreen> Screen = FindScreenBySlatePosition(X, Y);
+	FMacScreenRef Screen = FindScreenBySlatePosition(X, Y);
 	const bool bUseHighDPIMode = FPlatformApplicationMisc::IsHighDPIModeEnabled();
 	const float DPIScaleFactor = bUseHighDPIMode ? Screen->Screen.backingScaleFactor : 1.0f;
 	const FVector2D OffsetOnScreen = FVector2D(X - Screen->FramePixels.origin.x, Y - Screen->FramePixels.origin.y) / DPIScaleFactor;
@@ -2019,7 +2018,7 @@ unichar FMacApplication::TranslateKeyCodeToUniCode(uint32 KeyCode, uint32 Modifi
 	// Some just don't work as expected
 	switch(KeyCode)
 	{
-		case kVK_PageUp:	return NSPageUpFunctionKey; 
+		case kVK_PageUp:	return NSPageUpFunctionKey;
 		case kVK_PageDown:	return NSPageDownFunctionKey;
 		case kVK_End:		return NSEndFunctionKey;
 		case kVK_Home:		return NSHomeFunctionKey;
@@ -2174,8 +2173,7 @@ void FDisplayMetrics::RebuildDisplayMetrics(FDisplayMetrics& OutDisplayMetrics)
 
 	FScopeLock Lock(&GAllScreensMutex);
 
-	const TArray<TSharedRef<FMacScreen>>& AllScreens = FMacApplication::GetAllScreens();
-	TSharedRef<FMacScreen> PrimaryScreen = AllScreens[0];
+	FMacScreenRef& PrimaryScreen = GAllScreens[0];
 
 	const NSRect ScreenFrame = PrimaryScreen->FramePixels;
 	const NSRect VisibleFrame = PrimaryScreen->VisibleFramePixels;
@@ -2187,7 +2185,7 @@ void FDisplayMetrics::RebuildDisplayMetrics(FDisplayMetrics& OutDisplayMetrics)
 	OutDisplayMetrics.MonitorInfo.Empty();
 
 	NSRect WholeWorkspace = {{0,0},{0,0}};
-	for (TSharedRef<FMacScreen> Screen : AllScreens)
+	for (FMacScreenRef& Screen : GAllScreens)
 	{
 		WholeWorkspace = NSUnionRect(WholeWorkspace, Screen->FramePixels);
 
@@ -2359,7 +2357,7 @@ void FMacApplication::OnAccessibleEventRaised(TSharedRef<IAccessibleWidget> Widg
 					// will never be received by OSX for whatever reason
 					NSAccessibilityPostNotificationWithUserInfo(NSApp.mainWindow, NSAccessibilityAnnouncementRequestedNotification, AnnouncementInfo);
 					AccessibilityAnnouncementDelayTimer = Nil;
-				}];	
+				}];
 			}, NSDefaultRunLoopMode, false);
 			break;
 		}
diff --git a/Engine/Source/Runtime/ApplicationCore/Private/Mac/MacPlatformApplicationMisc.cpp b/Engine/Source/Runtime/ApplicationCore/Private/Mac/MacPlatformApplicationMisc.cpp
--- a/Engine/Source/Runtime/ApplicationCore/Private/Mac/MacPlatformApplicationMisc.cpp
+++ b/Engine/Source/Runtime/ApplicationCore/Private/Mac/MacPlatformApplicationMisc.cpp
@@ -440,7 +440,7 @@ float FMacPlatformApplicationMisc::GetDPIScaleFactorAtPoint(float X, float Y)
 {
 	if (MacApplication && FPlatformApplicationMisc::IsHighDPIModeEnabled())
 	{
-		TSharedRef<FMacScreen> Screen = FMacApplication::FindScreenBySlatePosition(X, Y);
+		FMacScreenRef Screen = FMacApplication::FindScreenBySlatePosition(X, Y);
 		return Screen->Screen.backingScaleFactor;
 	}
 	return 1.0f;
diff --git a/Engine/Source/Runtime/ApplicationCore/Private/Mac/MacWindow.cpp b/Engine/Source/Runtime/ApplicationCore/Private/Mac/MacWindow.cpp
--- a/Engine/Source/Runtime/ApplicationCore/Private/Mac/MacWindow.cpp
+++ b/Engine/Source/Runtime/ApplicationCore/Private/Mac/MacWindow.cpp
@@ -39,7 +39,7 @@ void FMacWindow::Initialize( FMacApplication* const Application, const TSharedRe
 	// Finally, let's initialize the new native window object.  Calling this function will often cause OS
 	// window messages to be sent! (such as activation messages)
 
-	TSharedRef<FMacScreen> TargetScreen = FMacApplication::FindScreenBySlatePosition(Definition->XDesiredPositionOnScreen, Definition->YDesiredPositionOnScreen);
+	FMacScreenRef TargetScreen = FMacApplication::FindScreenBySlatePosition(Definition->XDesiredPositionOnScreen, Definition->YDesiredPositionOnScreen);
 
 	const int32 SizeX = FMath::Max(FMath::CeilToInt( Definition->WidthDesiredOnScreen ), 1);
 	const int32 SizeY = FMath::Max(FMath::CeilToInt( Definition->HeightDesiredOnScreen ), 1);
diff --git a/Engine/Source/Runtime/ApplicationCore/Public/Mac/MacApplication.h b/Engine/Source/Runtime/ApplicationCore/Public/Mac/MacApplication.h
--- a/Engine/Source/Runtime/ApplicationCore/Public/Mac/MacApplication.h
+++ b/Engine/Source/Runtime/ApplicationCore/Public/Mac/MacApplication.h
@@ -123,6 +123,7 @@ struct FMacScreen
 	FMacScreen(NSScreen* InScreen) : Screen([InScreen retain]), Frame(InScreen.frame), VisibleFrame(InScreen.visibleFrame), FramePixels(InScreen.frame), VisibleFramePixels(InScreen.visibleFrame) {}
 	~FMacScreen() { [Screen release]; }
 };
+typedef TSharedRef<FMacScreen, ESPMode::ThreadSafe> FMacScreenRef;
 
 /**
  * Mac-specific application implementation.
@@ -248,11 +249,9 @@ public:
 
 	static void UpdateScreensArray();
 
-	static const TArray<TSharedRef<FMacScreen>>& GetAllScreens() { return AllScreens; }
+	static FMacScreenRef FindScreenBySlatePosition(float X, float Y);
 
-	static TSharedRef<FMacScreen> FindScreenBySlatePosition(float X, float Y);
-
-	static TSharedRef<FMacScreen> FindScreenByCocoaPosition(float X, float Y);
+	static FMacScreenRef FindScreenByCocoaPosition(float X, float Y);
 
 	static FVector2D ConvertSlatePositionToCocoa(float X, float Y);
 
@@ -380,8 +379,6 @@ private:
 
 	TArray<FCocoaWindow*> WindowsRequiringTextInvalidation;
 
-	static TArray<TSharedRef<FMacScreen>> AllScreens;
-
 	TSharedPtr<FMacTextInputMethodSystem> TextInputMethodSystem;
 
 	bool bIsWorkspaceSessionActive;
