FApp::HasFocus() API to unify checking focus state

This is a defensive change to more uniformly gate conditions of various pathways that want to "nice" CPU consumption when the application is not the foreground application (doesn't "Have Focus"). In benchmarking scenarios we don't want a focus-grabbing application to cause change in behavior while taking performance measurements. This could be in local workstation benchmarking where other foreground apps will naturally come to the front or in data center build where antics that might otherwise be completely unknowable could impact benchmark measurements.

This change was formulated after finding some "nice" CPU conditionals that were instrumented, but not all of them are. This change aims to prevent having to instrument, profile, and mitigate each condition separately.

#rnx

diff --git a/Engine/Plugins/FX/Niagara/Source/Niagara/Private/NiagaraWorldManager.cpp b/Engine/Plugins/FX/Niagara/Source/Niagara/Private/NiagaraWorldManager.cpp
--- a/Engine/Plugins/FX/Niagara/Source/Niagara/Private/NiagaraWorldManager.cpp
+++ b/Engine/Plugins/FX/Niagara/Source/Niagara/Private/NiagaraWorldManager.cpp
@@ -931,7 +931,7 @@ void FNiagaraWorldManager::Tick(ETickingGroup TickGroup, float DeltaSeconds, ELe
 #if WITH_EDITOR //PLATFORM_DESKTOP
 		{
 			QUICK_SCOPE_CYCLE_COUNTER(STAT_Niagara_IsThisApplicationForeground);
-			bAppHasFocus = FPlatformApplicationMisc::IsThisApplicationForeground();
+			bAppHasFocus = FApp::HasFocus();
 		}
 #else
 		bAppHasFocus = true;
diff --git a/Engine/Source/Editor/UnrealEd/Private/EditorEngine.cpp b/Engine/Source/Editor/UnrealEd/Private/EditorEngine.cpp
--- a/Engine/Source/Editor/UnrealEd/Private/EditorEngine.cpp
+++ b/Engine/Source/Editor/UnrealEd/Private/EditorEngine.cpp
@@ -1510,7 +1510,7 @@ void UEditorEngine::Tick( float DeltaSeconds, bool bIdleMode )
 	}
 
 	// Find out if the editor has focus. Audio should only play if the editor has focus.
-	const bool bHasFocus = FPlatformApplicationMisc::IsThisApplicationForeground();
+	const bool bHasFocus = FApp::HasFocus();
 
 	if (bHasFocus || GetDefault<ULevelEditorMiscSettings>()->bAllowBackgroundAudio)
 	{
@@ -6083,9 +6083,9 @@ bool UEditorEngine::ShouldThrottleCPUUsage() const
 
 	const bool bRunningCommandlet = IsRunningCommandlet();
 
-	const bool bIsForeground = FPlatformApplicationMisc::IsThisApplicationForeground();
+	const bool bHasFocus = FApp::HasFocus();
 
-	if( !bIsForeground && !bRunningCommandlet )
+	if( !bHasFocus && !bRunningCommandlet )
 	{
 		const UEditorPerformanceSettings* Settings = GetDefault<UEditorPerformanceSettings>();
 		bShouldThrottle = Settings->bThrottleCPUWhenNotForeground;
diff --git a/Engine/Source/Editor/UnrealEd/Private/FeedbackContextEditor.cpp b/Engine/Source/Editor/UnrealEd/Private/FeedbackContextEditor.cpp
--- a/Engine/Source/Editor/UnrealEd/Private/FeedbackContextEditor.cpp
+++ b/Engine/Source/Editor/UnrealEd/Private/FeedbackContextEditor.cpp
@@ -458,15 +458,15 @@ void FFeedbackContextEditor::StartSlowTask( const FText& Task, bool bShowCancelB
 					OnCancelClicked = FOnCancelClickedDelegate::CreateRaw(this, &FFeedbackContextEditor::OnUserCancel);
 				}
 
-				const bool bFocusAndActivate = FPlatformApplicationMisc::IsThisApplicationForeground();
+				const bool bFocus = FApp::HasFocus();
 				FReceiveUserCancelDelegate ReceiveUserCancelDelegate = FReceiveUserCancelDelegate::CreateRaw(this, &FFeedbackContextEditor::ReceivedUserCancel);
 				TSharedRef<SWindow> SlowTaskWindowRef = SNew(SWindow)
 					.SizingRule(ESizingRule::Autosized)
 					.AutoCenter(EAutoCenter::PreferredWorkArea)
 					.IsPopupWindow(true)
 					.CreateTitleBar(true)
-					.ActivationPolicy(bFocusAndActivate ? EWindowActivationPolicy::Always : EWindowActivationPolicy::Never)
-					.FocusWhenFirstShown(bFocusAndActivate);
+					.ActivationPolicy(bFocus ? EWindowActivationPolicy::Always : EWindowActivationPolicy::Never)
+					.FocusWhenFirstShown(bFocus);
 
 				SlowTaskWindowRef->SetContent(
 					SNew(SSlowTaskWidget)
diff --git a/Engine/Source/Runtime/ApplicationCore/Private/Linux/LinuxPlatformApplicationMisc.cpp b/Engine/Source/Runtime/ApplicationCore/Private/Linux/LinuxPlatformApplicationMisc.cpp
--- a/Engine/Source/Runtime/ApplicationCore/Private/Linux/LinuxPlatformApplicationMisc.cpp
+++ b/Engine/Source/Runtime/ApplicationCore/Private/Linux/LinuxPlatformApplicationMisc.cpp
@@ -251,6 +251,7 @@ uint32 FLinuxPlatformApplicationMisc::WindowStyle()
 void FLinuxPlatformApplicationMisc::PreInit()
 {
 	MessageBoxExtCallback = MessageBoxExtImpl;
+	FApp::SetHasFocusFunction(&FLinuxPlatformApplicationMisc::IsThisApplicationForeground);
 }
 
 void FLinuxPlatformApplicationMisc::Init()
@@ -441,7 +442,7 @@ void FLinuxPlatformApplicationMisc::PumpMessages( bool bFromMainLoop )
 			}
 		}
 
-		bool bHasFocus = FApp::UseVRFocus() ? FApp::HasVRFocus() : FLinuxPlatformApplicationMisc::IsThisApplicationForeground();
+		bool bHasFocus = FApp::HasFocus();
 
 		// if its our window, allow sound, otherwise apply multiplier
 		FApp::SetVolumeMultiplier( bHasFocus ? 1.0f : FApp::GetUnfocusedVolumeMultiplier() );
diff --git a/Engine/Source/Runtime/ApplicationCore/Private/Mac/MacPlatformApplicationMisc.cpp b/Engine/Source/Runtime/ApplicationCore/Private/Mac/MacPlatformApplicationMisc.cpp
--- a/Engine/Source/Runtime/ApplicationCore/Private/Mac/MacPlatformApplicationMisc.cpp
+++ b/Engine/Source/Runtime/ApplicationCore/Private/Mac/MacPlatformApplicationMisc.cpp
@@ -277,6 +277,7 @@ void FMacPlatformApplicationMisc::PreInit()
 
 	FMacApplication::UpdateScreensArray();
 	MessageBoxExtCallback = MessageBoxExtImpl;
+	FApp::SetHasFocusFunction(&FMacPlatformApplicationMisc::IsThisApplicationForeground);
 }
 
 void FMacPlatformApplicationMisc::PostInit()
diff --git a/Engine/Source/Runtime/ApplicationCore/Private/Windows/WindowsPlatformApplicationMisc.cpp b/Engine/Source/Runtime/ApplicationCore/Private/Windows/WindowsPlatformApplicationMisc.cpp
--- a/Engine/Source/Runtime/ApplicationCore/Private/Windows/WindowsPlatformApplicationMisc.cpp
+++ b/Engine/Source/Runtime/ApplicationCore/Private/Windows/WindowsPlatformApplicationMisc.cpp
@@ -32,6 +32,11 @@ THIRD_PARTY_INCLUDES_END
 typedef HRESULT(STDAPICALLTYPE *GetDpiForMonitorProc)(HMONITOR Monitor, int32 DPIType, uint32 *DPIX, uint32 *DPIY);
 APPLICATIONCORE_API GetDpiForMonitorProc GetDpiForMonitor;
 
+void FWindowsPlatformApplicationMisc::PreInit()
+{
+	FApp::SetHasFocusFunction(&FWindowsPlatformApplicationMisc::IsThisApplicationForeground);
+}
+
 void FWindowsPlatformApplicationMisc::LoadStartupModules()
 {
 #if !UE_SERVER
@@ -134,15 +139,14 @@ void FWindowsPlatformApplicationMisc::PumpMessages(bool bFromMainLoop)
 	WinPumpMessages();
 
 	// Determine if application has focus
-	bool bHasFocus = FApp::UseVRFocus() ? FApp::HasVRFocus() : FWindowsPlatformApplicationMisc::IsThisApplicationForeground();
+	bool bHasFocus = FApp::HasFocus();
 	static bool bHadFocus = false;
 
 #if WITH_EDITOR
 	// If editor thread doesn't have the focus, don't suck up too much CPU time.
 	if( GIsEditor )
 	{
-		bool bBenchmarking = FApp::IsBenchmarking();
-		if( !bHasFocus && !bBenchmarking )
+		if( !bHasFocus )
 		{
 			TRACE_CPUPROFILER_EVENT_SCOPE(EditorIsInBackgroundSleep);
 			// Sleep for a bit to not eat up all CPU time.
diff --git a/Engine/Source/Runtime/ApplicationCore/Public/Windows/WindowsPlatformApplicationMisc.h b/Engine/Source/Runtime/ApplicationCore/Public/Windows/WindowsPlatformApplicationMisc.h
--- a/Engine/Source/Runtime/ApplicationCore/Public/Windows/WindowsPlatformApplicationMisc.h
+++ b/Engine/Source/Runtime/ApplicationCore/Public/Windows/WindowsPlatformApplicationMisc.h
@@ -6,6 +6,7 @@
 
 struct APPLICATIONCORE_API FWindowsPlatformApplicationMisc : public FGenericPlatformApplicationMisc
 {
+	static void PreInit();
 	static void LoadStartupModules();
 	static class FOutputDeviceConsole* CreateConsoleOutputDevice();
 	static class FOutputDeviceError* GetErrorOutputDevice();
diff --git a/Engine/Source/Runtime/Core/Private/Misc/App.cpp b/Engine/Source/Runtime/Core/Private/Misc/App.cpp
--- a/Engine/Source/Runtime/Core/Private/Misc/App.cpp
+++ b/Engine/Source/Runtime/Core/Private/Misc/App.cpp
@@ -41,6 +41,7 @@ float FApp::VolumeMultiplier = 1.0f;
 float FApp::UnfocusedVolumeMultiplier = 0.0f;
 bool FApp::bUseVRFocus = false;
 bool FApp::bHasVRFocus = false;
+bool (*FApp::HasFocusFunction)() = nullptr;
 
 
 /* FApp static interface
@@ -294,6 +295,35 @@ void FApp::SetHasVRFocus(bool bInHasVRFocus)
 	bHasVRFocus = bInHasVRFocus;
 }
 
+void FApp::SetHasFocusFunction(bool (*InHasFocusFunction)())
+{
+	HasFocusFunction = InHasFocusFunction;
+}
+
+bool FApp::HasFocus()
+{
+	if (FApp::IsBenchmarking())
+	{
+		return true;
+	}
+
+	if (FApp::UseVRFocus())
+	{
+		return FApp::HasVRFocus();
+	}
+
+	// by default we assume we have focus, it's a worse thing to encounter a bug where focus is locked off, vs. locked on
+	bool bHasFocus = true;
+
+	// desktop platforms are more or less why we have this abstraction, to dip into ApplicationCore's Platform implementation
+#if PLATFORM_DESKTOP
+	check(HasFocusFunction);
+#endif
+
+	// call the HasFocusFunction, if we have one. otherwise fall back to the default
+	return HasFocusFunction ? HasFocusFunction() : bHasFocus;
+}
+
 void FApp::PrintStartupLogMessages()
 {
 	UE_LOG(LogInit, Log, TEXT("Build: %s"), FApp::GetBuildVersion());
diff --git a/Engine/Source/Runtime/Core/Public/Misc/App.h b/Engine/Source/Runtime/Core/Public/Misc/App.h
--- a/Engine/Source/Runtime/Core/Public/Misc/App.h
+++ b/Engine/Source/Runtime/Core/Public/Misc/App.h
@@ -730,6 +730,19 @@ public:
 		return bHasVRFocus;
 	}
 	
+	/**
+	 * Sets HasFocus, which is a function that indicates that the application window has focus.
+	 *
+	 * @param  InHasFocusFunction	address of a function that can query the focus state, typically &FPlatformApplicationMisc::IsThisApplicationForeground
+	 */
+	static void SetHasFocusFunction(bool (*InHasFocusFunction)());
+
+	/**
+	 * Gets Focus, Indicates that the application should continue to render
+	 * Audio and Video as if it had window focus, even though it may not.
+	 */
+	static bool HasFocus();
+
 	/* If the random seed started with a constant or on time, can be affected by -FIXEDSEED or -BENCHMARK */
 	static bool bUseFixedSeed;
 
@@ -805,6 +818,9 @@ private:
 
 	/** Holds a flag indicating if app has focus in side the VR headset */
 	static bool bHasVRFocus;
+
+	/** Holds a function address that can indicate if application has focus */
+	static bool (*HasFocusFunction)();
 };
 
 
diff --git a/Engine/Source/Runtime/Launch/Private/LaunchEngineLoop.cpp b/Engine/Source/Runtime/Launch/Private/LaunchEngineLoop.cpp
--- a/Engine/Source/Runtime/Launch/Private/LaunchEngineLoop.cpp
+++ b/Engine/Source/Runtime/Launch/Private/LaunchEngineLoop.cpp
@@ -4703,7 +4703,7 @@ bool FEngineLoop::ShouldUseIdleMode() const
 	if (FApp::IsGame()
 		&& FPlatformProperties::SupportsWindowedMode()
 		&& CVarIdleWhenNotForeground->GetValueOnGameThread()
-		&& !FPlatformApplicationMisc::IsThisApplicationForeground())
+		&& !FApp::HasFocus())
 	{
 		bIdleMode = true;
 	}
