Fix crash when disconnecting / reconnecting a bluetooth controller outside of the app
�rb jack.porter

diff --git a/Engine/Source/Runtime/ApplicationCore/Private/Apple/AppleControllerInterface.cpp b/Engine/Source/Runtime/ApplicationCore/Private/Apple/AppleControllerInterface.cpp
--- a/Engine/Source/Runtime/ApplicationCore/Private/Apple/AppleControllerInterface.cpp
+++ b/Engine/Source/Runtime/ApplicationCore/Private/Apple/AppleControllerInterface.cpp
@@ -86,19 +86,6 @@ void FAppleControllerInterface::SetCurrentController(GCController* Controller)
 {
     int32 ControllerIndex = 0;
 
-    for (ControllerIndex = 0; ControllerIndex < UE_ARRAY_COUNT(Controllers); ControllerIndex++)
-    {
-        if (Controllers[ControllerIndex].Controller == Controller)
-        {
-            break;
-        }
-    }
-    if (ControllerIndex == UE_ARRAY_COUNT(Controllers))
-    {
-        HandleConnection(Controller);
-    }
-
-
     for (ControllerIndex = 0; ControllerIndex < UE_ARRAY_COUNT(Controllers); ControllerIndex++)
     {
         if (Controllers[ControllerIndex].Controller == Controller)
