Introduced UE_WINDOWS_USING_UIA macro to facilitate Windows accessibility only if
the Windows UIA dll is used.
#rnx

diff --git a/Engine/Source/Runtime/ApplicationCore/ApplicationCore.Build.cs b/Engine/Source/Runtime/ApplicationCore/ApplicationCore.Build.cs
--- a/Engine/Source/Runtime/ApplicationCore/ApplicationCore.Build.cs
+++ b/Engine/Source/Runtime/ApplicationCore/ApplicationCore.Build.cs
@@ -31,9 +31,15 @@ public class ApplicationCore : ModuleRules
 			AddEngineThirdPartyPrivateStaticDependencies(Target,
 				"XInput"
 			);
-			if (Target.bCompileWithAccessibilitySupport && !Target.bIsBuildingConsoleApplication)
+			// We only enable UIA with Win64
+			if (Target.bCompileWithAccessibilitySupport && !Target.bIsBuildingConsoleApplication && Target.Platform == UnrealTargetPlatform.Win64)
 			{
 				PublicSystemLibraries.Add("uiautomationcore.lib");
+				PublicDefinitions.Add("UE_WINDOWS_USING_UIA=1");
+			}
+			else
+			{
+				PublicDefinitions.Add("UE_WINDOWS_USING_UIA=0");
 			}
 
 			// Uses DXGI to query GPU hardware prior to RHI startup
diff --git a/Engine/Source/Runtime/ApplicationCore/Private/Windows/Accessibility/WindowsUIABaseProvider.cpp b/Engine/Source/Runtime/ApplicationCore/Private/Windows/Accessibility/WindowsUIABaseProvider.cpp
--- a/Engine/Source/Runtime/ApplicationCore/Private/Windows/Accessibility/WindowsUIABaseProvider.cpp
+++ b/Engine/Source/Runtime/ApplicationCore/Private/Windows/Accessibility/WindowsUIABaseProvider.cpp
@@ -1,6 +1,6 @@
 // Copyright Epic Games, Inc. All Rights Reserved.
 
-#if WITH_ACCESSIBILITY
+#if WITH_ACCESSIBILITY && UE_WINDOWS_USING_UIA
 
 #include "Windows/Accessibility/WindowsUIABaseProvider.h"
 
diff --git a/Engine/Source/Runtime/ApplicationCore/Private/Windows/Accessibility/WindowsUIAControlProvider.cpp b/Engine/Source/Runtime/ApplicationCore/Private/Windows/Accessibility/WindowsUIAControlProvider.cpp
--- a/Engine/Source/Runtime/ApplicationCore/Private/Windows/Accessibility/WindowsUIAControlProvider.cpp
+++ b/Engine/Source/Runtime/ApplicationCore/Private/Windows/Accessibility/WindowsUIAControlProvider.cpp
@@ -1,6 +1,6 @@
 // Copyright Epic Games, Inc. All Rights Reserved.
 
-#if WITH_ACCESSIBILITY
+#if WITH_ACCESSIBILITY && UE_WINDOWS_USING_UIA
 
 #include "Windows/Accessibility/WindowsUIAControlProvider.h"
 #include "Windows/Accessibility/WindowsUIAWidgetProvider.h"
diff --git a/Engine/Source/Runtime/ApplicationCore/Private/Windows/Accessibility/WindowsUIAManager.cpp b/Engine/Source/Runtime/ApplicationCore/Private/Windows/Accessibility/WindowsUIAManager.cpp
--- a/Engine/Source/Runtime/ApplicationCore/Private/Windows/Accessibility/WindowsUIAManager.cpp
+++ b/Engine/Source/Runtime/ApplicationCore/Private/Windows/Accessibility/WindowsUIAManager.cpp
@@ -1,6 +1,6 @@
 // Copyright Epic Games, Inc. All Rights Reserved.
 
-#if WITH_ACCESSIBILITY
+#if WITH_ACCESSIBILITY && UE_WINDOWS_USING_UIA
 
 #include "Windows/Accessibility/WindowsUIAManager.h"
 
diff --git a/Engine/Source/Runtime/ApplicationCore/Private/Windows/Accessibility/WindowsUIAPropertyGetters.cpp b/Engine/Source/Runtime/ApplicationCore/Private/Windows/Accessibility/WindowsUIAPropertyGetters.cpp
--- a/Engine/Source/Runtime/ApplicationCore/Private/Windows/Accessibility/WindowsUIAPropertyGetters.cpp
+++ b/Engine/Source/Runtime/ApplicationCore/Private/Windows/Accessibility/WindowsUIAPropertyGetters.cpp
@@ -1,6 +1,6 @@
 // Copyright Epic Games, Inc. All Rights Reserved.
 
-#if WITH_ACCESSIBILITY
+#if WITH_ACCESSIBILITY && UE_WINDOWS_USING_UIA
 
 #include "Windows/Accessibility/WindowsUIAPropertyGetters.h"
 
diff --git a/Engine/Source/Runtime/ApplicationCore/Private/Windows/Accessibility/WindowsUIAWidgetProvider.cpp b/Engine/Source/Runtime/ApplicationCore/Private/Windows/Accessibility/WindowsUIAWidgetProvider.cpp
--- a/Engine/Source/Runtime/ApplicationCore/Private/Windows/Accessibility/WindowsUIAWidgetProvider.cpp
+++ b/Engine/Source/Runtime/ApplicationCore/Private/Windows/Accessibility/WindowsUIAWidgetProvider.cpp
@@ -1,6 +1,6 @@
 // Copyright Epic Games, Inc. All Rights Reserved.
 
-#if WITH_ACCESSIBILITY
+#if WITH_ACCESSIBILITY && UE_WINDOWS_USING_UIA
 
 #include "Windows/Accessibility/WindowsUIAWidgetProvider.h"
 #include "GenericPlatform/Accessibility/GenericAccessibleInterfaces.h"
diff --git a/Engine/Source/Runtime/ApplicationCore/Private/Windows/WindowsApplication.cpp b/Engine/Source/Runtime/ApplicationCore/Private/Windows/WindowsApplication.cpp
--- a/Engine/Source/Runtime/ApplicationCore/Private/Windows/WindowsApplication.cpp
+++ b/Engine/Source/Runtime/ApplicationCore/Private/Windows/WindowsApplication.cpp
@@ -26,7 +26,7 @@
 #include "Developer/SourceCodeAccess/Public/ISourceCodeAccessModule.h"
 #endif
 
-#if WITH_ACCESSIBILITY
+#if WITH_ACCESSIBILITY && UE_WINDOWS_USING_UIA
 #include "Windows/Accessibility/WindowsUIAManager.h"
 #include "Windows/Accessibility/WindowsUIAWidgetProvider.h"
 #include <UIAutomation.h>
@@ -116,7 +116,7 @@ FWindowsApplication::FWindowsApplication( const HINSTANCE HInstance, const HICON
 		bAllowedToDeferMessageProcessing,
 		TEXT( "Whether windows message processing is deferred until tick or if they are processed immediately" ) )
 	, bInModalSizeLoop( false )
-#if WITH_ACCESSIBILITY
+#if WITH_ACCESSIBILITY && UE_WINDOWS_USING_UIA
 	, UIAManager(new FWindowsUIAManager(*this))
 #endif
 	, bSimulatingHighPrecisionMouseInputForRDP(false)
@@ -384,7 +384,9 @@ void FWindowsApplication::SetMessageHandler( const TSharedRef< FGenericApplicati
 void FWindowsApplication::SetAccessibleMessageHandler(const TSharedRef<FGenericAccessibleMessageHandler>& InAccessibleMessageHandler)
 {
 	GenericApplication::SetAccessibleMessageHandler(InAccessibleMessageHandler);
+#if UE_WINDOWS_USING_UIA
 	UIAManager->OnAccessibleMessageHandlerChanged();
+#endif
 }
 #endif
 
@@ -1683,7 +1685,7 @@ int32 FWindowsApplication::ProcessMessage( HWND hwnd, uint32 msg, WPARAM wParam,
 		case WM_DESTROY:
 			{
 				Windows.Remove( CurrentNativeEventWindow );
-#if WITH_ACCESSIBILITY
+#if WITH_ACCESSIBILITY && UE_WINDOWS_USING_UIA
 				// Tell UIA that the window no longer exists so that it can release some resources
 				if (GetAccessibleMessageHandler()->ApplicationIsAccessible())
 				{
@@ -1865,7 +1867,7 @@ int32 FWindowsApplication::ProcessMessage( HWND hwnd, uint32 msg, WPARAM wParam,
 			}
 			break;
 
-#if WITH_ACCESSIBILITY
+#if WITH_ACCESSIBILITY && UE_WINDOWS_USING_UIA
 		case WM_GETOBJECT:
 		{
 			if (GetAccessibleMessageHandler()->ApplicationIsAccessible())
diff --git a/Engine/Source/Runtime/ApplicationCore/Public/Windows/Accessibility/WindowsUIABaseProvider.h b/Engine/Source/Runtime/ApplicationCore/Public/Windows/Accessibility/WindowsUIABaseProvider.h
--- a/Engine/Source/Runtime/ApplicationCore/Public/Windows/Accessibility/WindowsUIABaseProvider.h
+++ b/Engine/Source/Runtime/ApplicationCore/Public/Windows/Accessibility/WindowsUIABaseProvider.h
@@ -2,7 +2,7 @@
 
 #pragma once
 
-#if WITH_ACCESSIBILITY
+#if WITH_ACCESSIBILITY && UE_WINDOWS_USING_UIA
 
 #include "Windows/AllowWindowsPlatformTypes.h"
 #include <Ole2.h>
diff --git a/Engine/Source/Runtime/ApplicationCore/Public/Windows/Accessibility/WindowsUIAControlProvider.h b/Engine/Source/Runtime/ApplicationCore/Public/Windows/Accessibility/WindowsUIAControlProvider.h
--- a/Engine/Source/Runtime/ApplicationCore/Public/Windows/Accessibility/WindowsUIAControlProvider.h
+++ b/Engine/Source/Runtime/ApplicationCore/Public/Windows/Accessibility/WindowsUIAControlProvider.h
@@ -2,7 +2,7 @@
 
 #pragma once
 
-#if WITH_ACCESSIBILITY
+#if WITH_ACCESSIBILITY && UE_WINDOWS_USING_UIA
 
 #include "Windows/Accessibility/WindowsUIABaseProvider.h"
 #include "Templates/SharedPointer.h"
diff --git a/Engine/Source/Runtime/ApplicationCore/Public/Windows/Accessibility/WindowsUIAManager.h b/Engine/Source/Runtime/ApplicationCore/Public/Windows/Accessibility/WindowsUIAManager.h
--- a/Engine/Source/Runtime/ApplicationCore/Public/Windows/Accessibility/WindowsUIAManager.h
+++ b/Engine/Source/Runtime/ApplicationCore/Public/Windows/Accessibility/WindowsUIAManager.h
@@ -2,7 +2,7 @@
 
 #pragma once
 
-#if WITH_ACCESSIBILITY
+#if WITH_ACCESSIBILITY && UE_WINDOWS_USING_UIA
 
 #include "Templates/SharedPointer.h"
 #include "GenericPlatform/Accessibility/GenericAccessibleInterfaces.h"
diff --git a/Engine/Source/Runtime/ApplicationCore/Public/Windows/Accessibility/WindowsUIAPropertyGetters.h b/Engine/Source/Runtime/ApplicationCore/Public/Windows/Accessibility/WindowsUIAPropertyGetters.h
--- a/Engine/Source/Runtime/ApplicationCore/Public/Windows/Accessibility/WindowsUIAPropertyGetters.h
+++ b/Engine/Source/Runtime/ApplicationCore/Public/Windows/Accessibility/WindowsUIAPropertyGetters.h
@@ -2,7 +2,7 @@
 
 #pragma once
 
-#if WITH_ACCESSIBILITY
+#if WITH_ACCESSIBILITY && UE_WINDOWS_USING_UIA
 
 #include "Windows/AllowWindowsPlatformTypes.h"
 #include <Ole2.h>
diff --git a/Engine/Source/Runtime/ApplicationCore/Public/Windows/Accessibility/WindowsUIAWidgetProvider.h b/Engine/Source/Runtime/ApplicationCore/Public/Windows/Accessibility/WindowsUIAWidgetProvider.h
--- a/Engine/Source/Runtime/ApplicationCore/Public/Windows/Accessibility/WindowsUIAWidgetProvider.h
+++ b/Engine/Source/Runtime/ApplicationCore/Public/Windows/Accessibility/WindowsUIAWidgetProvider.h
@@ -2,7 +2,7 @@
 
 #pragma once
 
-#if WITH_ACCESSIBILITY
+#if WITH_ACCESSIBILITY && UE_WINDOWS_USING_UIA
 
 #include "Windows/Accessibility/WindowsUIABaseProvider.h"
 #include "Misc/Variant.h"
diff --git a/Engine/Source/Runtime/ApplicationCore/Public/Windows/WindowsApplication.h b/Engine/Source/Runtime/ApplicationCore/Public/Windows/WindowsApplication.h
--- a/Engine/Source/Runtime/ApplicationCore/Public/Windows/WindowsApplication.h
+++ b/Engine/Source/Runtime/ApplicationCore/Public/Windows/WindowsApplication.h
@@ -522,7 +522,7 @@ private:
 
 	TSharedPtr<FTaskbarList> TaskbarList;
 
-#if WITH_ACCESSIBILITY
+#if WITH_ACCESSIBILITY && UE_WINDOWS_USING_UIA
 	/** Handler for WM_GetObject messages that come in */
 	TUniquePtr<class FWindowsUIAManager> UIAManager;
 #endif
