Change FPlatformUserId into a proper struct with an internal id that will be allocated by the platform
Converting int->FPlatformUserId is currently deprecated, but the other direction is not deprecated until the online/input code can be converted
GenericApplicaitonMessageHandler will convert back and forth between controller Id and FPlatformUserId as needed
Add support for platform user id to LocalPlayer, by default it is synchronized with ControllerId

diff --git a/Engine/Plugins/Online/Android/OnlineSubsystemGameCircle/Source/OnlineSubsystemGameCircle/Private/OnlineIdentityInterfaceGameCircle.cpp b/Engine/Plugins/Online/Android/OnlineSubsystemGameCircle/Source/OnlineSubsystemGameCircle/Private/OnlineIdentityInterfaceGameCircle.cpp
--- a/Engine/Plugins/Online/Android/OnlineSubsystemGameCircle/Source/OnlineSubsystemGameCircle/Private/OnlineIdentityInterfaceGameCircle.cpp
+++ b/Engine/Plugins/Online/Android/OnlineSubsystemGameCircle/Source/OnlineSubsystemGameCircle/Private/OnlineIdentityInterfaceGameCircle.cpp
@@ -130,7 +130,7 @@ FPlatformUserId FOnlineIdentityGameCircle::GetPlatformUserIdFromUniqueNetId(cons
 		auto CurrentUniqueId = GetUniquePlayerId(i);
 		if (CurrentUniqueId.IsValid() && (*CurrentUniqueId == NetId))
 		{
-			return i;
+			return GetPlatformUserIdFromLocalUserNum(i);
 		}
 	}
 
diff --git a/Engine/Plugins/Online/Android/OnlineSubsystemGooglePlay/Source/Private/OnlineIdentityInterfaceGooglePlay.cpp b/Engine/Plugins/Online/Android/OnlineSubsystemGooglePlay/Source/Private/OnlineIdentityInterfaceGooglePlay.cpp
--- a/Engine/Plugins/Online/Android/OnlineSubsystemGooglePlay/Source/Private/OnlineIdentityInterfaceGooglePlay.cpp
+++ b/Engine/Plugins/Online/Android/OnlineSubsystemGooglePlay/Source/Private/OnlineIdentityInterfaceGooglePlay.cpp
@@ -188,7 +188,7 @@ FPlatformUserId FOnlineIdentityGooglePlay::GetPlatformUserIdFromUniqueNetId(cons
 		auto CurrentUniqueId = GetUniquePlayerId(i);
 		if (CurrentUniqueId.IsValid() && (*CurrentUniqueId == NetId))
 		{
-			return i;
+			return GetPlatformUserIdFromLocalUserNum(i);
 		}
 	}
 
diff --git a/Engine/Plugins/Online/IOS/OnlineSubsystemIOS/Source/Private/OnlineIdentityInterfaceIOS.cpp b/Engine/Plugins/Online/IOS/OnlineSubsystemIOS/Source/Private/OnlineIdentityInterfaceIOS.cpp
--- a/Engine/Plugins/Online/IOS/OnlineSubsystemIOS/Source/Private/OnlineIdentityInterfaceIOS.cpp
+++ b/Engine/Plugins/Online/IOS/OnlineSubsystemIOS/Source/Private/OnlineIdentityInterfaceIOS.cpp
@@ -379,7 +379,7 @@ FPlatformUserId FOnlineIdentityIOS::GetPlatformUserIdFromUniqueNetId(const FUniq
 		auto CurrentUniqueId = GetUniquePlayerId(i);
 		if (CurrentUniqueId.IsValid() && (*CurrentUniqueId == InUniqueNetId))
 		{
-			return i;
+			return GetPlatformUserIdFromLocalUserNum(i);
 		}
 	}
 
diff --git a/Engine/Plugins/Online/OnlineSubsystem/Source/Private/OnlineIdentityInterface.cpp b/Engine/Plugins/Online/OnlineSubsystem/Source/Private/OnlineIdentityInterface.cpp
--- a/Engine/Plugins/Online/OnlineSubsystem/Source/Private/OnlineIdentityInterface.cpp
+++ b/Engine/Plugins/Online/OnlineSubsystem/Source/Private/OnlineIdentityInterface.cpp
@@ -1,6 +1,7 @@
 // Copyright Epic Games, Inc. All Rights Reserved.
 
 #include "Interfaces/OnlineIdentityInterface.h"
+#include "HAL/PlatformMisc.h"
 
 FString ToDebugString(IOnlineIdentity::EPrivilegeResults PrivilegeResults)
 {
@@ -39,4 +40,26 @@ FString ToDebugString(EUserPrivileges::Type UserPrivilege)
 	case EUserPrivileges::CanUserCrossPlay: return TEXT("CanUserCrossPlay");
 	default: return TEXT("Unknown");
 	}
-}
\ No newline at end of file
+}
+
+FPlatformUserId IOnlineIdentity::GetPlatformUserIdFromLocalUserNum(int32 LocalUserNum) const
+{
+	if (LocalUserNum >= 0 && LocalUserNum < MAX_LOCAL_PLAYERS)
+	{
+		return FPlatformMisc::GetPlatformUserForUserIndex(LocalUserNum);
+	}
+
+	return PLATFORMUSERID_NONE;
+}
+
+int32 IOnlineIdentity::GetLocalUserNumFromPlatformUserId(FPlatformUserId PlatformUserId) const
+{
+	int32 Index = FPlatformMisc::GetUserIndexForPlatformUser(PlatformUserId);
+
+	if (Index >= 0 && Index < MAX_LOCAL_PLAYERS)
+	{
+		return Index;
+	}
+
+	return INDEX_NONE;
+}
diff --git a/Engine/Plugins/Online/OnlineSubsystem/Source/Public/Interfaces/OnlineIdentityInterface.h b/Engine/Plugins/Online/OnlineSubsystem/Source/Public/Interfaces/OnlineIdentityInterface.h
--- a/Engine/Plugins/Online/OnlineSubsystem/Source/Public/Interfaces/OnlineIdentityInterface.h
+++ b/Engine/Plugins/Online/OnlineSubsystem/Source/Public/Interfaces/OnlineIdentityInterface.h
@@ -473,7 +473,7 @@ public:
 	virtual void GetUserPrivilege(const FUniqueNetId& LocalUserId, EUserPrivileges::Type Privilege, const FOnGetUserPrivilegeCompleteDelegate& Delegate) = 0;
 
 	/**
-	 * Temporary hack to get a corresponding FUniqueNetId from a PlatformUserId
+	 * Converts from an online unique id to a PlatformUserId used by application code
 	 *
 	 * @param UniqueNetId The unique id to look up
 	 *
@@ -481,6 +481,24 @@ public:
 	 */
 	virtual FPlatformUserId GetPlatformUserIdFromUniqueNetId(const FUniqueNetId& UniqueNetId) const = 0;
 
+	/**
+	 * Converts from a local user num used by the online system to a PlatformUserId used by application code
+	 *
+	 * @param LocalUserNum the controller number of the associated user
+	 *
+	 * @return The corresponding id or PLATFORMID_NONE if invalid
+	 */
+	ONLINESUBSYSTEM_API virtual FPlatformUserId GetPlatformUserIdFromLocalUserNum(int32 LocalUserNum) const;
+
+	/**
+	 * Converts from a PlatformUserId used by application code to a local user num used by the online system
+	 *
+	 * @param PlatformUserId The application-level platform user id
+	 *
+	 * @return The corresponding user number of INDEX_NONE if invalid
+	 */
+	ONLINESUBSYSTEM_API virtual int32 GetLocalUserNumFromPlatformUserId(FPlatformUserId PlatformUserId) const;
+
 	/**
 	 * Get the auth type associated with accounts for this platform
 	 *
diff --git a/Engine/Plugins/Online/OnlineSubsystemAmazon/Source/Private/OnlineIdentityAmazon.cpp b/Engine/Plugins/Online/OnlineSubsystemAmazon/Source/Private/OnlineIdentityAmazon.cpp
--- a/Engine/Plugins/Online/OnlineSubsystemAmazon/Source/Private/OnlineIdentityAmazon.cpp
+++ b/Engine/Plugins/Online/OnlineSubsystemAmazon/Source/Private/OnlineIdentityAmazon.cpp
@@ -367,7 +367,7 @@ FPlatformUserId FOnlineIdentityAmazon::GetPlatformUserIdFromUniqueNetId(const FU
 		auto CurrentUniqueId = GetUniquePlayerId(i);
 		if (CurrentUniqueId.IsValid() && (*CurrentUniqueId == UniqueNetId))
 		{
-			return i;
+			return GetPlatformUserIdFromLocalUserNum(i);
 		}
 	}
 
diff --git a/Engine/Plugins/Online/OnlineSubsystemApple/Source/Private/OnlineIdentityInterfaceApple.cpp b/Engine/Plugins/Online/OnlineSubsystemApple/Source/Private/OnlineIdentityInterfaceApple.cpp
--- a/Engine/Plugins/Online/OnlineSubsystemApple/Source/Private/OnlineIdentityInterfaceApple.cpp
+++ b/Engine/Plugins/Online/OnlineSubsystemApple/Source/Private/OnlineIdentityInterfaceApple.cpp
@@ -395,7 +395,7 @@ FPlatformUserId FOnlineIdentityApple::GetPlatformUserIdFromUniqueNetId(const FUn
 		FUniqueNetIdPtr CurrentUniqueId = GetUniquePlayerId(i);
 		if (CurrentUniqueId.IsValid() && (*CurrentUniqueId == InUniqueNetId))
 		{
-			return i;
+			return GetPlatformUserIdFromLocalUserNum(i);
 		}
 	}
 
diff --git a/Engine/Plugins/Online/OnlineSubsystemEOS/Source/OnlineSubsystemEOS/Private/UserManagerEOS.cpp b/Engine/Plugins/Online/OnlineSubsystemEOS/Source/OnlineSubsystemEOS/Private/UserManagerEOS.cpp
--- a/Engine/Plugins/Online/OnlineSubsystemEOS/Source/OnlineSubsystemEOS/Private/UserManagerEOS.cpp
+++ b/Engine/Plugins/Online/OnlineSubsystemEOS/Source/OnlineSubsystemEOS/Private/UserManagerEOS.cpp
@@ -1288,7 +1288,7 @@ void FUserManagerEOS::RevokeAuthToken(const FUniqueNetId& LocalUserId, const FOn
 
 FPlatformUserId FUserManagerEOS::GetPlatformUserIdFromUniqueNetId(const FUniqueNetId& UniqueNetId) const
 {
-	return GetLocalUserNumFromUniqueNetId(UniqueNetId);
+	return GetPlatformUserIdFromLocalUserNum(GetLocalUserNumFromUniqueNetId(UniqueNetId));
 }
 
 void FUserManagerEOS::GetUserPrivilege(const FUniqueNetId& UserId, EUserPrivileges::Type Privilege, const FOnGetUserPrivilegeCompleteDelegate& Delegate)
diff --git a/Engine/Plugins/Online/OnlineSubsystemFacebook/Source/Private/OnlineIdentityFacebookCommon.cpp b/Engine/Plugins/Online/OnlineSubsystemFacebook/Source/Private/OnlineIdentityFacebookCommon.cpp
--- a/Engine/Plugins/Online/OnlineSubsystemFacebook/Source/Private/OnlineIdentityFacebookCommon.cpp
+++ b/Engine/Plugins/Online/OnlineSubsystemFacebook/Source/Private/OnlineIdentityFacebookCommon.cpp
@@ -299,7 +299,7 @@ FPlatformUserId FOnlineIdentityFacebookCommon::GetPlatformUserIdFromUniqueNetId(
 		auto CurrentUniqueId = GetUniquePlayerId(i);
 		if (CurrentUniqueId.IsValid() && (*CurrentUniqueId == UniqueNetId))
 		{
-			return i;
+			return GetPlatformUserIdFromLocalUserNum(i);
 		}
 	}
 
diff --git a/Engine/Plugins/Online/OnlineSubsystemGoogle/Source/Private/OnlineIdentityGoogleCommon.cpp b/Engine/Plugins/Online/OnlineSubsystemGoogle/Source/Private/OnlineIdentityGoogleCommon.cpp
--- a/Engine/Plugins/Online/OnlineSubsystemGoogle/Source/Private/OnlineIdentityGoogleCommon.cpp
+++ b/Engine/Plugins/Online/OnlineSubsystemGoogle/Source/Private/OnlineIdentityGoogleCommon.cpp
@@ -505,7 +505,7 @@ FPlatformUserId FOnlineIdentityGoogleCommon::GetPlatformUserIdFromUniqueNetId(co
 		auto CurrentUniqueId = GetUniquePlayerId(i);
 		if (CurrentUniqueId.IsValid() && (*CurrentUniqueId == UniqueNetId))
 		{
-			return i;
+			return GetPlatformUserIdFromLocalUserNum(i);
 		}
 	}
 
diff --git a/Engine/Plugins/Online/OnlineSubsystemNull/Source/Private/OnlineIdentityNull.cpp b/Engine/Plugins/Online/OnlineSubsystemNull/Source/Private/OnlineIdentityNull.cpp
--- a/Engine/Plugins/Online/OnlineSubsystemNull/Source/Private/OnlineIdentityNull.cpp
+++ b/Engine/Plugins/Online/OnlineSubsystemNull/Source/Private/OnlineIdentityNull.cpp
@@ -349,7 +349,7 @@ FPlatformUserId FOnlineIdentityNull::GetPlatformUserIdFromUniqueNetId(const FUni
 		auto CurrentUniqueId = GetUniquePlayerId(i);
 		if (CurrentUniqueId.IsValid() && (*CurrentUniqueId == UniqueNetId))
 		{
-			return i;
+			return GetPlatformUserIdFromLocalUserNum(i);
 		}
 	}
 
diff --git a/Engine/Plugins/Online/OnlineSubsystemOculus/Source/Private/OnlineIdentityOculus.cpp b/Engine/Plugins/Online/OnlineSubsystemOculus/Source/Private/OnlineIdentityOculus.cpp
--- a/Engine/Plugins/Online/OnlineSubsystemOculus/Source/Private/OnlineIdentityOculus.cpp
+++ b/Engine/Plugins/Online/OnlineSubsystemOculus/Source/Private/OnlineIdentityOculus.cpp
@@ -293,7 +293,7 @@ FPlatformUserId FOnlineIdentityOculus::GetPlatformUserIdFromUniqueNetId(const FU
 		auto CurrentUniqueId = GetUniquePlayerId(i);
 		if (CurrentUniqueId.IsValid() && (*CurrentUniqueId == UniqueNetId))
 		{
-			return i;
+			return GetPlatformUserIdFromLocalUserNum(i);
 		}
 	}
 
diff --git a/Engine/Plugins/Online/OnlineSubsystemSteam/Source/Private/OnlineIdentityInterfaceSteam.cpp b/Engine/Plugins/Online/OnlineSubsystemSteam/Source/Private/OnlineIdentityInterfaceSteam.cpp
--- a/Engine/Plugins/Online/OnlineSubsystemSteam/Source/Private/OnlineIdentityInterfaceSteam.cpp
+++ b/Engine/Plugins/Online/OnlineSubsystemSteam/Source/Private/OnlineIdentityInterfaceSteam.cpp
@@ -227,7 +227,7 @@ FPlatformUserId FOnlineIdentitySteam::GetPlatformUserIdFromUniqueNetId(const FUn
 		auto CurrentUniqueId = GetUniquePlayerId(i);
 		if (CurrentUniqueId.IsValid() && (*CurrentUniqueId == UniqueNetId))
 		{
-			return i;
+			return GetPlatformUserIdFromLocalUserNum(i);
 		}
 	}
 
diff --git a/Engine/Plugins/Online/OnlineSubsystemTencent/Source/Private/OnlineIdentityTencentRail.cpp b/Engine/Plugins/Online/OnlineSubsystemTencent/Source/Private/OnlineIdentityTencentRail.cpp
--- a/Engine/Plugins/Online/OnlineSubsystemTencent/Source/Private/OnlineIdentityTencentRail.cpp
+++ b/Engine/Plugins/Online/OnlineSubsystemTencent/Source/Private/OnlineIdentityTencentRail.cpp
@@ -341,7 +341,7 @@ FPlatformUserId FOnlineIdentityTencent::GetPlatformUserIdFromUniqueNetId(const F
 		auto CurrentUniqueId = GetUniquePlayerId(i);
 		if (CurrentUniqueId.IsValid() && (*CurrentUniqueId == UniqueNetId))
 		{
-			return i;
+			return GetPlatformUserIdFromLocalUserNum(i);
 		}
 	}
 
diff --git a/Engine/Plugins/Online/OnlineSubsystemTwitch/Source/Private/OnlineIdentityTwitch.cpp b/Engine/Plugins/Online/OnlineSubsystemTwitch/Source/Private/OnlineIdentityTwitch.cpp
--- a/Engine/Plugins/Online/OnlineSubsystemTwitch/Source/Private/OnlineIdentityTwitch.cpp
+++ b/Engine/Plugins/Online/OnlineSubsystemTwitch/Source/Private/OnlineIdentityTwitch.cpp
@@ -717,7 +717,7 @@ FPlatformUserId FOnlineIdentityTwitch::GetPlatformUserIdFromUniqueNetId(const FU
 		FUniqueNetIdPtr CurrentUniqueId = GetUniquePlayerId(PlayerIdx);
 		if (CurrentUniqueId.IsValid() && (*CurrentUniqueId == UniqueNetId))
 		{
-			return PlayerIdx;
+			return GetPlatformUserIdFromLocalUserNum(PlayerIdx);
 		}
 	}
 
diff --git a/Engine/Plugins/Online/VoiceChat/EOSVoiceChat/Source/EOSVoiceChat/Private/EOSVoiceChatUser.cpp b/Engine/Plugins/Online/VoiceChat/EOSVoiceChat/Source/EOSVoiceChat/Private/EOSVoiceChatUser.cpp
--- a/Engine/Plugins/Online/VoiceChat/EOSVoiceChat/Source/EOSVoiceChat/Private/EOSVoiceChatUser.cpp
+++ b/Engine/Plugins/Online/VoiceChat/EOSVoiceChat/Source/EOSVoiceChat/Private/EOSVoiceChatUser.cpp
@@ -2408,7 +2408,8 @@ bool FEOSVoiceChatUser::Exec(UWorld* InWorld, const TCHAR* Cmd, FOutputDevice& A
 		{
 			FString Token = InsecureGetLoginToken(PlayerName);
 
-			Login(0, PlayerName, Token, FOnVoiceChatLoginCompleteDelegate::CreateLambda([this, &Ar](const FString& LoggedInPlayerName, const FVoiceChatResult& Result)
+			FPlatformUserId FirstUser = FPlatformMisc::GetPlatformUserForUserIndex(0);
+			Login(FirstUser, PlayerName, Token, FOnVoiceChatLoginCompleteDelegate::CreateLambda([this, &Ar](const FString& LoggedInPlayerName, const FVoiceChatResult& Result)
 			{
 				EOSVOICECHATUSER_LOG(Display, TEXT("EOS LOGIN playername:%s result:%s"), *LoggedInPlayerName, *LexToString(Result));
 			}));
diff --git a/Engine/Source/Runtime/ApplicationCore/Private/Android/AndroidInputInterface.cpp b/Engine/Source/Runtime/ApplicationCore/Private/Android/AndroidInputInterface.cpp
--- a/Engine/Source/Runtime/ApplicationCore/Private/Android/AndroidInputInterface.cpp
+++ b/Engine/Source/Runtime/ApplicationCore/Private/Android/AndroidInputInterface.cpp
@@ -127,7 +127,7 @@ void FAndroidInputInterface::ResetGamepadAssignments()
 	{
 		if (DeviceMapping[DeviceIndex].DeviceState == MappingState::Valid)
 		{
-			FCoreDelegates::OnControllerConnectionChange.Broadcast(false, -1, DeviceIndex);
+			FCoreDelegates::OnControllerConnectionChange.Broadcast(false, PLATFORMUSERID_NONE, DeviceIndex);
 		}
 
 		DeviceMapping[DeviceIndex].DeviceInfo.DeviceId = 0;
@@ -142,7 +142,7 @@ void FAndroidInputInterface::ResetGamepadAssignmentToController(int32 Controller
 
 	if (DeviceMapping[ControllerId].DeviceState == MappingState::Valid)
 	{
-		FCoreDelegates::OnControllerConnectionChange.Broadcast(false, -1, ControllerId);
+		FCoreDelegates::OnControllerConnectionChange.Broadcast(false, PLATFORMUSERID_NONE, ControllerId);
 	}
 
 	DeviceMapping[ControllerId].DeviceInfo.DeviceId = 0;
@@ -976,7 +976,7 @@ void FAndroidInputInterface::SendControllerEvents()
 							CurrentDevice.bSupportsHat = true;
 						}
 
-						FCoreDelegates::OnControllerConnectionChange.Broadcast(true, -1, DeviceIndex);
+						FCoreDelegates::OnControllerConnectionChange.Broadcast(true, PLATFORMUSERID_NONE, DeviceIndex);
 
 						FPlatformMisc::LowLevelOutputDebugStringf(TEXT("Assigned new gamepad controller %d: DeviceId=%d, ControllerId=%d, DeviceName=%s, Descriptor=%s"),
 							DeviceIndex, CurrentDevice.DeviceInfo.DeviceId, CurrentDevice.DeviceInfo.ControllerId, *CurrentDevice.DeviceInfo.Name, *CurrentDevice.DeviceInfo.Descriptor);
@@ -995,7 +995,7 @@ void FAndroidInputInterface::SendControllerEvents()
 						OldControllerData[FoundMatch].DeviceId = FoundMatch;
 
 						//@TODO: uncomment this line in the future when disconnects are detected
-						//FCoreDelegates::OnControllerConnectionChange.Broadcast(true, -1, FoundMatch);
+						//FCoreDelegates::OnControllerConnectionChange.Broadcast(true, PLATFORMUSERID_NONE, FoundMatch);
 
 						FPlatformMisc::LowLevelOutputDebugStringf(TEXT("Reconnected gamepad controller %d: DeviceId=%d, ControllerId=%d, DeviceName=%s, Descriptor=%s"),
 							FoundMatch, DeviceMapping[FoundMatch].DeviceInfo.DeviceId, CurrentDevice.DeviceInfo.ControllerId, *CurrentDevice.DeviceInfo.Name, *CurrentDevice.DeviceInfo.Descriptor);
diff --git a/Engine/Source/Runtime/ApplicationCore/Private/Linux/LinuxApplication.cpp b/Engine/Source/Runtime/ApplicationCore/Private/Linux/LinuxApplication.cpp
--- a/Engine/Source/Runtime/ApplicationCore/Private/Linux/LinuxApplication.cpp
+++ b/Engine/Source/Runtime/ApplicationCore/Private/Linux/LinuxApplication.cpp
@@ -2079,7 +2079,7 @@ void FLinuxApplication::AddGameController(int Index)
 		}
 	}
 
-	FCoreDelegates::OnControllerConnectionChange.Broadcast(true, -1, Id);
+	FCoreDelegates::OnControllerConnectionChange.Broadcast(true, PLATFORMUSERID_NONE, Id);
 }
 
 void FLinuxApplication::RemoveGameController(SDL_JoystickID Id)
@@ -2100,5 +2100,5 @@ void FLinuxApplication::RemoveGameController(SDL_JoystickID Id)
 	SDL_GameControllerClose(ControllerState.Controller);
 	ControllerStates.Remove(Id);
 
-	FCoreDelegates::OnControllerConnectionChange.Broadcast(false, -1, Id);			
+	FCoreDelegates::OnControllerConnectionChange.Broadcast(false, PLATFORMUSERID_NONE, Id);
 }
diff --git a/Engine/Source/Runtime/ApplicationCore/Private/Mac/HIDInputInterface.cpp b/Engine/Source/Runtime/ApplicationCore/Private/Mac/HIDInputInterface.cpp
--- a/Engine/Source/Runtime/ApplicationCore/Private/Mac/HIDInputInterface.cpp
+++ b/Engine/Source/Runtime/ApplicationCore/Private/Mac/HIDInputInterface.cpp
@@ -439,7 +439,7 @@ void HIDInputInterface::OnNewHIDController(IOReturn Result, IOHIDDeviceRef Devic
 
 			CFRelease(ElementsArray);
 
-			FCoreDelegates::OnControllerConnectionChange.Broadcast(true, -1, ControllerIndex);
+			FCoreDelegates::OnControllerConnectionChange.Broadcast(true, PLATFORMUSERID_NONE, ControllerIndex);
 
 			bIsGamepadAttached = true;
 		}
@@ -627,7 +627,7 @@ void HIDInputInterface::HIDDeviceRemovalCallback(void* Context, IOReturn Result,
 	{
 		if(HIDInput->ControllerStates[Index].Device.DeviceRef == DeviceRef)
 		{
-			FCoreDelegates::OnControllerConnectionChange.Broadcast(false, -1, Index);
+			FCoreDelegates::OnControllerConnectionChange.Broadcast(false, PLATFORMUSERID_NONE, Index);
 			HIDInput->ControllerStates[Index].Device.DeviceRef = NULL;
 			break;
 		}
diff --git a/Engine/Source/Runtime/ApplicationCore/Private/Windows/XInputInterface.cpp b/Engine/Source/Runtime/ApplicationCore/Private/Windows/XInputInterface.cpp
--- a/Engine/Source/Runtime/ApplicationCore/Private/Windows/XInputInterface.cpp
+++ b/Engine/Source/Runtime/ApplicationCore/Private/Windows/XInputInterface.cpp
@@ -145,11 +145,11 @@ void XInputInterface::SendControllerEvents()
 			// If the controller is connected now but was not before, refresh the information
 			if (!bWasConnected && ControllerState.bIsConnected)
 			{
-				FCoreDelegates::OnControllerConnectionChange.Broadcast(true, -1, ControllerState.ControllerId);
+				FCoreDelegates::OnControllerConnectionChange.Broadcast(true, PLATFORMUSERID_NONE, ControllerState.ControllerId);
 			}
 			else if (bWasConnected && !ControllerState.bIsConnected)
 			{
-				FCoreDelegates::OnControllerConnectionChange.Broadcast(false, -1, ControllerState.ControllerId);
+				FCoreDelegates::OnControllerConnectionChange.Broadcast(false, PLATFORMUSERID_NONE, ControllerState.ControllerId);
 			}
 			
 			bool CurrentStates[MAX_NUM_CONTROLLER_BUTTONS] = {0};
diff --git a/Engine/Source/Runtime/ApplicationCore/Public/GenericPlatform/GenericApplicationMessageHandler.h b/Engine/Source/Runtime/ApplicationCore/Public/GenericPlatform/GenericApplicationMessageHandler.h
--- a/Engine/Source/Runtime/ApplicationCore/Public/GenericPlatform/GenericApplicationMessageHandler.h
+++ b/Engine/Source/Runtime/ApplicationCore/Public/GenericPlatform/GenericApplicationMessageHandler.h
@@ -205,6 +205,8 @@ private:
 	static TArray<FInputDeviceScope*> ScopeStack;
 };
 
+PRAGMA_DISABLE_DEPRECATION_WARNINGS
+
 /** Interface that defines how to handle interaction with a user via hardware input and output */
 class FGenericApplicationMessageHandler
 {
@@ -291,18 +293,40 @@ public:
 		return false;
 	}
 
-	virtual bool OnControllerAnalog( FGamepadKeyNames::Type KeyName, int32 ControllerId, float AnalogValue )
+	/** 
+	 * Return true if this message handler expects FPlatformUserIds. This base class will convert both directions.
+	 * As part of a larger fixup to allow mapping of multiple input devices to the same player, 
+	 * physical device id will be passed as part of the InputScope above and used to compute a logical input user.
+	 */
+	virtual bool ShouldUsePlatformUserId() const
+	{
+		return false;
+	}
+
+	virtual bool OnControllerAnalog(FGamepadKeyNames::Type KeyName, FPlatformUserId PlatformUserId, float AnalogValue)
 	{
+		if (!ShouldUsePlatformUserId())
+		{
+			return OnControllerAnalog(KeyName, PlatformUserId.GetInternalId(), AnalogValue);
+		}
 		return false;
 	}
 
-	virtual bool OnControllerButtonPressed( FGamepadKeyNames::Type KeyName, int32 ControllerId, bool IsRepeat )
+	virtual bool OnControllerButtonPressed(FGamepadKeyNames::Type KeyName, FPlatformUserId PlatformUserId, bool IsRepeat)
 	{
+		if (!ShouldUsePlatformUserId())
+		{
+			return OnControllerButtonPressed(KeyName, PlatformUserId.GetInternalId(), IsRepeat);
+		}
 		return false;
 	}
 
-	virtual bool OnControllerButtonReleased( FGamepadKeyNames::Type KeyName, int32 ControllerId, bool IsRepeat )
+	virtual bool OnControllerButtonReleased(FGamepadKeyNames::Type KeyName, FPlatformUserId PlatformUserId, bool IsRepeat)
 	{
+		if (!ShouldUsePlatformUserId())
+		{
+			return OnControllerButtonReleased(KeyName, PlatformUserId.GetInternalId(), IsRepeat);
+		}
 		return false;
 	}
 
@@ -319,40 +343,48 @@ public:
     {
     }
 
-	UE_DEPRECATED(4.20, "This function signature is deprecated, use OnTouchStarted that takes a Force")
-	virtual bool OnTouchStarted( const TSharedPtr< FGenericWindow >& Window, const FVector2D& Location, int32 TouchIndex, int32 ControllerId )
-	{
-		return OnTouchStarted(Window, Location, 1.0f, TouchIndex, ControllerId);
-	}
-
-	virtual bool OnTouchStarted( const TSharedPtr< FGenericWindow >& Window, const FVector2D& Location, float Force, int32 TouchIndex, int32 ControllerId )
+	virtual bool OnTouchStarted( const TSharedPtr< FGenericWindow >& Window, const FVector2D& Location, float Force, int32 TouchIndex, FPlatformUserId PlatformUserId )
 	{
+		if (!ShouldUsePlatformUserId())
+		{
+			return OnTouchStarted(Window, Location, Force, TouchIndex, PlatformUserId.GetInternalId());
+		}
 		return false;
 	}
 
-	UE_DEPRECATED(4.20, "This function signature is deprecated, use OnTouchMoved that takes a Force")
-	virtual bool OnTouchMoved( const FVector2D& Location, int32 TouchIndex, int32 ControllerId )
-	{
-		return OnTouchMoved(Location, 1.0f, TouchIndex, ControllerId);
-	}
-
-	virtual bool OnTouchMoved( const FVector2D& Location, float Force, int32 TouchIndex, int32 ControllerId )
+	virtual bool OnTouchMoved( const FVector2D& Location, float Force, int32 TouchIndex, FPlatformUserId PlatformUserId )
 	{
+		if (!ShouldUsePlatformUserId())
+		{
+			return OnTouchMoved(Location, Force, TouchIndex, PlatformUserId.GetInternalId());
+		}
 		return false;
 	}
 
-	virtual bool OnTouchEnded( const FVector2D& Location, int32 TouchIndex, int32 ControllerId )
+	virtual bool OnTouchEnded( const FVector2D& Location, int32 TouchIndex, FPlatformUserId PlatformUserId )
 	{
+		if (!ShouldUsePlatformUserId())
+		{
+			return OnTouchEnded(Location, TouchIndex, PlatformUserId.GetInternalId());
+		}
 		return false;
 	}
 
-	virtual bool OnTouchForceChanged(const FVector2D& Location, float Force, int32 TouchIndex, int32 ControllerId)
+	virtual bool OnTouchForceChanged(const FVector2D& Location, float Force, int32 TouchIndex, FPlatformUserId PlatformUserId)
 	{
+		if (!ShouldUsePlatformUserId())
+		{
+			return OnTouchForceChanged(Location, Force, TouchIndex, PlatformUserId.GetInternalId());
+		}
 		return false;
 	}
 
-	virtual bool OnTouchFirstMove(const FVector2D& Location, float Force, int32 TouchIndex, int32 ControllerId)
+	virtual bool OnTouchFirstMove(const FVector2D& Location, float Force, int32 TouchIndex, FPlatformUserId PlatformUserId)
 	{
+		if (!ShouldUsePlatformUserId())
+		{
+			return OnTouchFirstMove(Location, Force, TouchIndex, PlatformUserId.GetInternalId());
+		}
 		return false;
 	}
 
@@ -361,8 +393,12 @@ public:
 
 	}
 
-	virtual bool OnMotionDetected( const FVector& Tilt, const FVector& RotationRate, const FVector& Gravity, const FVector& Acceleration, int32 ControllerId )
+	virtual bool OnMotionDetected( const FVector& Tilt, const FVector& RotationRate, const FVector& Gravity, const FVector& Acceleration, FPlatformUserId PlatformUserId )
 	{
+		if (!ShouldUsePlatformUserId())
+		{
+			return OnMotionDetected(Tilt, RotationRate, Gravity, Acceleration, PlatformUserId.GetInternalId());
+		}
 		return false;
 	}
 
@@ -475,4 +511,80 @@ public:
 	{
 
 	}
+
+	// Deprecate these when engine code has been converted to handle platform user id
+	virtual bool OnControllerAnalog(FGamepadKeyNames::Type KeyName, int32 ControllerId, float AnalogValue)
+	{
+		if (ShouldUsePlatformUserId())
+		{
+			return OnControllerAnalog(KeyName, FPlatformUserId(ControllerId), AnalogValue);
+		}
+		return false;
+	}
+	virtual bool OnControllerButtonPressed(FGamepadKeyNames::Type KeyName, int32 ControllerId, bool IsRepeat)
+	{
+		if (ShouldUsePlatformUserId())
+		{
+			return OnControllerButtonPressed(KeyName, FPlatformUserId(ControllerId), IsRepeat);
+		}
+		return false;
+	}
+	virtual bool OnControllerButtonReleased(FGamepadKeyNames::Type KeyName, int32 ControllerId, bool IsRepeat)
+	{
+		if (ShouldUsePlatformUserId())
+		{
+			return OnControllerButtonReleased(KeyName, FPlatformUserId(ControllerId), IsRepeat);
+		}
+		return false;
+	}
+	virtual bool OnTouchStarted(const TSharedPtr< FGenericWindow >& Window, const FVector2D& Location, float Force, int32 TouchIndex, int32 ControllerId)
+	{
+		if (ShouldUsePlatformUserId())
+		{
+			return OnTouchStarted(Window, Location, Force, TouchIndex, FPlatformUserId(ControllerId));
+		}
+		return false;
+	}
+	virtual bool OnTouchMoved(const FVector2D& Location, float Force, int32 TouchIndex, int32 ControllerId)
+	{
+		if (ShouldUsePlatformUserId())
+		{
+			return OnTouchMoved(Location, Force, TouchIndex, FPlatformUserId(ControllerId));
+		}
+		return false;
+	}
+	virtual bool OnTouchEnded(const FVector2D& Location, int32 TouchIndex, int32 ControllerId)
+	{
+		if (ShouldUsePlatformUserId())
+		{
+			return OnTouchEnded(Location, TouchIndex, FPlatformUserId(ControllerId));
+		}
+		return false;
+	}
+	virtual bool OnTouchForceChanged(const FVector2D& Location, float Force, int32 TouchIndex, int32 ControllerId)
+	{
+		if (ShouldUsePlatformUserId())
+		{
+			return OnTouchForceChanged(Location, Force, TouchIndex, FPlatformUserId(ControllerId));
+		}
+		return false;
+	}
+	virtual bool OnTouchFirstMove(const FVector2D& Location, float Force, int32 TouchIndex, int32 ControllerId)
+	{
+		if (ShouldUsePlatformUserId())
+		{
+			return OnTouchFirstMove(Location, Force, TouchIndex, FPlatformUserId(ControllerId));
+		}
+		return false;
+	}
+	virtual bool OnMotionDetected(const FVector& Tilt, const FVector& RotationRate, const FVector& Gravity, const FVector& Acceleration, int32 ControllerId)
+	{
+		if (ShouldUsePlatformUserId())
+		{
+			return OnMotionDetected(Tilt, RotationRate, Gravity, Acceleration, FPlatformUserId(ControllerId));
+		}
+		return false;
+	}
 };
+
+PRAGMA_ENABLE_DEPRECATION_WARNINGS
\ No newline at end of file
diff --git a/Engine/Source/Runtime/Core/Private/GenericPlatform/GenericPlatformMisc.cpp b/Engine/Source/Runtime/Core/Private/GenericPlatform/GenericPlatformMisc.cpp
--- a/Engine/Source/Runtime/Core/Private/GenericPlatform/GenericPlatformMisc.cpp
+++ b/Engine/Source/Runtime/Core/Private/GenericPlatform/GenericPlatformMisc.cpp
@@ -1474,6 +1474,17 @@ void FGenericPlatformMisc::GetValidTargetPlatforms(class TArray<class FString>&
 	TargetPlatformNames.Add(FPlatformProperties::PlatformName());
 }
 
+FPlatformUserId FGenericPlatformMisc::GetPlatformUserForUserIndex(int32 LocalUserIndex)
+{
+	// These currently map 1:1 but that could change with the input system rework
+	return FPlatformUserId::CreateFromInternalId(LocalUserIndex);
+}
+
+int32 FGenericPlatformMisc::GetUserIndexForPlatformUser(FPlatformUserId PlatformUser)
+{
+	return PlatformUser.GetInternalId();
+}
+
 TArray<uint8> FGenericPlatformMisc::GetSystemFontBytes()
 {
 	return TArray<uint8>();
diff --git a/Engine/Source/Runtime/Core/Public/GenericPlatform/GenericPlatformMisc.h b/Engine/Source/Runtime/Core/Public/GenericPlatform/GenericPlatformMisc.h
--- a/Engine/Source/Runtime/Core/Public/GenericPlatform/GenericPlatformMisc.h
+++ b/Engine/Source/Runtime/Core/Public/GenericPlatform/GenericPlatformMisc.h
@@ -1203,6 +1203,22 @@ public:
 	 */
 	static void GetValidTargetPlatforms(TArray<FString>& TargetPlatformNames);
 
+	/*
+	 * Converts from a local user index to a FPlatformUserId
+	 *
+	 * @param LocalUserIndex	A logical index where 0 should always be a valid user
+	 * @returns					A platform-allocated opaque user id, or PLATFORMUSERID_NONE
+	 */
+	static FPlatformUserId GetPlatformUserForUserIndex(int32 LocalUserIndex);
+
+	/*
+	 * Converts from a FPlatformUserId to local user index
+	 *
+	 * @param LocalUserIndex	A platform-allocated opaque user id
+	 * @returns					A valid logical index, or INDEX_NONE
+	 */
+	static int32 GetUserIndexForPlatformUser(FPlatformUserId PlatformUser);
+
 	/**
 	 * Returns whether the platform wants to use a touch screen for virtual joysticks.
 	 */
diff --git a/Engine/Source/Runtime/Core/Public/Misc/CoreMiscDefines.h b/Engine/Source/Runtime/Core/Public/Misc/CoreMiscDefines.h
--- a/Engine/Source/Runtime/Core/Public/Misc/CoreMiscDefines.h
+++ b/Engine/Source/Runtime/Core/Public/Misc/CoreMiscDefines.h
@@ -111,9 +111,6 @@ enum EForceInit
 enum ENoInit {NoInit};
 enum EInPlace {InPlace};
 
-// Handle type to stably track users on a specific platform
-typedef int32 FPlatformUserId;
-const FPlatformUserId PLATFORMUSERID_NONE = INDEX_NONE;
 #endif // RC_INVOKED
 
 // When passed to pragma message will result in clickable warning in VS
@@ -278,3 +275,67 @@ struct TStaticDeprecateExpression
 	TypeName& operator=(const TypeName&) = delete; \
 	TypeName& operator=(TypeName&&) = delete;
 
+
+/** 
+ * Handle that defines a local user on this platform.
+ * This used to be just a typedef int32 that was used interchangeably as ControllerId and LocalUserIndex.
+ * Moving forward these will be allocated by the platform application layer.
+ */
+struct FPlatformUserId
+{
+	/** Create a default invalid Id */
+	FORCEINLINE FPlatformUserId() : InternalId(INDEX_NONE) {}
+
+	FPlatformUserId(const FPlatformUserId&) = default;
+	FPlatformUserId& operator=(const FPlatformUserId&) = default;
+
+	/** Sees if this is a valid user */
+	FORCEINLINE bool IsValid() const
+	{
+		return InternalId != INDEX_NONE;
+	}
+
+	/** Returns the internal id for debugging/etc */
+	FORCEINLINE int32 GetInternalId() const
+	{
+		return InternalId;
+	}
+
+	/** Explicit function to create from an internal id */
+	FORCEINLINE static FPlatformUserId CreateFromInternalId(int32 InInternalId)
+	{
+		FPlatformUserId IdToReturn;
+		IdToReturn.InternalId = InInternalId;
+		return IdToReturn;
+	}
+
+	FORCEINLINE bool operator==(const FPlatformUserId& Other) const
+	{
+		return InternalId == Other.InternalId;
+	}
+
+	FORCEINLINE bool operator!=(const FPlatformUserId& Other) const
+	{
+		return InternalId != Other.InternalId;
+	}
+
+	FORCEINLINE friend uint32 GetTypeHash(const FPlatformUserId& UserId)
+	{
+		return UserId.InternalId;
+	}
+
+	// This is needed until input and online code handles FPlatformUserId properly
+	UE_DEPRECATED(5.0, "Implicit conversion from user index is deprecated, use FPlatformMisc::GetPlatformUserForUserIndex")
+	FORCEINLINE constexpr FPlatformUserId(int32 InIndex) : InternalId(InIndex) {}
+
+	// This should be deprecated when the online code uniformly handles FPlatformUserId */
+	// UE_DEPRECATED(5.x, "Implicit conversion to user index is deprecated, use FPlatformMisc::GetUserIndexForPlatformUser")
+	FORCEINLINE constexpr operator int32() const { return InternalId; }
+
+private:
+	/** Raw id, will be allocated by application layer */
+	int32 InternalId;
+};
+
+/** Static invalid platform user */
+const FPlatformUserId PLATFORMUSERID_NONE;
diff --git a/Engine/Source/Runtime/Engine/Classes/Engine/Engine.h b/Engine/Source/Runtime/Engine/Classes/Engine/Engine.h
--- a/Engine/Source/Runtime/Engine/Classes/Engine/Engine.h
+++ b/Engine/Source/Runtime/Engine/Classes/Engine/Engine.h
@@ -2222,7 +2222,27 @@ public:
 	 *	Returns the first ULocalPlayer that matches the given ControllerId. 
 	 *  This will search across all world contexts.
 	 */
-	class ULocalPlayer* FindFirstLocalPlayerFromControllerId(int32 ControllerId) const;
+	ULocalPlayer* FindFirstLocalPlayerFromControllerId(int32 ControllerId) const;
+
+	/** 
+	 * If true, we're running in a backward compatible mode where FPlatformUserId and ControllerId are the same.
+	 * If false, there can be more than one local player with the same platform user id
+	 */
+	virtual bool IsControllerIdUsingPlatformUserId() const { return true; }
+
+	/**
+	 * Returns the first ULocalPlayer that matches the given platform user id, or the first player if the id is invalid
+	 * This will search across all world contexts.
+	 */
+	ULocalPlayer* FindFirstLocalPlayerFromPlatformUserId(FPlatformUserId PlatformUserId) const;
+
+	/**
+	 * Returns the first LocalPlayer that matches the given platform user id
+	 *
+	 * @param	PlatformUserId	Platform user id to search for
+	 * @return	The player that has the PlatformUserId specified, or nullptr if no players have that PlatformUserId
+	 */
+	ULocalPlayer* GetLocalPlayerFromPlatformUserId(UWorld* InWorld, const FPlatformUserId PlatformUserId) const;
 
 	/**
 	 * return the number of entries in the GamePlayers array
diff --git a/Engine/Source/Runtime/Engine/Classes/Engine/LocalPlayer.h b/Engine/Source/Runtime/Engine/Classes/Engine/LocalPlayer.h
--- a/Engine/Source/Runtime/Engine/Classes/Engine/LocalPlayer.h
+++ b/Engine/Source/Runtime/Engine/Classes/Engine/LocalPlayer.h
@@ -204,6 +204,10 @@ public:
 	DECLARE_EVENT_TwoParams(ULocalPlayer, FOnControllerIdChanged, int32 /*NewId*/, int32 /*OldId*/);
 	FOnControllerIdChanged& OnControllerIdChanged() const { return OnControllerIdChangedEvent; }
 
+	/** Event called when this local player has been assigned to a new platform-level user */
+	DECLARE_EVENT_TwoParams(ULocalPlayer, FOnPlatformUserIdChanged, FPlatformUserId /*NewId*/, FPlatformUserId /*OldId*/);
+	FOnPlatformUserIdChanged& OnPlatformUserIdChanged() { return OnPlatformUserIdChangedEvent; }
+
 private:
 	TArray<FSceneViewStateReference> ViewStates;
 
@@ -213,6 +217,12 @@ private:
 
 	mutable FOnControllerIdChanged OnControllerIdChangedEvent;
 
+	/** The platform user this player is assigned to, could correspond to multiple input devices */
+	FPlatformUserId PlatformUserId;
+
+	/** Event called when platform user id changes */
+	FOnPlatformUserIdChanged OnPlatformUserIdChangedEvent;
+
 	FSubsystemCollection<ULocalPlayerSubsystem> SubsystemCollection;
 
 public:
@@ -265,8 +275,8 @@ public:
 	const FReply& GetSlateOperations() const { return SlateOperations; }
 
 	/** Get the SlateUser that this LocalPlayer corresponds to */
-	TSharedPtr<FSlateUser> GetSlateUser();
-	TSharedPtr<const FSlateUser> GetSlateUser() const;
+	virtual TSharedPtr<FSlateUser> GetSlateUser();
+	virtual TSharedPtr<const FSlateUser> GetSlateUser() const;
 
 	/**
 	 * Get the world the players actor belongs to
@@ -340,7 +350,7 @@ public:
 	* @param	StereoPass - whether we are drawing the full viewport, or a stereo left / right pass
 	* @return	true if the view options were filled in. false in various fail conditions.
 	*/
-	bool CalcSceneViewInitOptions(
+	virtual bool CalcSceneViewInitOptions(
 		struct FSceneViewInitOptions& OutInitOptions, 
 		FViewport* Viewport,
 		class FViewElementDrawer* ViewDrawer = NULL,
@@ -410,6 +420,22 @@ public:
 	 */
 	int32 GetControllerId() const { return ControllerId; }
 
+	/**
+	 * Changes the platform user that is assigned to this player
+	 */
+	virtual void SetPlatformUserId(FPlatformUserId InPlatformUserId);
+
+	/**
+	 * Returns the platform user that is assigned to this player
+	 */
+	FPlatformUserId GetPlatformUserId() const { return PlatformUserId; }
+
+	/**
+	 * Returns the logical local player index where 0 is the first player
+	 * By default, this uses index in the game instance array
+	 */
+	virtual int32 GetLocalPlayerIndex() const;
+
 	/** 
 	 * Retrieves this player's name/tag from the online subsystem
 	 * if this function returns a non-empty string, the returned name will replace the "Name" URL parameter
@@ -429,12 +455,16 @@ public:
 	 */
 	virtual FString GetGameLoginOptions() const { return TEXT(""); }
 
-	/** 
-	 * Retrieves this player's unique net ID from the online subsystem 
+	// This should be deprecated when engine code has been changed to expect FPlatformUserId
+	// UE_DEPRECATED(5.x, "Use GetUniqueNetIdForPlatformUser instead")
+	FUniqueNetIdRepl GetUniqueNetIdFromCachedControllerId() const;
+
+	/**
+	 * Retrieves this player's unique net ID from the online subsystem using the platform user Id
 	 *
 	 * @return unique Id associated with this player
 	 */
-	FUniqueNetIdRepl GetUniqueNetIdFromCachedControllerId() const;
+	virtual FUniqueNetIdRepl GetUniqueNetIdForPlatformUser() const;
 
 	/** 
 	 * Retrieves this player's unique net ID that was previously cached
@@ -451,9 +481,9 @@ public:
 	 *
 	 * @return unique Id associated with this player
 	 */
-	FUniqueNetIdRepl GetPreferredUniqueNetId() const;
+	virtual FUniqueNetIdRepl GetPreferredUniqueNetId() const;
 
-	/** Returns true if the cached unique net id, is the one assigned to the controller id from the OSS */
+	UE_DEPRECATED(5.0, "Platform User Id now has priority over ControllerId, these are not expected to be the same")
 	bool IsCachedUniqueNetIdPairedWithControllerId() const;
 
 	/**
@@ -496,7 +526,7 @@ public:
 	/**
 	 * Clear cached view state.  Suitable for calling when cleaning up the world but the view state has some references objects (usually mids) owned by the world (thus preventing GC) 
 	 */
-	void CleanupViewState();
+	virtual void CleanupViewState();
 
 	/** Locked view state needs access to GetViewPoint. */
 	friend class FLockedViewState;
diff --git a/Engine/Source/Runtime/Engine/Private/LocalPlayer.cpp b/Engine/Source/Runtime/Engine/Private/LocalPlayer.cpp
--- a/Engine/Source/Runtime/Engine/Private/LocalPlayer.cpp
+++ b/Engine/Source/Runtime/Engine/Private/LocalPlayer.cpp
@@ -789,6 +789,8 @@ bool ULocalPlayer::CalcSceneViewInitOptions(
 
 	ViewInitOptions.SceneViewStateInterface = ViewStates[ViewIndex].GetReference();
 	ViewInitOptions.ViewActor = PlayerController->GetViewTarget();
+
+	// TODO: Switch to GetLocalPlayerIndex during GetControllerId deprecation, this is only used by MotionControllerComponent
 	ViewInitOptions.PlayerIndex = GetControllerId();
 	ViewInitOptions.ViewElementDrawer = ViewDrawer;
 	ViewInitOptions.BackgroundColor = FLinearColor::Black;
@@ -1528,13 +1530,42 @@ void ULocalPlayer::SetControllerId( int32 NewControllerId )
 		ControllerId = -1;
 
 		// see if another player is already using this ControllerId; if so, swap controllerIds with them
+		// TODO: Re-evaluate if this swap logic makes sense during controller id deprecation
 		GEngine->SwapControllerId(this, CurrentControllerId, NewControllerId);
 		ControllerId = NewControllerId;
 
 		OnControllerIdChanged().Broadcast(NewControllerId, CurrentControllerId);
+
+		if (GEngine->IsControllerIdUsingPlatformUserId())
+		{
+			// This won't recurse back because we've already modified ControllerId
+			SetPlatformUserId(FPlatformMisc::GetPlatformUserForUserIndex(NewControllerId));
+		}
+	}
+}
+
+void ULocalPlayer::SetPlatformUserId(FPlatformUserId InPlatformUserId)
+{
+	if (InPlatformUserId != PlatformUserId)
+	{
+		const FPlatformUserId CurrentPlatformUserId = PlatformUserId;
+
+		PlatformUserId = InPlatformUserId;
+		OnPlatformUserIdChanged().Broadcast(InPlatformUserId, CurrentPlatformUserId);
+
+		if (GEngine->IsControllerIdUsingPlatformUserId())
+		{
+			SetControllerId(FPlatformMisc::GetUserIndexForPlatformUser(PlatformUserId));
+		}
 	}
 }
 
+int32 ULocalPlayer::GetLocalPlayerIndex() const
+{
+	// TODO: Add caching
+	return GetIndexInGameInstance();
+}
+
 FString ULocalPlayer::GetNickname() const
 {
 	UWorld* World = GetWorld();
@@ -1542,7 +1573,7 @@ FString ULocalPlayer::GetNickname() const
 	{
 		// Try to get platform identity first
 		FString PlatformNickname;
-		if (UOnlineEngineInterface::Get()->GetPlayerPlatformNickname(World, ControllerId, PlatformNickname))
+		if (UOnlineEngineInterface::Get()->GetPlayerPlatformNickname(World, PlatformUserId, PlatformNickname))
 		{
 			return PlatformNickname;
 		}
@@ -1568,6 +1599,17 @@ FUniqueNetIdRepl ULocalPlayer::GetUniqueNetIdFromCachedControllerId() const
 	return FUniqueNetIdRepl();
 }
 
+FUniqueNetIdRepl ULocalPlayer::GetUniqueNetIdForPlatformUser() const
+{
+	UWorld* World = GetWorld();
+	if (World != nullptr)
+	{
+		return FUniqueNetIdRepl(UOnlineEngineInterface::Get()->GetUniquePlayerId(World, PlatformUserId));
+	}
+
+	return FUniqueNetIdRepl();
+}
+
 FUniqueNetIdRepl ULocalPlayer::GetCachedUniqueNetId() const
 {
 	return CachedUniqueNetId;
@@ -1587,8 +1629,8 @@ FUniqueNetIdRepl ULocalPlayer::GetPreferredUniqueNetId() const
 		return GetCachedUniqueNetId();
 	}
 
-	// If the cached unique net id is not valid, then get the one paired with the controller
-	return GetUniqueNetIdFromCachedControllerId();
+	// If the cached unique net id is not valid, then use the platfomr user
+	return GetUniqueNetIdForPlatformUser();
 }
 
 bool ULocalPlayer::IsCachedUniqueNetIdPairedWithControllerId() const
diff --git a/Engine/Source/Runtime/Engine/Private/UnrealEngine.cpp b/Engine/Source/Runtime/Engine/Private/UnrealEngine.cpp
--- a/Engine/Source/Runtime/Engine/Private/UnrealEngine.cpp
+++ b/Engine/Source/Runtime/Engine/Private/UnrealEngine.cpp
@@ -11951,6 +11951,41 @@ ULocalPlayer* UEngine::FindFirstLocalPlayerFromControllerId(int32 ControllerId)
 	return nullptr;
 }
 
+ULocalPlayer* UEngine::FindFirstLocalPlayerFromPlatformUserId(FPlatformUserId PlatformUserId) const
+{
+	for (auto It = WorldList.CreateConstIterator(); It; ++It)
+	{
+		const FWorldContext& Context = *It;
+		if (Context.World() && Context.OwningGameInstance && (Context.WorldType == EWorldType::Game || Context.WorldType == EWorldType::PIE))
+		{
+			const TArray<class ULocalPlayer*>& LocalPlayers = Context.OwningGameInstance->GetLocalPlayers();
+
+			for (ULocalPlayer* LocalPlayer : LocalPlayers)
+			{
+				if (LocalPlayer && (!PlatformUserId.IsValid() || LocalPlayer->GetPlatformUserId() == PlatformUserId))
+				{
+					return LocalPlayer;
+				}
+			}
+		}
+	}
+
+	return nullptr;
+}
+
+ULocalPlayer* UEngine::GetLocalPlayerFromPlatformUserId(UWorld* InWorld, const FPlatformUserId PlatformUserId) const
+{
+	for (ULocalPlayer* const Player : GetGamePlayers(InWorld))
+	{
+		if (Player && Player->GetPlatformUserId() == PlatformUserId)
+		{
+			return Player;
+		}
+	}
+
+	return nullptr;
+}
+
 int32 UEngine::GetNumGamePlayers(UWorld *InWorld)
 {
 	return GetGamePlayers(InWorld).Num();
diff --git a/Engine/Source/Runtime/Engine/Public/VideoRecordingSystem.h b/Engine/Source/Runtime/Engine/Public/VideoRecordingSystem.h
--- a/Engine/Source/Runtime/Engine/Public/VideoRecordingSystem.h
+++ b/Engine/Source/Runtime/Engine/Public/VideoRecordingSystem.h
@@ -56,7 +56,7 @@ struct FVideoRecordingParameters
 		, bAutoStart(false)
 		, bAutoContinue(false)
 		, bExportToLibrary(true)
-		, UserId(0)
+		, UserId(PLATFORMUSERID_NONE)
 	{}
 
 	FVideoRecordingParameters(uint64 InRecordingLengthSeconds, bool bInAutoStart, bool bInAutoContinue, bool bInExportToLibrary, FPlatformUserId InUserId)
