Remove unnecessary TruncToInt causing deprecation warnings

diff --git a/Engine/Source/Runtime/ApplicationCore/Private/IOS/IOSCursor.cpp b/Engine/Source/Runtime/ApplicationCore/Private/IOS/IOSCursor.cpp
--- a/Engine/Source/Runtime/ApplicationCore/Private/IOS/IOSCursor.cpp
+++ b/Engine/Source/Runtime/ApplicationCore/Private/IOS/IOSCursor.cpp
@@ -49,10 +49,10 @@ void FIOSCursor::Lock(const RECT* const Bounds)
 	}
 	else
 	{
-		CursorClipRect.Min.X = FMath::TruncToInt(Bounds->left);
-		CursorClipRect.Min.Y = FMath::TruncToInt(Bounds->top);
-		CursorClipRect.Max.X = FMath::TruncToInt(Bounds->right) - 1;
-		CursorClipRect.Max.Y = FMath::TruncToInt(Bounds->bottom) - 1;
+		CursorClipRect.Min.X = Bounds->left;
+		CursorClipRect.Min.Y = Bounds->top;
+		CursorClipRect.Max.X = Bounds->right - 1;
+		CursorClipRect.Max.Y = Bounds->bottom - 1;
 	}
 
 	FVector2D Position = GetPosition();
