Remove wrongful throttling that doesn't honor editor settings and tool overrides
  - This is already handled in the engine loop and could lead to multiple sleeps being done (see ShouldThrottleCPUUsage).

diff --git a/Engine/Source/Runtime/ApplicationCore/Private/Windows/WindowsPlatformApplicationMisc.cpp b/Engine/Source/Runtime/ApplicationCore/Private/Windows/WindowsPlatformApplicationMisc.cpp
--- a/Engine/Source/Runtime/ApplicationCore/Private/Windows/WindowsPlatformApplicationMisc.cpp
+++ b/Engine/Source/Runtime/ApplicationCore/Private/Windows/WindowsPlatformApplicationMisc.cpp
@@ -142,20 +142,6 @@ void FWindowsPlatformApplicationMisc::PumpMessages(bool bFromMainLoop)
 	bool bHasFocus = FApp::HasFocus();
 	static bool bHadFocus = false;
 
-#if WITH_EDITOR
-	// If editor thread doesn't have the focus, don't suck up too much CPU time.
-	if( GIsEditor )
-	{
-		if( !bHasFocus )
-		{
-			TRACE_CPUPROFILER_EVENT_SCOPE(EditorIsInBackgroundSleep);
-			// Sleep for a bit to not eat up all CPU time.
-			FPlatformProcess::Sleep(0.005f);
-		}
-		bHadFocus = bHasFocus;
-	}
-#endif
-
 #if !UE_SERVER
 	// For non-editor clients, record if the active window is in focus
 	if( bHadFocus != bHasFocus )
