Fix current gamepad behavior in iOS15 to reflect Apple's change of callback order. #rb jack.porter #jira UE-132399

diff --git a/Engine/Source/Runtime/ApplicationCore/Private/Apple/AppleControllerInterface.cpp b/Engine/Source/Runtime/ApplicationCore/Private/Apple/AppleControllerInterface.cpp
--- a/Engine/Source/Runtime/ApplicationCore/Private/Apple/AppleControllerInterface.cpp
+++ b/Engine/Source/Runtime/ApplicationCore/Private/Apple/AppleControllerInterface.cpp
@@ -29,6 +29,7 @@ FAppleControllerInterface::FAppleControllerInterface( const TSharedRef< FGeneric
 		[notificationCenter addObserverForName:GCControllerDidConnectNotification object:nil queue:currentQueue usingBlock:^(NSNotification* Notification)
 		{
 			HandleConnection(Notification.object);
+            SetCurrentController(Notification.object);
 		}];
 
 		dispatch_async(dispatch_get_main_queue(), ^
