Generic Accessible Interfaces refactor to support multiple users and focuses.

- Moved FAccessibleEventArgs out of FGenericAccessibleMessageHandler to avoid the need to include GenericAccessibleInterfaces.h in header files where all
they really need is a forward declaration of the struct.
- Removed instances of FGenericAccessibleMessageHandler::FAccessibleEventArgs to FAccessibleEventArgs.

- Updated API for IAccessibleWidget and IAccessibleWindow that deal with getting and setting focus to also take a user index parameter. This facilitates multiple accessible users for the same application.
-Updated platform files that leverage the old IAccessibleWidget and IAccessibleWindow focus API

- Introduced FGenericAccessibleUser and FGenericAccessibleUserRegistry to GenericAccessibleInterfaces.h. An FGenericAccessibleUser is meant to be analogous to FSlateUser that keeps track of hardware users' focus. An FGenericAccessibleUserRegistry is a registry of FGenericAccessibleUsers. These interfaces are meant to insulate FGenericAccessibleMessageHandler from storing and tracking
focus etc when that should be an implementation detail by other modules.
- Refactored how FSlateCoreAccessibleWidgets set focus in the application. Accessible focus is now unified to a single code path that starts from FSlateAccessibleMessageHandler. Requests to shift accessible focus should be done by raising a Slate accessible event or through IAccessibleWidget::SetUserFocus()
- Introduced finer grain querying of focus capabilities for an IAccessibleWidget. Users can now check if an AccessibleWidget supports regular focus (keyboard/gamepad), if it supports accessible focus and if the widget can currently accept regular or accessible focus.
-Updated platforms to use the new FGenericAccessibleUser and FGenericAccessibleUserRegistry ApI to update and retrieve user focus.

diff --git a/Engine/Source/Runtime/ApplicationCore/Private/GenericPlatform/Accessibility/GenericAccessibleInterfaces.cpp b/Engine/Source/Runtime/ApplicationCore/Private/GenericPlatform/Accessibility/GenericAccessibleInterfaces.cpp
--- a/Engine/Source/Runtime/ApplicationCore/Private/GenericPlatform/Accessibility/GenericAccessibleInterfaces.cpp
+++ b/Engine/Source/Runtime/ApplicationCore/Private/GenericPlatform/Accessibility/GenericAccessibleInterfaces.cpp
@@ -15,6 +15,12 @@ FAutoConsoleVariableRef AllowAccessibilityRef(
 	TEXT("If false, all queries from accessible APIs will be ignored. On some platforms, the application must be restarted in order to take effect.")
 );
 
+FGenericAccessibleMessageHandler::FGenericAccessibleMessageHandler()
+	: bApplicationIsAccessible(false)
+	, bIsActive(false)
+	, AccessibleUserRegistry(MakeShared<FGenericAccessibleUserRegistry>())
+{}
+
 bool FGenericAccessibleMessageHandler::ApplicationIsAccessible() const
 {
 	return bApplicationIsAccessible && GAllowAccessibility;
@@ -39,4 +45,71 @@ void FGenericAccessibleMessageHandler::SetActive(bool bActive)
 		}
 	}
 }
+
+bool FGenericAccessibleUserRegistry::RegisterUser(const TSharedRef<FGenericAccessibleUser>& User)
+{
+	const FAccessibleUserIndex UserIndex = User->GetIndex();
+	if (!IsUserRegistered(UserIndex))
+	{
+		UsersMap.Add(UserIndex, User);
+		User->OnRegistered();
+		return true;
+	}
+	UE_LOG(LogAccessibility, Verbose, TEXT("Registering accessible user %d failed. Another accessible user with accessible index %d already registered."), UserIndex, UserIndex);
+	return false;
+}
+
+bool FGenericAccessibleUserRegistry::UnregisterUser(const FAccessibleUserIndex UserIndex)
+{
+	TSharedRef<FGenericAccessibleUser>* User = UsersMap.Find(UserIndex);
+	if (User)
+	{
+		(*User)->OnUnregistered();
+		UsersMap.Remove(UserIndex);
+		return true;
+	}
+	UE_LOG(LogAccessibility, Verbose, TEXT("Failed to unregister accessible user with accessible index %d. No accessible user with accessible index %d currently registered."), UserIndex, UserIndex);
+	return false;
+}
+
+void FGenericAccessibleUserRegistry::UnregisterAllUsers()
+{
+	for (TPair<FAccessibleUserIndex, TSharedRef<FGenericAccessibleUser>>& UserPair : UsersMap)
+	{
+		UserPair.Value->OnUnregistered();
+	}
+	UsersMap.Empty();
+}
+
+bool FGenericAccessibleUserRegistry::IsUserRegistered(const FAccessibleUserIndex UserIndex) const
+{
+	return UsersMap.Contains(UserIndex);
+}
+
+TSharedPtr<FGenericAccessibleUser> FGenericAccessibleUserRegistry::GetUser(const FAccessibleUserIndex UserIndex) const
+{
+	if (IsUserRegistered(UserIndex))
+	{
+		return UsersMap[UserIndex];
+	}
+	return nullptr;
+}
+
+int32 FGenericAccessibleUserRegistry::GetNumberofUsers() const
+{
+	return UsersMap.Num();
+}
+
+TArray<TSharedRef<FGenericAccessibleUser>> FGenericAccessibleUserRegistry::GetAllUsers() const
+{
+	TArray<TSharedRef<FGenericAccessibleUser>> AllUsers;
+	AllUsers.Reserve(UsersMap.Num());
+	for (const TPair<FAccessibleUserIndex, TSharedRef<FGenericAccessibleUser>>& UserPair : UsersMap)
+	{
+		AllUsers.Add(UserPair.Value);
+	}
+	return AllUsers;
+}
+
 #endif
+
diff --git a/Engine/Source/Runtime/ApplicationCore/Private/IOS/IOSApplication.cpp b/Engine/Source/Runtime/ApplicationCore/Private/IOS/IOSApplication.cpp
--- a/Engine/Source/Runtime/ApplicationCore/Private/IOS/IOSApplication.cpp
+++ b/Engine/Source/Runtime/ApplicationCore/Private/IOS/IOSApplication.cpp
@@ -71,6 +71,10 @@ void FIOSApplication::SetMessageHandler( const TSharedRef< FGenericApplicationMe
 void FIOSApplication::SetAccessibleMessageHandler(const TSharedRef<FGenericAccessibleMessageHandler>& InAccessibleMessageHandler)
 {
 	GenericApplication::SetAccessibleMessageHandler(InAccessibleMessageHandler);
+	// This user is what IOS Voiceover will interact with
+	FGenericAccessibleUserRegistry& UserRegistry = InAccessibleMessageHandler->GetAccessibleUserRegistry();
+	// We failed to register the primary user, this should only happen if another user with the 0th index has already been registered.
+	ensure(UserRegistry.RegisterUser(MakeShared<FGenericAccessibleUser>(FGenericAccessibleUserRegistry::GetPrimaryUserIndex())));
 	InAccessibleMessageHandler->SetAccessibleEventDelegate(FGenericAccessibleMessageHandler::FAccessibleEvent::CreateRaw(this, &FIOSApplication::OnAccessibleEventRaised));
 	InAccessibleMessageHandler->SetActive(UIAccessibilityIsVoiceOverRunning());
 }
@@ -268,7 +272,7 @@ FAutoConsoleVariableRef IOSAccessibleAnnouncementDealyRef(
 );
 
 
-void FIOSApplication::OnAccessibleEventRaised(const FGenericAccessibleMessageHandler::FAccessibleEventArgs& Args)
+void FIOSApplication::OnAccessibleEventRaised(const FAccessibleEventArgs& Args)
 {
 	// This should only be triggered by the accessible message handler which initiates from the Slate thread.
 	check(IsInGameThread());
diff --git a/Engine/Source/Runtime/ApplicationCore/Private/Mac/Accessibility/MacAccessibilityElement.cpp b/Engine/Source/Runtime/ApplicationCore/Private/Mac/Accessibility/MacAccessibilityElement.cpp
--- a/Engine/Source/Runtime/ApplicationCore/Private/Mac/Accessibility/MacAccessibilityElement.cpp
+++ b/Engine/Source/Runtime/ApplicationCore/Private/Mac/Accessibility/MacAccessibilityElement.cpp
@@ -266,7 +266,8 @@
 			TSharedPtr<IAccessibleWidget> Widget = [FMacAccessibilityManager AccessibilityManager].MacApplication->GetAccessibleMessageHandler()->GetAccessibleWidgetFromId(TempId);
 			if (Widget.IsValid() && Widget->IsEnabled() && !Widget->IsHidden())
 			{
-				Widget->SetFocus();
+				// Voiceover only recognizes 1 user, the primary accessible user.
+				Widget->SetUserFocus(FGenericAccessibleUserRegistry::GetPrimaryUserIndex());
 			}
 		}, @[ NSDefaultRunLoopMode, UnrealIMEEventMode, UnrealShowEventMode, UnrealResizeEventMode, UnrealFullscreenEventMode], false); // GameThreadCall
 	}
diff --git a/Engine/Source/Runtime/ApplicationCore/Private/Mac/MacApplication.cpp b/Engine/Source/Runtime/ApplicationCore/Private/Mac/MacApplication.cpp
--- a/Engine/Source/Runtime/ApplicationCore/Private/Mac/MacApplication.cpp
+++ b/Engine/Source/Runtime/ApplicationCore/Private/Mac/MacApplication.cpp
@@ -275,6 +275,11 @@ void FMacApplication::SetMessageHandler(const TSharedRef<FGenericApplicationMess
 void FMacApplication::SetAccessibleMessageHandler(const TSharedRef<FGenericAccessibleMessageHandler>& InAccessibleMessageHandler)
 {
 	GenericApplication::SetAccessibleMessageHandler(InAccessibleMessageHandler);
+	// We register the primary user (keyboard).
+	// This user is what Mac Voiceover will interact with 
+	FGenericAccessibleUserRegistry& UserRegistry = AccessibleMessageHandler->GetAccessibleUserRegistry();
+	// We failed to register the primary user, this should only happen if another user with the 0th index has already been registered.
+	ensure(UserRegistry.RegisterUser(MakeShared<FGenericAccessibleUser>(FGenericAccessibleUserRegistry::GetPrimaryUserIndex())));
 	InAccessibleMessageHandler->SetAccessibleEventDelegate(FGenericAccessibleMessageHandler::FAccessibleEvent::CreateRaw(this, &FMacApplication::OnAccessibleEventRaised));
 	
 	MainThreadCall(^{
@@ -2299,7 +2304,7 @@ FAutoConsoleVariableRef MacAccessibleAnnouncementDealyRef(
 	TEXT("We need to introduce a small delay to avoid OSX system accessibility announcements from stomping on our requested user announcement. Delays <= 0.05f are too short and result in the announcement being dropped. Delays ~0.075f result in unstable delivery")
 );
 
-void FMacApplication::OnAccessibleEventRaised(const FGenericAccessibleMessageHandler::FAccessibleEventArgs& Args)
+void FMacApplication::OnAccessibleEventRaised(const FAccessibleEventArgs& Args)
 {
 	// This should only be triggered by the accessible message handler which initiates from the Slate thread.
 	check(IsInGameThread());
diff --git a/Engine/Source/Runtime/ApplicationCore/Private/Windows/Accessibility/WindowsUIAManager.cpp b/Engine/Source/Runtime/ApplicationCore/Private/Windows/Accessibility/WindowsUIAManager.cpp
--- a/Engine/Source/Runtime/ApplicationCore/Private/Windows/Accessibility/WindowsUIAManager.cpp
+++ b/Engine/Source/Runtime/ApplicationCore/Private/Windows/Accessibility/WindowsUIAManager.cpp
@@ -86,7 +86,14 @@ FWindowsUIAManager::FWindowsUIAManager(const FWindowsApplication& InApplication)
 
 void FWindowsUIAManager::OnAccessibleMessageHandlerChanged()
 {
-	WindowsApplication.GetAccessibleMessageHandler()->SetAccessibleEventDelegate(FGenericAccessibleMessageHandler::FAccessibleEvent::CreateRaw(this, &FWindowsUIAManager::OnEventRaised));
+	TSharedRef<FGenericAccessibleMessageHandler> MessageHandler = WindowsApplication.GetAccessibleMessageHandler();
+	// We register the primary user (keyboard) 
+	// This user is what UIA will interact with
+	FGenericAccessibleUserRegistry& UserRegistry = MessageHandler->GetAccessibleUserRegistry();
+	// We failed to register the primary user, this should only happen if another user with the 0th index has already been registered.
+	ensure(UserRegistry.RegisterUser(MakeShared<FGenericAccessibleUser>(FGenericAccessibleUserRegistry::GetPrimaryUserIndex())));
+
+	MessageHandler->SetAccessibleEventDelegate(FGenericAccessibleMessageHandler::FAccessibleEvent::CreateRaw(this, &FWindowsUIAManager::OnEventRaised));
 }
 
 FWindowsUIAManager::~FWindowsUIAManager()
@@ -173,7 +180,7 @@ void FWindowsUIAManager::OnAccessibilityDisabled()
 	}
 }
 
-void FWindowsUIAManager::OnEventRaised(const FGenericAccessibleMessageHandler::FAccessibleEventArgs& Args)
+void FWindowsUIAManager::OnEventRaised(const FAccessibleEventArgs& Args)
 {
 	if (UiaClientsAreListening())
 	{
@@ -185,7 +192,7 @@ void FWindowsUIAManager::OnEventRaised(const FGenericAccessibleMessageHandler::F
 		{
 			// On focus change, emit a generic FocusChanged event as well as a per-Provider PropertyChanged event
 			// todo: handle difference between any focus vs keyboard focus
-			if (Args.Widget->HasFocus())
+			if (Args.Widget->HasUserFocus(0))
 			{
 				UiaRaiseAutomationEvent(&ScopedProvider.Provider, UIA_AutomationFocusChangedEventId);
 			}
diff --git a/Engine/Source/Runtime/ApplicationCore/Private/Windows/Accessibility/WindowsUIAWidgetProvider.cpp b/Engine/Source/Runtime/ApplicationCore/Private/Windows/Accessibility/WindowsUIAWidgetProvider.cpp
--- a/Engine/Source/Runtime/ApplicationCore/Private/Windows/Accessibility/WindowsUIAWidgetProvider.cpp
+++ b/Engine/Source/Runtime/ApplicationCore/Private/Windows/Accessibility/WindowsUIAWidgetProvider.cpp
@@ -560,7 +560,8 @@ HRESULT STDCALL FWindowsUIAWidgetProvider::GetPropertyValue(PROPERTYID propertyI
 		case UIA_HasKeyboardFocusPropertyId:
 		{
 			pRetVal->vt = VT_BOOL;
-			pRetVal->boolVal = Widget->HasFocus() ? VARIANT_TRUE : VARIANT_FALSE;
+			// UIA only recognizes 1 user, the primary accessible user
+			pRetVal->boolVal = Widget->HasUserFocus(FGenericAccessibleUserRegistry::GetPrimaryUserIndex()) ? VARIANT_TRUE : VARIANT_FALSE;
 		}
 			break;
 		case UIA_HelpTextPropertyId:
@@ -806,7 +807,8 @@ HRESULT STDCALL FWindowsUIAWidgetProvider::SetFocus()
 		{
 			if (Widget->SupportsFocus())
 			{
-				Widget->SetFocus();
+				// UIA only recognizes 1 user, the primary accessible user for the application
+				Widget->SetUserFocus(FGenericAccessibleUserRegistry::GetPrimaryUserIndex());
 				ReturnValue = S_OK;
 				return;
 			}
@@ -958,7 +960,8 @@ HRESULT STDCALL FWindowsUIAWindowProvider::GetFocus(IRawElementProviderFragment*
 		*pRetVal = nullptr;
 		if (IsValid())
 		{
-			TSharedPtr<IAccessibleWidget> Focus = Widget->AsWindow()->GetFocusedWidget();
+			// UIA only assumes 1 user so we always assume that the primary user (the keyboard user) is using the accessibility services
+			TSharedPtr<IAccessibleWidget> Focus = Widget->AsWindow()->GetUserFocusedWidget(0);
 			if (Focus.IsValid())
 			{
 				*pRetVal = static_cast<IRawElementProviderFragment*>(&UIAManager->GetWidgetProvider(Focus.ToSharedRef()));
diff --git a/Engine/Source/Runtime/ApplicationCore/Public/GenericPlatform/Accessibility/GenericAccessibleInterfaces.h b/Engine/Source/Runtime/ApplicationCore/Public/GenericPlatform/Accessibility/GenericAccessibleInterfaces.h
--- a/Engine/Source/Runtime/ApplicationCore/Public/GenericPlatform/Accessibility/GenericAccessibleInterfaces.h
+++ b/Engine/Source/Runtime/ApplicationCore/Public/GenericPlatform/Accessibility/GenericAccessibleInterfaces.h
@@ -86,6 +86,138 @@ enum class EAccessibleEvent : uint8
 	WidgetRemoved
 };
 
+/** An index that uniquely identifies every registered accessible user in the application. */
+typedef int32 FAccessibleUserIndex;
+
+/**
+ * An accessible user is an input source that can interact with the application and needs to provide accessibility services such as a screen reader.
+ * This base class provides the basic information all accessible users need to implement to ensure accessibility services function properly.
+ * Users should subclass this class to provide additional functionality for their custom accessible users.
+ * All accessible users must be registered with an FGenericAccessibleUserRegistry for the users to be retrieved and interacted with.
+ * @see FGenericAccessibleUserRegistry
+ */
+class APPLICATIONCORE_API FGenericAccessibleUser
+{
+	friend class FGenericAccessibleUserRegistry;
+public:
+	FGenericAccessibleUser(const FAccessibleUserIndex InUserIndex)
+		: UserIndex(InUserIndex)
+	{
+		
+	}
+	virtual ~FGenericAccessibleUser() = default;
+	/** Returns the unique index that identifies this accessible user.*/
+	FAccessibleUserIndex GetIndex() const
+	{
+		return UserIndex;
+	}
+	/**
+	 * Returns the accessible widget the accessible user is currently focused on. If no accessible widget is focused, nullptr is returned.
+	 * @see IAccessibleWidget
+	 */
+	TSharedPtr<IAccessibleWidget> GetFocusedAccessibleWidget() const
+	{
+		return FocusedAccessibleWidget.IsValid() ? FocusedAccessibleWidget.Pin() : nullptr;
+	}
+	/**
+	 * Sets which accessible widget the accessible user is currently focused on.
+	 *
+	 * Note: This function is meant to directly set the widget an accessible user is focused on, bypassing all of the focus change accessible events.
+	 * This should only be used by FGenericAccessibleMessageHandler and its subclasses to update the  the accessible focus for an accessible user.
+	 * Most of the time, you want to use IAccessibleWidget::SetUserFocus to request an accessible widget be focused on by a user. It raises the appropriate accessible events
+	 * so listeners for accessible events can do the appropriate event handling for the focus change event.
+	 * @param InWidget The widget this accessible user will be focused on.
+	 * @return True if the passed in widget was successfully set. Else returns false.
+	 */
+	bool SetFocusedAccessibleWidget(const TSharedRef<IAccessibleWidget>& InWidget)
+	{
+		FocusedAccessibleWidget = InWidget;
+		return true;
+	}
+	/** Clear the accessible focused widget. */
+	void ClearFocusedAccessibleWidget()
+	{
+		FocusedAccessibleWidget.Reset();
+	}
+protected:
+	/**
+	 * Called by FGenericAccessibleUserRegistry::RegisterUser.
+	 * @see FGenericAccessibleUserRegistry
+	 */
+	virtual void OnRegistered() {}
+	/**
+	 * Called by FGenericAccessibleUserRegistry::UnregisterUser.
+	 * @see FGenericAccessibleUserRegistry
+	 */
+	virtual void OnUnregistered() {}
+	
+private:
+	/** The unique index for this accessible user */
+	FAccessibleUserIndex UserIndex;
+	/** The focused accessible widget widget associated with this user. */
+	TWeakPtr<IAccessibleWidget> FocusedAccessibleWidget;
+};
+
+/**
+ * The base class for an accessible user registry that all accessible users must register with.
+ * The registry controls the lifetime of all accessible users.
+ * All created accessible users must be registered with an accessible user registry. This should be the only means
+ * of retrieving and interacting with accessible users in other parts of the application.
+ * Users can subclass this class to provide additional functionality for their custom use cases.
+ * @see FGenericAccessibleUser, FGenericAccessibleMessageHandler
+ */
+class APPLICATIONCORE_API FGenericAccessibleUserRegistry
+{
+public:
+	virtual ~FGenericAccessibleUserRegistry() = default;
+	/**
+	 * Registers an accessible user with the accessible registry.
+	 * It is up to the caller to ensure that the passed in user has the correct user index.
+	 * The passed in user will not be registered if another user with the same user index has already been registered.
+	 * This function calls FGenericAccessibleUser::OnRegistered if the passed in user is successfully registered.
+	 * Once an accessible user has been successfully registered, it can be retrieved and interacted with by the rest of the application.
+	 * @param User An accessible user you want to register with the accessible user registry.
+	 * @return True if the passed in user was successfully registered. Else, returns false.
+	 */
+	bool RegisterUser(const TSharedRef<FGenericAccessibleUser>& User);
+	/**
+	 * Unregisters an accessible user from the accessible registry.
+	 * If the passed in user index is not associated with a registered accessible user, nothing will happen.
+	 * This function calls FGenericAccessibleUser::OnUnregistered if the passed in user is successfully unregistered.
+	 * Once an accessible user has been unregistered, the rest of the application should have no way of retrieving or interacting with the accessible user.
+	 * @param UserIndex The index of the accessible user you want to unregister.
+	 * @return True if the accessible user with the associated user index is successfully unregistered. Else, returns false.
+	 */
+	bool UnregisterUser(const FAccessibleUserIndex UserIndex);
+	/**
+	 * Unregisters all accessible users from the accessible registry.
+	 * This function calls FGenericAccessibleUser::OnUnregistered when each registered user is successfully unregistered.
+	 */
+	void UnregisterAllUsers();
+	/** Returns true if the passed in user index is associated with a registered accessible user. Else, returns false. */
+	bool IsUserRegistered(const FAccessibleUserIndex UserIndex) const;
+	/** Returns the accessible user associated with the passed in user index. If the passed in user index is not associated with any accessible user, nullptr is returned. */
+	TSharedPtr<FGenericAccessibleUser> GetUser(const FAccessibleUserIndex UserIndex) const;
+	/** Returns the number of accessible users registered with the accessible user registry. */
+	int32 GetNumberofUsers() const;
+	/** Returns an array of all accessible users that are currently registered with the accessible registry. */
+	TArray<TSharedRef<FGenericAccessibleUser>> GetAllUsers() const;
+	/**
+	 * Returns the index of the primary accessible user.
+	 * The primary accessible user should correspond to the default input source that every application should have. This user should also be the primary cursor user.
+	 * E.g On desktop platforms, that would correspond to the keyboard user.
+	 */
+	static FAccessibleUserIndex GetPrimaryUserIndex()
+	{
+		// @TODOAccessibility: Consider allowing remapping of primary user index
+		static const FAccessibleUserIndex PrimaryUserIndex = 0;
+		return PrimaryUserIndex;
+	}
+protected:
+	/** A map of accessible user indices to their associated accessible users. */
+	TMap<FAccessibleUserIndex, TSharedRef<FGenericAccessibleUser>> UsersMap;
+};
+
 /**
  * An accessible window corresponds to a native OS window. Fake windows that are embedded
  * within other widgets that simply look and feel like windows are not IAccessibleWindows.
@@ -120,11 +252,13 @@ public:
 	 */
 	virtual TSharedPtr<IAccessibleWidget> GetChildAtPosition(int32 X, int32 Y) = 0;
 	/**
-	 * Retrieves the currently accessibilit focused widget
+	 * Retrieves the focused accessible widget for a user.
+	 * If the passed in user index is not registered with an accessible user registry, nullptr will be returned.
 	 *
-	 * @return The widget that has accessibility focus.
+	 * @param UserIndex The user index of the accessible user to query for its accessible focus widget.
+	 * @return The accessible widget that is focused by the accessible user. Returns nullptr if the accessible user index does not correspond to a registered accessible user.
 	 */
-	virtual TSharedPtr<IAccessibleWidget> GetFocusedWidget() const = 0;
+	virtual TSharedPtr<IAccessibleWidget> GetUserFocusedWidget(const FAccessibleUserIndex UserIndex) const = 0;
 	/**
 	 * Request that the window closes itself. This may not happen immediately.
 	 */
@@ -432,21 +566,37 @@ public:
 	 */
 	virtual bool IsHidden() const = 0;
 	/**
-	 * Whether the widget supports accessibility focus or not.
+	 * Whether the widget can ever support keyboard/gamepad focus.
 	 *
-	 * @return true if the widget can receive accessibility focus.
+	 * @return true if the widget can ever receive keyboard/gamepad focus. Else, returns false.
 	 */
 	virtual bool SupportsFocus() const = 0;
 	/**
-	 * Whether the widget has accessibility focus or not.
+	 * Whether the widget can ever support accessible focus.
 	 *
-	 * @return true if the widget currently has accessibility focus.
+	 * @return true if the widget can ever receive accessible focus. Else, returns false.
 	 */
-	virtual bool HasFocus() const = 0;
-	/** Assign keyboard/gamepad focus to this widget, if it supports it. If not, focus should not be affected. 
-	* Note: For widgets that support accessibility focus but not keyboard/gamepad focus, this should NOT repreplace the widget with accessibility focus. 
+	virtual bool SupportsAccessibleFocus() const = 0;
+	/**
+	 * Whether the widget can currently support accessible focus.
+	 *
+	 * @return true if the widget can currently receive accessible focus. Else, returns false.
+	 */
+	virtual bool CanCurrentlyAcceptAccessibleFocus() const = 0;
+	/**
+	 * Whether the widget has accessible focus or not by a particular user .
+	 *
+	 * @param UserIndex The user index associated with the accessible user to query if this accessible widget is currently being focused.
+	 * @return true if the widget currently has accessible focus by the user. Else, returns false.
+	 */
+	virtual bool HasUserFocus(const FAccessibleUserIndex UserIndex) const = 0;
+	/**
+	 * Assign accessible focus to the widget. Also sets keyboard/gamepad focus if the widget supports it.
+	*
+	 * @param UserIndex The user index associated with the accessible user that requested focus for this accessible widget.
+	 * @return True if focus was successfully set to this widget for the requested user. Else false.
 	*/
-	virtual void SetFocus() = 0;
+	virtual bool SetUserFocus(const FAccessibleUserIndex UserIndex) = 0;
 
 	/**
 	 * Attempt to cast this to an IAccessibleWindow
@@ -488,56 +638,59 @@ public:
 	virtual IAccessibleTableRow* AsTableRow() { return nullptr; }
 };
 
+/**
+ * The arguments for an accessible event that is raised by accessible widgets to be
+ * passed to an accessibility event handler such as a native OS.
+ * It is up to the client to create an instance of this struct and fill in the appropriate data members.
+ * @see FGenericAccessibleMessageHandler::RaiseEvent
+ */
+struct APPLICATIONCORE_API FAccessibleEventArgs
+{
+	FAccessibleEventArgs(TSharedRef<IAccessibleWidget> InWidget, EAccessibleEvent InEvent, FVariant InOldValue = FVariant(), FVariant InNewValue = FVariant(), FAccessibleUserIndex InUserIndex = 0)
+		: Widget(InWidget)
+		, Event(InEvent)
+		, OldValue(InOldValue)
+		, NewValue(InNewValue)
+		, UserIndex(InUserIndex)
+	{}
+	
+	/** The accessible widget that generated the event */
+	TSharedRef<IAccessibleWidget> Widget;
+	/** The type of event generated */
+	EAccessibleEvent Event;
+	/** If this was a property changed event, the 'before' value */
+	FVariant OldValue;
+	/** If this was a property changed event, the 'after' value. This may also be set for other events such as Notification. */
+	FVariant NewValue;
+	/** The Id of the user this event is intended for. Think of a hardware device such as a controller or keyboard/mouse. */
+	FAccessibleUserIndex UserIndex;
+};
+
 /**
  * Platform and application-agnostic messaging system for accessible events. The message handler
  * lives in GenericApplication and any subclass that wishes to support accessibility should subclass
  * this and use GenericAppliation::SetAccessibleMessageHandler to enable functionality.
  *
- * GetAccessibleWindow() is tne entry point to all accessible widgets. Once the window is retrieved, it
+ * GetAccessibleWindow() is the entry point to all accessible widgets. Once the window is retrieved, it
  * can be queried for children in various ways. RaiseEvent() allows messages to bubble back up to the
  * native OS through anything bound to the AccessibleEventDelegate.
  *
  * Callers can use ApplicationIsAccessible() to see if accessibility is supported or not. Alternatively,
  * calling GetAccessibleWindow and seeing if the result is valid should provide the same information.
+ *
+ * Callers should also use GetAccessibleUserRegistry() to register and interact with accessible users. Accessible users must be registered
+ * with the set accessible user registry for the rest of the application to be able to retrieve and interact with the accessible users.
  */
-class FGenericAccessibleMessageHandler
+class APPLICATIONCORE_API FGenericAccessibleMessageHandler
 {
 public:
-	/**
-	 * The arguments for an accessible event that is raised by accessible widgets to be
-	 * passed to an accessibility event handler such as a native OS.
-	 * It is up to the client to create an instance of this struct and fill in the appropriate data members.
-	 * @see RaiseEvent
-	 */
-	struct FAccessibleEventArgs
-	{
-		FAccessibleEventArgs(TSharedRef<IAccessibleWidget> InWidget, EAccessibleEvent InEvent, FVariant InOldValue = FVariant(), FVariant InNewValue = FVariant(), int32 InUserId = 0)
-			: Widget(InWidget)
-			, Event(InEvent)
-			, OldValue(InOldValue)
-			, NewValue(InNewValue)
-			, UserId(InUserId)
-		{}
-		
-		/** The accessible widget that generated the event */
-		TSharedRef<IAccessibleWidget> Widget;
-		/** The type of event generated */
-		EAccessibleEvent Event;
-		/** If this was a property changed event, the 'before' value */
-		FVariant OldValue;
-		/** If this was a property changed event, the 'after' value. This may also be set for other events such as Notification. */
-		FVariant NewValue;
-		/** The Id of the user this event is intended for. Think of a hardware device such as a controller or keyboard/mouse. */
-		int32 UserId;
-	};
 	/**
 	 * A delegate accessible event handlers such as platform accessibility APIs can
 	 * listen for accessibility events raised by widgets.  .
 	 */
 	DECLARE_DELEGATE_OneParam(FAccessibleEvent, const FAccessibleEventArgs&);
 
-	FGenericAccessibleMessageHandler() : bApplicationIsAccessible(false), bIsActive(false) {}
-
+	FGenericAccessibleMessageHandler();
 	virtual ~FGenericAccessibleMessageHandler()
 	{
 		UnbindAccessibleEventDelegate();
@@ -636,12 +789,37 @@ public:
 	 * earlier announcement requests will get stomped by later announcements on certain platforms.
 	 */
 	virtual void MakeAccessibleAnnouncement(const FString& AnnouncementString) { }
+
+	/**
+	 * Retrieves the accessible user registry. This is the means of retrieving, registering and interacting with accessible useres in the rest of the application.
+	 */
+	FGenericAccessibleUserRegistry& GetAccessibleUserRegistry()
+	{
+		return *AccessibleUserRegistry;
+	}
+	/**
+	 * Retrieves the accessible user registry. This is the means of retrieving, registering and interacting with accessible useres in the rest of the application.
+	 */
+	const FGenericAccessibleUserRegistry& GetAccessibleUserRegistry() const
+	{
+		return *AccessibleUserRegistry;
+	}
+	/**
+	 * Sets a new accessible user registry for the application.
+	 * Note: As of now, none of the registered users will carry over to the new accessible user manager. It is up to the caller to unregister all currently registered users and
+	 * register them with the new accessible user registry.   .
+	 */
+	void SetAccessibleUserRegistry(const TSharedRef<FGenericAccessibleUserRegistry>& InAccessibleUserRegistry)
+	{
+		// @TODOAccessibility: Have some means of storing a default manager
+		AccessibleUserRegistry = InAccessibleUserRegistry;
+	}
 protected:
 	/** Triggered when bIsActive changes from false to true. */
 	virtual void OnActivate() {}
 	/** Triggered when bIsActive changes from true to false. */
 	virtual void OnDeactivate() {}
-
+	
 	/** Subclasses should override this to indicate that they support accessibility. */
 	bool bApplicationIsAccessible;
 
@@ -650,6 +828,11 @@ private:
 	bool bIsActive;
 	/** Delegate for the platform layer to listen to widget events */
 	FAccessibleEvent AccessibleEventDelegate;
+	/**
+	 * Registry for all accessible users in the application.
+	 * Accessible users must be registered with the accessible user registry to allow other parts of the application to retrieve and interact with the accessible users.
+	 */
+	TSharedRef<FGenericAccessibleUserRegistry> AccessibleUserRegistry;
 };
 
 #endif
diff --git a/Engine/Source/Runtime/ApplicationCore/Public/IOS/IOSApplication.h b/Engine/Source/Runtime/ApplicationCore/Public/IOS/IOSApplication.h
--- a/Engine/Source/Runtime/ApplicationCore/Public/IOS/IOSApplication.h
+++ b/Engine/Source/Runtime/ApplicationCore/Public/IOS/IOSApplication.h
@@ -49,7 +49,7 @@ private:
 
 	FIOSApplication();
 #if WITH_ACCESSIBILITY
-	void OnAccessibleEventRaised(const FGenericAccessibleMessageHandler::FAccessibleEventArgs& Args);
+	void OnAccessibleEventRaised(const FAccessibleEventArgs& Args);
 #endif
 
 private:
diff --git a/Engine/Source/Runtime/ApplicationCore/Public/Mac/MacApplication.h b/Engine/Source/Runtime/ApplicationCore/Public/Mac/MacApplication.h
--- a/Engine/Source/Runtime/ApplicationCore/Public/Mac/MacApplication.h
+++ b/Engine/Source/Runtime/ApplicationCore/Public/Mac/MacApplication.h
@@ -322,7 +322,7 @@ private:
 #endif
 
 #if WITH_ACCESSIBILITY
-	void OnAccessibleEventRaised(const FGenericAccessibleMessageHandler::FAccessibleEventArgs& Args);
+	void OnAccessibleEventRaised(const FAccessibleEventArgs& Args);
 	#endif
 
 private:
diff --git a/Engine/Source/Runtime/ApplicationCore/Public/Windows/Accessibility/WindowsUIAManager.h b/Engine/Source/Runtime/ApplicationCore/Public/Windows/Accessibility/WindowsUIAManager.h
--- a/Engine/Source/Runtime/ApplicationCore/Public/Windows/Accessibility/WindowsUIAManager.h
+++ b/Engine/Source/Runtime/ApplicationCore/Public/Windows/Accessibility/WindowsUIAManager.h
@@ -86,7 +86,7 @@ public:
 
 private:
 	/** Callback function for processing events raised from the AccessibleMessageHandler */
-	void OnEventRaised(const FGenericAccessibleMessageHandler::FAccessibleEventArgs& Args);
+	void OnEventRaised(const FAccessibleEventArgs& Args);
 
 	/** Called when the first window widget provider is requested. */
 	void OnAccessibilityEnabled(); 
diff --git a/Engine/Source/Runtime/SlateCore/Private/Widgets/Accessibility/SlateAccessibleMessageHandler.cpp b/Engine/Source/Runtime/SlateCore/Private/Widgets/Accessibility/SlateAccessibleMessageHandler.cpp
--- a/Engine/Source/Runtime/SlateCore/Private/Widgets/Accessibility/SlateAccessibleMessageHandler.cpp
+++ b/Engine/Source/Runtime/SlateCore/Private/Widgets/Accessibility/SlateAccessibleMessageHandler.cpp
@@ -167,21 +167,6 @@ void FSlateAccessibleMessageHandler::ProcessAccessibleTasks()
 #endif 
 }
 
-TSharedPtr<FSlateAccessibleWidget> FSlateAccessibleMessageHandler::GetAccessibilityFocusedWidget() const
-{
-	return AccessibilityFocusedWidget.Pin();
-}
-
-void FSlateAccessibleMessageHandler::SetAccessibilityFocusedWidget(const TSharedRef<FSlateAccessibleWidget>& NewAccessibilityFocusedWidget)
-{
-	AccessibilityFocusedWidget = NewAccessibilityFocusedWidget;
-}
-
-void FSlateAccessibleMessageHandler::ClearAccessibilityFocus()
-{
-	AccessibilityFocusedWidget.Reset();
-}
-
 void FSlateAccessibleMessageHandler::OnWidgetRemoved(SWidget* Widget)
 {
 	if (IsActive())
@@ -189,16 +174,50 @@ void FSlateAccessibleMessageHandler::OnWidgetRemoved(SWidget* Widget)
 		TSharedPtr<FSlateAccessibleWidget> RemovedWidget = FSlateAccessibleWidgetCache::RemoveWidget(Widget);
 		if (RemovedWidget.IsValid())
 		{
-			RaiseEvent(FGenericAccessibleMessageHandler::FAccessibleEventArgs(RemovedWidget.ToSharedRef(), EAccessibleEvent::WidgetRemoved));
+			RaiseEvent(FAccessibleEventArgs(RemovedWidget.ToSharedRef(), EAccessibleEvent::WidgetRemoved));
 			// If this ensure fails, bDirty = true must be called to ensure the tree is kept up to date.
 			ensureMsgf(!Widget->GetParentWidget().IsValid(), TEXT("A widget was unexpectedly deleted before detaching from its parent."));
 		}
 	}
 }
 
+void FSlateAccessibleMessageHandler::HandleAccessibleWidgetFocusChangeEvent(const TSharedRef<IAccessibleWidget>& Widget, bool bIsWidgetGainingFocus, FAccessibleUserIndex UserIndex)
+{
+	// we should have already checked that the user index is registered before this function is called
+	check(GetAccessibleUserRegistry().IsUserRegistered(UserIndex));
+	if (bIsWidgetGainingFocus)
+	{
+		TSharedPtr<FGenericAccessibleUser> AccessibleUser = GetAccessibleUserRegistry().GetUser(UserIndex);
+		// For widgets that are non-keyboard focusable but support accessibility,
+		// we have to manually raise a focus change event
+		// to signal the platform that the currently accessible focused widget is losing focus
+		// As regular navigation focus in FSlateApplication doesn't apply to these widgets
+		if (!Widget->SupportsFocus() && Widget->SupportsAccessibleFocus())
+		{
+			TSharedPtr<IAccessibleWidget> CurrentAccessibilityFocusedWidget = AccessibleUser->GetFocusedAccessibleWidget();
+			if (CurrentAccessibilityFocusedWidget && (Widget != CurrentAccessibilityFocusedWidget))
+			{
+				// let the currently focused widget lose focus 
+				HandleAccessibleWidgetFocusChangeEvent(CurrentAccessibilityFocusedWidget.ToSharedRef(), false, UserIndex);
+			}
+		}
+		// Update the widget that currently has accessible focus for the user
+		AccessibleUser->SetFocusedAccessibleWidget(Widget);
+	}
+	// if the widget is losing focus and is not focusable by keyboard/gamepad
+	else if (!Widget->SupportsFocus())
+	{
+		// We manually raise a focus change event here   to allow modules listening to FAccessibleEventDelegate
+		// to do additional processing to unfocus from the widget
+		// E.g Platforms need this information to unfocus in the platform accessibility API
+		RaiseEvent(FAccessibleEventArgs(Widget, EAccessibleEvent::FocusChange, true, false, UserIndex));
+	}
+}
+
 void FSlateAccessibleMessageHandler::OnWidgetEventRaised(const FSlateWidgetAccessibleEventArgs& Args)
 {
-	if (IsActive())
+	// we only process events for accessible users that have been registered
+	if (IsActive() && GetAccessibleUserRegistry().IsUserRegistered(Args.SlateUserIndex))
 	{
 		SCOPE_CYCLE_COUNTER(STAT_AccessibilitySlateEventRaised);
 		// todo: not sure what to do for a case like focus changed to not-accessible widget. maybe pass through a nullptr?
@@ -206,35 +225,22 @@ void FSlateAccessibleMessageHandler::OnWidgetEventRaised(const FSlateWidgetAcces
 		{
 			TSharedRef<FSlateAccessibleWidget> AccessibleWidget = FSlateAccessibleWidgetCache::GetAccessibleWidget(Args.Widget);
 			// Perform FSlateAccessibleMessageHandler preprocessing here  before platform processing
-			// Most of the time, accessibility focus is the same as Slate focus. We sync them here
+			// Most of the time, accessible focus is the same as Slate focus. We sync them here
 			// as a focus event is raised from SlateApplication for keyboard/gamepad focusable widgets.
 			// However, we can also raise a focus change event manually on non-Slate focusable widgets e.g STableRow
-			// That's where accessibility focus and Slate focus can diverge
+			// That's where accessible focus and Slate focus can diverge
 			if (Args.Event == EAccessibleEvent::FocusChange)
 			{
+				// The old focus value is the same as the new focus value.
+				// This is most likely a user error when raising the event.
+				ensure(Args.OldValue.GetValue<bool>() != Args.NewValue.GetValue<bool>());
 				// If the Widget is gaining focus
 				if (Args.NewValue.GetValue<bool>())
 				{
-					// For widgets that are non-keyboard focusable but support accessibility,
-					// we ened to manually raise a focus change event
-					// to signal the platform that the currently accessibility focused widget is losing focus
-					// As regular navigation focus in FSlateApplication doesn't apply to these widgets
-					if (!Args.Widget->SupportsKeyboardFocus())
-					{
-						TSharedPtr<IAccessibleWidget> CurrentAccessibilityFocusedWidget = GetAccessibilityFocusedWidget();
-						if (CurrentAccessibilityFocusedWidget.IsValid())
-						{
-							// We manually raise a focus change event here   to allow platforms to unfocus
-							// from the currently focused widget
-							//FSlateAccessibleMessageHandler::RaiseEvent(CurrentAccessibilityFocusedWidget.ToSharedRef(), EAccessibleEvent::FocusChange, true, false);
-							FSlateAccessibleMessageHandler::RaiseEvent(FGenericAccessibleMessageHandler::FAccessibleEventArgs(CurrentAccessibilityFocusedWidget.ToSharedRef(), EAccessibleEvent::FocusChange, true, false, Args.SlateUserId));
-						}
-					}
-					// Update the widget that currently has accessibility focus
-					SetAccessibilityFocusedWidget(AccessibleWidget);
+					HandleAccessibleWidgetFocusChangeEvent(AccessibleWidget, true, Args.SlateUserIndex);
 				}
 			}
-			FSlateAccessibleMessageHandler::RaiseEvent(FGenericAccessibleMessageHandler::FAccessibleEventArgs(AccessibleWidget, Args.Event, Args.OldValue, Args.NewValue, Args.SlateUserId));
+			RaiseEvent(FAccessibleEventArgs(AccessibleWidget, Args.Event, Args.OldValue, Args.NewValue, Args.SlateUserIndex));
 		}
 	}
 }
@@ -321,7 +327,7 @@ void FSlateAccessibleMessageHandler::MakeAccessibleAnnouncement(const FString& A
 		TSharedPtr<SWidget> ActiveTopLevelWindow = FSlateApplicationBase::Get().GetActiveTopLevelWindow();
 		if(ActiveTopLevelWindow.IsValid())
 		{
-			RaiseEvent(FGenericAccessibleMessageHandler::FAccessibleEventArgs(FSlateAccessibleWidgetCache::GetAccessibleWidget(ActiveTopLevelWindow.ToSharedRef()), EAccessibleEvent::Notification, FString(), AnnouncementString));
+			RaiseEvent(FAccessibleEventArgs(FSlateAccessibleWidgetCache::GetAccessibleWidget(ActiveTopLevelWindow.ToSharedRef()), EAccessibleEvent::Notification, FString(), AnnouncementString));
 		}
 	}
 }
diff --git a/Engine/Source/Runtime/SlateCore/Private/Widgets/Accessibility/SlateCoreAccessibleWidgets.cpp b/Engine/Source/Runtime/SlateCore/Private/Widgets/Accessibility/SlateCoreAccessibleWidgets.cpp
--- a/Engine/Source/Runtime/SlateCore/Private/Widgets/Accessibility/SlateCoreAccessibleWidgets.cpp
+++ b/Engine/Source/Runtime/SlateCore/Private/Widgets/Accessibility/SlateCoreAccessibleWidgets.cpp
@@ -153,20 +153,46 @@ bool FSlateAccessibleWidget::IsHidden() const
 
 bool FSlateAccessibleWidget::SupportsFocus() const
 {
-	// all widgets that support accessibility support accessibility focus right now 
+	if (Widget.IsValid())
+	{
+		return Widget.Pin()->SupportsKeyboardFocus();
+	}
+	return false;
+}
+
+bool FSlateAccessibleWidget::SupportsAccessibleFocus() const
+{
+	// all widgets that support accessibility support accessible focus right now
 	// This check is analogous to Widget.Pin()->IsAccessible()
-	// By definition all FSlateAccessibleWidgets are accessible. So we just return true 
+	// By definition all FSlateAccessibleWidgets are accessible. So we just return true
 	return true;
 }
 
-bool FSlateAccessibleWidget::HasFocus() const
+bool FSlateAccessibleWidget::CanCurrentlyAcceptAccessibleFocus() const
+{
+	return IsEnabled() && !IsHidden();
+}
+
+bool FSlateAccessibleWidget::HasUserFocus(const FAccessibleUserIndex UserIndex) const
 {
-	TSharedPtr<FSlateAccessibleWidget> AccessibilityFocusedWidget = FSlateApplicationBase::Get().GetAccessibleMessageHandler()->GetAccessibilityFocusedWidget();
-	return AccessibilityFocusedWidget == AsShared();
+	FGenericAccessibleUserRegistry& UserManager = FSlateApplicationBase::Get().GetAccessibleMessageHandler()->GetAccessibleUserRegistry();
+	TSharedPtr<FGenericAccessibleUser> User = UserManager.GetUser(UserIndex);
+	if (User)
+	{
+		return (User->GetFocusedAccessibleWidget()) == AsShared();
+	}
+	return false;
 }
 
-void FSlateAccessibleWidget::SetFocus()
+bool FSlateAccessibleWidget::SetUserFocus(const FAccessibleUserIndex UserIndex)
 {
+	// Most likely a mistake to set focus on a widget that cannot currently accept focus
+	if (!CanCurrentlyAcceptAccessibleFocus())
+	{
+		// @TODOAccessibility: Log information about the widget when we create a ToString() function.
+		UE_LOG(LogAccessibility, Warning, TEXT("Attempting to set accessible focus on an accessible widget that currently cannot be focused. Focus not changed."));
+		return false;
+	}
 	if (SupportsFocus())
 	{
 		TSharedPtr<SWindow> WidgetWindow = GetSlateWindow();
@@ -177,14 +203,19 @@ void FSlateAccessibleWidget::SetFocus()
 			FWidgetPath WidgetPath;
 			if (FSlateWindowHelper::FindPathToWidget(WindowArray, Widget.Pin().ToSharedRef(), WidgetPath))
 			{
-				// From FSlateApplication::SetUserFocus(), it seems that 
-				// no focus change will occur if we pass in a widget that cannot accept keyboard focus.
-				// this behavior is fine for widgets that don't support keyboard focus but do support accessibility focus 
-//@TODOAccessibility: Possible early out if not keyboard focusable 				
-				FSlateApplicationBase::Get().SetKeyboardFocus(WidgetPath, EFocusCause::SetDirectly);
+				// Focus accessible events are already raised from this function call
+				return FSlateApplicationBase::Get().SetUserFocus(UserIndex, WidgetPath, EFocusCause::SetDirectly);
 			}
 		}
 	}
+	// The widget is not keyboard/gamepad focusable but supports accessible focus
+	else if (SupportsAccessibleFocus())
+	{
+		// we manually raise an accessible focus event 
+		FSlateApplicationBase::Get().GetAccessibleMessageHandler()->OnWidgetEventRaised(FSlateAccessibleMessageHandler::FSlateWidgetAccessibleEventArgs(Widget.Pin().ToSharedRef(), EAccessibleEvent::FocusChange, false, true, UserIndex));
+		return true;
+	}
+	return false;
 }
 
 void FSlateAccessibleWidget::UpdateParent(TSharedPtr<IAccessibleWidget> NewParent)
@@ -192,7 +223,7 @@ void FSlateAccessibleWidget::UpdateParent(TSharedPtr<IAccessibleWidget> NewParen
 	if (Parent != NewParent)
 	{
 		FSlateApplicationBase::Get().GetAccessibleMessageHandler()->RaiseEvent(
-			FGenericAccessibleMessageHandler::FAccessibleEventArgs(AsShared(), EAccessibleEvent::ParentChanged,
+			FAccessibleEventArgs(AsShared(), EAccessibleEvent::ParentChanged,
 			Parent.IsValid() ? Parent.Pin()->GetId() : IAccessibleWidget::InvalidAccessibleWidgetId,
 			NewParent.IsValid() ? NewParent->GetId() : IAccessibleWidget::InvalidAccessibleWidgetId));
 		Parent = StaticCastSharedPtr<FSlateAccessibleWidget>(NewParent);
@@ -320,9 +351,10 @@ TSharedPtr<IAccessibleWidget> FSlateAccessibleWindow::GetChildAtPosition(int32 X
 	return HitWidget;
 }
 
-TSharedPtr<IAccessibleWidget> FSlateAccessibleWindow::GetFocusedWidget() const
+TSharedPtr<IAccessibleWidget> FSlateAccessibleWindow::GetUserFocusedWidget(const FAccessibleUserIndex UserIndex) const
 {
-	return FSlateApplicationBase::Get().GetAccessibleMessageHandler()->GetAccessibilityFocusedWidget();
+	TSharedPtr<FGenericAccessibleUser> User = FSlateApplicationBase::Get().GetAccessibleMessageHandler()->GetAccessibleUserRegistry().GetUser(UserIndex);
+	return User ? User->GetFocusedAccessibleWidget() : nullptr;
 }
 
 FString FSlateAccessibleWindow::GetWidgetName() const
diff --git a/Engine/Source/Runtime/SlateCore/Public/Widgets/Accessibility/SlateAccessibleMessageHandler.h b/Engine/Source/Runtime/SlateCore/Public/Widgets/Accessibility/SlateAccessibleMessageHandler.h
--- a/Engine/Source/Runtime/SlateCore/Public/Widgets/Accessibility/SlateAccessibleMessageHandler.h
+++ b/Engine/Source/Runtime/SlateCore/Public/Widgets/Accessibility/SlateAccessibleMessageHandler.h
@@ -48,12 +48,12 @@ public:
 	 */
 	struct FSlateWidgetAccessibleEventArgs
 	{
-		FSlateWidgetAccessibleEventArgs(TSharedRef<SWidget> InWidget, EAccessibleEvent InEvent, FVariant InOldValue = FVariant(), FVariant InNewValue = FVariant(), int32 InSlateUserId = 0)
+		FSlateWidgetAccessibleEventArgs(TSharedRef<SWidget> InWidget, EAccessibleEvent InEvent, FVariant InOldValue = FVariant(), FVariant InNewValue = FVariant(), FAccessibleUserIndex InSlateUserIndex = 0)
 			: Widget(InWidget)
 			, Event(InEvent)
 			, OldValue(InOldValue)
 			, NewValue(InNewValue)
-			, SlateUserId(InSlateUserId)
+			, SlateUserIndex(InSlateUserIndex)
 		{}
 			
 		/** The widget that's raising the accessible event */
@@ -63,9 +63,9 @@ public:
 		/** The value of the property being changed before the change took place. */
 		FVariant OldValue;
 		/** The value of the property being changed after the change took place. Alternatively, can contain any miscellaneous data for the accessible event. */
-				FVariant NewValue;
+		FVariant NewValue;
 		/** The index of the Slate user that feedback for the accessible event should be directed towards. */
-		int32 SlateUserId;
+		FAccessibleUserIndex SlateUserIndex;
 	};
 	/**
 	 * Callback for a Slate widget indicating that a property change occurred. This may also be used by certain events
@@ -92,25 +92,9 @@ public:
 	 */
 	void Tick();
 
-	/**
-	* Returns the Widget that currently has accessibility focus in the application.
-	* accessibility focusable widgets are a superset of keyboard/gamepad focusable widgetes.
-	* i.e A widget can support accessibility focus but NOT keyboard/gamepad focus.
-	* Can return nullptr if no widget has accessibility focus. Up to caller to do validity check.
-	* @return  The widget that currently has accessible focus.
-	*/
-	TSharedPtr<FSlateAccessibleWidget> GetAccessibilityFocusedWidget() const;
-
-	/**
-	* Sets the currently accessibility focused widget in the application.
-	* If you're trying to clear the accessibility focus, use ClearAccessibilityFocus() 
-	* @param NewAccessibleFocusedWidget The widget that  now has accessible focus
-	*/
-	void SetAccessibilityFocusedWidget(const TSharedRef<FSlateAccessibleWidget>& NewAccessibilityFocusedWidget);
 
-	/** Reset the accessibility focused widget to nullptr*/
-	void ClearAccessibilityFocus();
 private:
+	void HandleAccessibleWidgetFocusChangeEvent(const TSharedRef<IAccessibleWidget>& FocusWidget, bool bIsWidgetGainingFocus, FAccessibleUserIndex UserIndex);
 	/**
 	*  A helper class that wraps an accessibility task and the event to be triggered when the task finishes executing. 
 	*/
@@ -179,14 +163,7 @@ private:
 
 #if ACCESSIBILITY_DEBUG_RESPONSIVENESS
 	FEvent* EnqueueEvent;
-#endif 
-
-	/** 
-	* The widget tht currently has accessibility focus 
-	* @see GetAccessibilityFocusedWidget, SetAccessibilityFocusedWidget
-	*/
-	TWeakPtr<FSlateAccessibleWidget> AccessibilityFocusedWidget;
-
+#endif
 	/** If true, Tick() will begin the update process to the accessible widget tree. Use MarkDirty() to set. */
 	bool bDirty;
 };
diff --git a/Engine/Source/Runtime/SlateCore/Public/Widgets/Accessibility/SlateCoreAccessibleWidgets.h b/Engine/Source/Runtime/SlateCore/Public/Widgets/Accessibility/SlateCoreAccessibleWidgets.h
--- a/Engine/Source/Runtime/SlateCore/Public/Widgets/Accessibility/SlateCoreAccessibleWidgets.h
+++ b/Engine/Source/Runtime/SlateCore/Public/Widgets/Accessibility/SlateCoreAccessibleWidgets.h
@@ -36,8 +36,10 @@ public:
 	virtual bool IsEnabled() const override final;
 	virtual bool IsHidden() const override final;
 	virtual bool SupportsFocus() const override final;
-	virtual bool HasFocus() const override final;
-	virtual void SetFocus() override final;
+	virtual bool SupportsAccessibleFocus() const override final;
+	virtual bool CanCurrentlyAcceptAccessibleFocus() const override final;
+	virtual bool HasUserFocus(const FAccessibleUserIndex UserIndex) const override final;
+	virtual bool SetUserFocus(const FAccessibleUserIndex UserIndex) override final;
 
 	virtual EAccessibleWidgetType GetWidgetType() const override { return WidgetType; }
 	virtual FString GetWidgetName() const override;
@@ -93,7 +95,7 @@ public:
 	// IAccessibleWindow
 	virtual TSharedPtr<FGenericWindow> GetNativeWindow() const override;
 	virtual TSharedPtr<IAccessibleWidget> GetChildAtPosition(int32 X, int32 Y) override;
-	virtual TSharedPtr<IAccessibleWidget> GetFocusedWidget() const override;
+	virtual TSharedPtr<IAccessibleWidget> GetUserFocusedWidget(const FAccessibleUserIndex UserIndex) const override;
 	virtual void Close() override;
 	virtual bool SupportsDisplayState(EWindowDisplayState State) const override;
 	virtual EWindowDisplayState GetDisplayState() const override;
