Fix a potential resource leak as reported by user.

diff --git a/Engine/Source/Runtime/ApplicationCore/Private/Windows/WindowsPlatformApplicationMisc.cpp b/Engine/Source/Runtime/ApplicationCore/Private/Windows/WindowsPlatformApplicationMisc.cpp
--- a/Engine/Source/Runtime/ApplicationCore/Private/Windows/WindowsPlatformApplicationMisc.cpp
+++ b/Engine/Source/Runtime/ApplicationCore/Private/Windows/WindowsPlatformApplicationMisc.cpp
@@ -171,7 +171,10 @@ void FWindowsPlatformApplicationMisc::PreventScreenSaver()
 
 FLinearColor FWindowsPlatformApplicationMisc::GetScreenPixelColor(const FVector2D& InScreenPos, float /*InGamma*/)
 {
-	COLORREF PixelColorRef = GetPixel(GetDC(HWND_DESKTOP), InScreenPos.X, InScreenPos.Y);
+	HDC TempDC = GetDC(HWND_DESKTOP);
+	COLORREF PixelColorRef = GetPixel(TempDC, InScreenPos.X, InScreenPos.Y);
+
+	ReleaseDC(HWND_DESKTOP, TempDC);
 
 	FColor sRGBScreenColor(
 		(PixelColorRef & 0xFF),
