Introduced screen reader navigation policies. This allows the screen reader to navigate around UI elements on the screen separate from regular Slate navigation.
- Introduced IScreenReaderNavigationPolicy. This is an interface class that helps the screen reader navigate to the appropriate target accessible widget.
- Introduced FScreenReaderDefaultNavigationPolicy. This is a policy that all FScreenReaderUsers have by default. It helps the screen reader user find target accessible widgets that are enabled and can be focused on.
- Introduced focus navigation functions to FScreenReaderUser. This is the means clients can control and shift accessible focus for FScreenReaderUsers.
- IAccessibleWidget can now retrieve the logical next and previous widget in the accessible widget hierarchy. This functionality is similar to that of mobile screen readers like Voiceover and Android Talkback which allow users to navigate around the UI hierarchy with left and right swipes.
- IAccessibleWidget can now search for a target widget from a source widget based on a search criteria. This is useful for searching for widgets by characteristics (such as enabled state, widget type etc).
- Fixed a small bug in accessible focus where an accessible widget with application focus cannot receive accessible focus again and have its contents announced to the end user.

diff --git a/Engine/Plugins/Experimental/ScreenReader/Source/ScreenReader/Private/GenericPlatform/ScreenReaderUser.cpp b/Engine/Plugins/Experimental/ScreenReader/Source/ScreenReader/Private/GenericPlatform/ScreenReaderUser.cpp
--- a/Engine/Plugins/Experimental/ScreenReader/Source/ScreenReader/Private/GenericPlatform/ScreenReaderUser.cpp
+++ b/Engine/Plugins/Experimental/ScreenReader/Source/ScreenReader/Private/GenericPlatform/ScreenReaderUser.cpp
@@ -2,10 +2,13 @@
 
 #include "GenericPlatform/ScreenReaderUser.h"
 #include "Announcement/ScreenReaderAnnouncementChannel.h"
+#include "Navigation/ScreenReaderNavigationPolicy.h"
 #include "TextToSpeech.h"
+#include "ScreenReaderLog.h"
 
 FScreenReaderUser::FScreenReaderUser(FAccessibleUserIndex InUserIndex)
 	: FGenericAccessibleUser(InUserIndex)
+	, NavigationPolicy(MakeShared<FScreenReaderDefaultNavigationPolicy>())
 	, bActive(false)
 {
 	// @TODOAccessibility: For now, just default it to the platform TTS. Should give a way to allow users to use custom TTS though 
@@ -120,3 +123,98 @@ void FScreenReaderUser::OnUnregistered()
 	Deactivate();
 }
 
+void FScreenReaderUser::SetNavigationPolicy(const TSharedRef<IScreenReaderNavigationPolicy>& InNavigationPolicy)
+{
+	NavigationPolicy = InNavigationPolicy;
+}
+
+FScreenReaderReply FScreenReaderUser::NavigateToNextSibling()
+{
+	TSharedPtr<IAccessibleWidget> CurrentFocusedAccessibleWidget = GetFocusedAccessibleWidget();
+	if (CurrentFocusedAccessibleWidget)
+	{
+		TSharedPtr<IAccessibleWidget> NextSibling = NavigationPolicy->GetNextSiblingFrom(CurrentFocusedAccessibleWidget.ToSharedRef());
+		if (NextSibling && NextSibling->SetUserFocus(GetIndex()))
+		{
+			UE_LOG(LogScreenReaderNavigation, Verbose, TEXT("User %d now focused on next sibling %s."), GetIndex(), *NextSibling->ToString());
+			return FScreenReaderReply::Handled();
+		}
+	}
+	return FScreenReaderReply::Unhandled();
+}
+
+FScreenReaderReply FScreenReaderUser::NavigateToPreviousSibling()
+{
+	TSharedPtr<IAccessibleWidget> CurrentFocusedAccessibleWidget = GetFocusedAccessibleWidget();
+	if (CurrentFocusedAccessibleWidget)
+	{
+		TSharedPtr<IAccessibleWidget> PreviousSibling = NavigationPolicy->GetPreviousSiblingFrom(CurrentFocusedAccessibleWidget.ToSharedRef());
+		if (PreviousSibling && PreviousSibling->SetUserFocus(GetIndex()))
+		{
+			UE_LOG(LogScreenReaderNavigation, Verbose, TEXT("User %d now focused on previous sibling %s."), GetIndex(), *PreviousSibling->ToString());
+			return FScreenReaderReply::Handled();
+		}
+	}
+	return FScreenReaderReply::Unhandled();
+}
+
+FScreenReaderReply FScreenReaderUser::NavigateToFirstAncestor()
+{
+	TSharedPtr<IAccessibleWidget> CurrentFocusedAccessibleWidget = GetFocusedAccessibleWidget();
+	if (CurrentFocusedAccessibleWidget)
+	{
+		TSharedPtr<IAccessibleWidget> Ancestor = NavigationPolicy->GetFirstAncestorFrom(CurrentFocusedAccessibleWidget.ToSharedRef());
+		if (Ancestor && Ancestor->SetUserFocus(GetIndex()))
+		{
+			UE_LOG(LogScreenReaderNavigation, Verbose, TEXT("User %d now focused on ancestor %s."), GetIndex(), *Ancestor->ToString());
+			return FScreenReaderReply::Handled();
+		}
+	}
+	return FScreenReaderReply::Unhandled();
+}
+
+FScreenReaderReply FScreenReaderUser::NavigateToFirstChild()
+{
+	TSharedPtr<IAccessibleWidget> CurrentFocusedAccessibleWidget = GetFocusedAccessibleWidget();
+	if (CurrentFocusedAccessibleWidget)
+	{
+		TSharedPtr<IAccessibleWidget> Child = NavigationPolicy->GetFirstChildFrom(CurrentFocusedAccessibleWidget.ToSharedRef());
+		if (Child && Child->SetUserFocus(GetIndex()))
+		{
+			UE_LOG(LogScreenReaderNavigation, Verbose, TEXT("User %d now focused on first child %s."), GetIndex(), *Child->ToString());
+			return FScreenReaderReply::Handled();
+		}
+	}
+	return FScreenReaderReply::Unhandled();
+}
+
+FScreenReaderReply FScreenReaderUser::NavigateToNextWidgetInHierarchy()
+{
+	TSharedPtr<IAccessibleWidget> CurrentFocusedAccessibleWidget = GetFocusedAccessibleWidget();
+	if (CurrentFocusedAccessibleWidget)
+	{
+		TSharedPtr<IAccessibleWidget> NextWidget = NavigationPolicy->GetNextWidgetInHierarchyFrom(CurrentFocusedAccessibleWidget.ToSharedRef());
+		if (NextWidget && NextWidget->SetUserFocus(GetIndex()))
+		{
+			UE_LOG(LogScreenReaderNavigation, Verbose, TEXT("User %d now focused on next widget %s."), GetIndex(), *NextWidget->ToString());
+			return FScreenReaderReply::Handled();
+		}
+	}
+	return FScreenReaderReply::Unhandled();
+}
+
+FScreenReaderReply FScreenReaderUser::NavigateToPreviousWidgetInHierarchy()
+{
+	TSharedPtr<IAccessibleWidget> CurrentFocusedAccessibleWidget = GetFocusedAccessibleWidget();
+	if (CurrentFocusedAccessibleWidget)
+	{
+		TSharedPtr<IAccessibleWidget> PreviousWidget = NavigationPolicy->GetPreviousWidgetInHierarchyFrom(CurrentFocusedAccessibleWidget.ToSharedRef());
+		if (PreviousWidget && PreviousWidget->SetUserFocus(GetIndex()))
+		{
+			UE_LOG(LogScreenReaderNavigation, Verbose, TEXT("User %d now focused on previous widget %s."), GetIndex(), *PreviousWidget->ToString());
+			return FScreenReaderReply::Handled();
+		}
+	}
+	return FScreenReaderReply::Unhandled();
+}
+
diff --git a/Engine/Plugins/Experimental/ScreenReader/Source/ScreenReader/Public/GenericPlatform/ScreenReaderUser.h b/Engine/Plugins/Experimental/ScreenReader/Source/ScreenReader/Public/GenericPlatform/ScreenReaderUser.h
--- a/Engine/Plugins/Experimental/ScreenReader/Source/ScreenReader/Public/GenericPlatform/ScreenReaderUser.h
+++ b/Engine/Plugins/Experimental/ScreenReader/Source/ScreenReader/Public/GenericPlatform/ScreenReaderUser.h
@@ -8,6 +8,7 @@
 
 class IAccessibleWidget;
 class FScreenReaderAnnouncementChannel;
+class IScreenReaderNavigationPolicy;
 
 /**
 * A user of the screen reader. Corresponds to a hardware device that users use.
@@ -18,6 +19,7 @@ class FScreenReaderAnnouncementChannel;
 * Responsibilities of the class are:
 * 1. Text to speech (TTS) requests - Users can request an announcement to be spoken via text to speech to  screen reader user
 * 2. Accessible focus handling - A screen reader user holds information about the accessible widget it is currently focused on
+ * 3. Accessible Navigation - Screen reader users have finer navigation controls around the accessible widget hierarchy as compared to regular tab navigation with keyboard or D-pad controls with a gamepad.
 * @see FScreenReaderBase, FScreenReaderAnnouncement, FGenericAccessibleUser
 */
 class SCREENREADER_API FScreenReaderUser : public FGenericAccessibleUser
@@ -90,6 +92,43 @@ public:
 	/** Returns true if the text to speech for this user is muted. Else returns false. */
 	bool IsSpeechMuted() const;
 
+	
+	// Navigation
+	/** Sets the navigation policy this screen reader user will be using. The screen reader navigation policy will affect all functionality to do with shifting focus to the next available widget. */
+	void SetNavigationPolicy(const TSharedRef<IScreenReaderNavigationPolicy>& InNavigationPolicy);
+	/**
+	 * Shifts focus to the next sibling from the user's currently focused accessible widget based on the current navigation policy.
+	 * @return FScreenReaderReply::Handled() if focus has successfully shifted to a next sibling from the currently focused accessible widget. Else returns FScreenReaderReply::Unhandled()
+	 */
+	FScreenReaderReply NavigateToNextSibling();
+	/**
+	 * Shifts focus to the previous sibling from the user's currently focused accessible widget based on the current navigation policy.
+	 * @return FScreenReaderReply::Handled() if focus has successfully shifted to a previous sibling from the currently focused accessible widget. Else returns FScreenReaderReply::Unhandled()
+	 */
+	FScreenReaderReply NavigateToPreviousSibling();
+	/**
+	 * Shifts focus to the first ancestor of the user's currently focused accessible widget based on the current navigation policy.
+	 * @return FScreenReaderReply::Handled() if focus has successfully shifted to the first ancestor of the currently focused accessible widget. Else returns FScreenReaderReply::Unhandled()
+	 */
+	FScreenReaderReply NavigateToFirstAncestor();
+	/**
+	 * Shifts focus to the first child from the user's currently focused accessible widget based on the current navigation policy.
+	 * Focus is unaffected if the user's currently focused accessible widget has no children.
+	 * @return FScreenReaderReply::Handled() if focus has successfully shifted to the first child from the currently focused accessible widget based on the navigation policy. Else returns FScreenReaderReply::Unhandled()
+	 */
+	FScreenReaderReply NavigateToFirstChild();
+	/**
+	 * Shifts focus to the logical next widget in the accessible widget hierarchy from the user's currently focused accessible widget based on the current navigation policy.
+	 * See IAccessibleWidget::GetNextWidgetInHierarchy for an explanation of what the logical next widget in the accessible hierarchy means.
+	 * @return FScreenReaderReply::Handled() if focus has successfully shifted to a next widget from the currently focused accessible widget. Else returns FScreenReaderReply::Unhandled()
+	 */
+	FScreenReaderReply NavigateToNextWidgetInHierarchy();
+	/**
+	 * Shifts focus to the logical previous widget in the accessible widget hierarchy from the user's currently focused accessible widget based on the current navigation policy.
+	 * See IAccessibleWidget::GetPreviousWidgetInHierarchy for an explanation of what the logical previous widget in the accessible hierarchy means.
+	 * @return FScreenReaderReply::Handled() if focus has successfully shifted to a logical previous widget from the currently focused accessible widget. Else returns FScreenReaderReply::Unhandled()
+	 */
+	FScreenReaderReply NavigateToPreviousWidgetInHierarchy();
 	/**
 	* Activates the screen reader user and fulfill requests for accessibility services such as text to speech that clients can make.
 	* When screen reader users are firstr registered with a screen reader, they are deactivated by default. Users must explicitly activate the screen reader user.
@@ -114,6 +153,7 @@ protected:
 private:
 	/** Responsible for handling all incoming announcement requests and speaking them via text to speech if possible */
 	TUniquePtr<FScreenReaderAnnouncementChannel> AnnouncementChannel;
+	TSharedRef<IScreenReaderNavigationPolicy> NavigationPolicy;
 	bool bActive;
 };
 
diff --git a/Engine/Plugins/Experimental/ScreenReader/Source/ScreenReader/Public/Navigation/ScreenReaderNavigationPolicy.h b/Engine/Plugins/Experimental/ScreenReader/Source/ScreenReader/Public/Navigation/ScreenReaderNavigationPolicy.h
new file mode 100644
--- /dev/null
+++ b/Engine/Plugins/Experimental/ScreenReader/Source/ScreenReader/Public/Navigation/ScreenReaderNavigationPolicy.h
@@ -0,0 +1,77 @@
+// Copyright Epic Games, Inc. All Rights Reserved.
+
+#pragma once
+#include "CoreMinimal.h"
+
+class IAccessibleWidget;
+/**
+ * An interface that controls how a screen reader user navigates around the accessible widget hierarchy and what accessible widgets can
+ * be navigated to from a source widget.
+ * Clients can provide custom implementations to provide different navigation behaviors for screen reader users at runtime. A sample implementation is FDefaultScreenReaderNavigationPolicy.
+ * @see FScreenReaderUser, FDefaultScreenReaderNavigaitonPolicy
+ */
+class SCREENREADER_API IScreenReaderNavigationPolicy
+{
+public:
+	virtual ~IScreenReaderNavigationPolicy() = default;
+	/** Returns the first instance of a next sibling from the source widget that satisfies the navigation policy. */
+	virtual TSharedPtr<IAccessibleWidget> GetNextSiblingFrom(const TSharedRef<IAccessibleWidget>& Source) const = 0;
+	/** Returns the first instance of a previous sibling from the source widget that satisfies the navigation policy. */
+	virtual TSharedPtr<IAccessibleWidget> GetPreviousSiblingFrom(const TSharedRef<IAccessibleWidget>& Source) const = 0;
+	/** Returns the first ancestor from the source widget that satisfies the navigation policy. */
+	virtual TSharedPtr<IAccessibleWidget> GetFirstAncestorFrom(const TSharedRef<IAccessibleWidget>& Source) const = 0;
+	/** Returns the first child from the source widget that satisfies the navigation policy. */
+	virtual TSharedPtr<IAccessibleWidget> GetFirstChildFrom(const TSharedRef<IAccessibleWidget>& Source) const = 0;
+	/**
+	 * Returns the first instance of a logical next widget from the source widget that satisfies the navigation policy.
+	 * See IAccessibleWidget::GetNextWidgetInHierarchy for an explanation of what the logical next widget in the accessible hierarchy means.
+	 * @see IAccessibleWidget
+	 */
+	virtual TSharedPtr<IAccessibleWidget> GetNextWidgetInHierarchyFrom(const TSharedRef<IAccessibleWidget>& Source) const = 0;
+	/**
+	 * Returns the first instance of a logical previous widget from the source widget that satisfies the navigation policy.
+	 * See IAccessibleWidget::GetPreviousWidgetInHierarchy for an explanation of what the logical previous widget in the accessible hierarchy means.
+	 * @see IAccessibleWidget
+	 */
+	virtual TSharedPtr<IAccessibleWidget> GetPreviousWidgetInHierarchyFrom(const TSharedRef<IAccessibleWidget>& Source) const = 0;
+};
+
+class SCREENREADER_API FScreenReaderDefaultNavigationPolicy final : public IScreenReaderNavigationPolicy
+{
+public:
+	struct FScreenReaderDefaultNavigationPredicate
+	{
+		bool operator()(const TSharedRef<IAccessibleWidget>& InWidget) const
+		{
+			return InWidget->CanCurrentlyAcceptAccessibleFocus();
+		}
+	};
+	virtual ~FScreenReaderDefaultNavigationPolicy() = default;
+	//~ Begin IScreenReaderNavigationPolicy Interface
+	virtual TSharedPtr<IAccessibleWidget> GetNextSiblingFrom(const TSharedRef<IAccessibleWidget>& Source) const override final
+	{
+		return IAccessibleWidget::SearchForNextSiblingFrom(Source, FScreenReaderDefaultNavigationPredicate());
+	}
+	virtual TSharedPtr<IAccessibleWidget> GetPreviousSiblingFrom(const TSharedRef<IAccessibleWidget>& Source) const override final
+	{
+		return IAccessibleWidget::SearchForPreviousSiblingFrom(Source, FScreenReaderDefaultNavigationPredicate());
+	}
+	virtual TSharedPtr<IAccessibleWidget> GetFirstAncestorFrom(const TSharedRef<IAccessibleWidget>& Source) const override final
+	{
+		return IAccessibleWidget::SearchForAncestorFrom(Source, FScreenReaderDefaultNavigationPredicate());
+	}
+	virtual TSharedPtr<IAccessibleWidget> GetFirstChildFrom(const TSharedRef<IAccessibleWidget>& Source) const override final
+	{
+		return IAccessibleWidget::SearchForFirstChildFrom(Source, FScreenReaderDefaultNavigationPredicate());
+	}
+	virtual TSharedPtr<IAccessibleWidget> GetNextWidgetInHierarchyFrom(const TSharedRef<IAccessibleWidget>& Source) const override final
+	{
+		return IAccessibleWidget::SearchForNextWidgetInHierarchyFrom(Source, FScreenReaderDefaultNavigationPredicate());
+	}
+	virtual TSharedPtr<IAccessibleWidget> GetPreviousWidgetInHierarchyFrom(const TSharedRef<IAccessibleWidget>& Source) const override final
+	{
+		return IAccessibleWidget::SearchForPreviousWidgetInHierarchyFrom(Source, FScreenReaderDefaultNavigationPredicate());
+	}
+	//~ End IScreenReaderNavigationPolicy Interface
+};
+
diff --git a/Engine/Source/Runtime/ApplicationCore/Public/GenericPlatform/Accessibility/GenericAccessibleInterfaces.h b/Engine/Source/Runtime/ApplicationCore/Public/GenericPlatform/Accessibility/GenericAccessibleInterfaces.h
--- a/Engine/Source/Runtime/ApplicationCore/Public/GenericPlatform/Accessibility/GenericAccessibleInterfaces.h
+++ b/Engine/Source/Runtime/ApplicationCore/Public/GenericPlatform/Accessibility/GenericAccessibleInterfaces.h
@@ -501,18 +501,158 @@ public:
 	 * @return The accessible parent widget of this widget.
 	 */
 	virtual TSharedPtr<IAccessibleWidget> GetParent() = 0;
+	/**
+	 * Returns the first instance of an ancestor from a source widget that satisfies a search criteria.
+	 * If the passed in source widget does not have any ancestors or if there are no ancestors that satisfy the search criteria, nullptr is returned.
+	 * The search criteria can be either a functor or lambda that acts as a predicate. It must return a bool and take in a const TSharedRef<IAccessibleWidget>& as an argument.
+	 *
+	 * @param Source The accessible widget to start searching from.
+	 * @param SearchCriteria A predicate that takes a const TSharedRef<IAccessibleWidget>& as an argument and returns true if the passed in widget satisfies the search criteria. When an accessible widget that satisfies the SearchCriteria is found, it is returned by the function.
+	 * @return The first instance of an ancestor from the source widget that satisfies the search criteria.
+	 */
+	template<typename PredicateType>
+	static TSharedPtr<IAccessibleWidget> SearchForAncestorFrom(const TSharedRef<IAccessibleWidget>& Source, PredicateType Predicate)
+	{
+		TSharedPtr<IAccessibleWidget> Ancestor = Source->GetParent();
+		while (Ancestor)
+		{
+			if (Predicate(Ancestor.ToSharedRef()))
+			{
+				return Ancestor;
+			}
+			Ancestor = Ancestor->GetNextSibling();
+		}
+		return nullptr;
+	}
 	/**
 	 * Retrieves the widget after this one in the parent's list of children. This should return nullptr for the last widget.
 	 *
 	 * @return The next widget on the same level of the UI hierarchy.
 	 */
 	virtual TSharedPtr<IAccessibleWidget> GetNextSibling() = 0;
+	/**
+	 * Returns the first instance of a next sibling from a source widget that satisfies a search criteria.
+	 * If the passed in source widget does not have any next siblings or if there are no next siblings that satisfy the search criteria, nullptr is returned.
+	 * The search criteria can be either a functor or lambda that acts as a predicate. It must return a bool and take in a const TSharedRef<IAccessibleWidget>& as an argument.
+	 *
+	 * @param Source The accessible widget to start searching from.
+	 * @param SearchCriteria A predicate that takes a const TSharedRef<IAccessibleWidget>& as an argument and returns true if the passed in widget satisfies the search criteria. When an accessible widget that satisfies the SearchCriteria is found, it is returned by the function.
+	 * @return The first instance of a next sibling from the source widget that satisfies the search criteria.
+	 */
+	template<typename PredicateType>
+	static TSharedPtr<IAccessibleWidget> SearchForNextSiblingFrom(const TSharedRef<IAccessibleWidget>& Source, PredicateType Predicate)
+	{
+		TSharedPtr<IAccessibleWidget> NextSibling = Source->GetNextSibling();
+		while (NextSibling)
+		{
+			if (Predicate(NextSibling.ToSharedRef()))
+			{
+				return NextSibling;
+			}
+			NextSibling = NextSibling->GetNextSibling();
+		}
+		return nullptr;
+	}
 	/**
 	 * Retrieves the widget before this one in the parent's list of children. This should return nullptr for the first widget.
 	 *
 	 * @return The previous widget on the same level of the UI hierarchy.
 	 */
 	virtual TSharedPtr<IAccessibleWidget> GetPreviousSibling() = 0;
+	/**
+	 * Returns the first instance of a previous sibling from a source widget that satisfies a search criteria.
+	 * If the passed in source widget does not have any previous siblings or if there are no previous siblings that satisfy the search criteria, nullptr is returned.
+	 * The search criteria can be either a functor or lambda that acts as a predicate. It must return a bool and take in a const TSharedRef<IAccessibleWidget>& as an argument.
+	 *
+	 * @param Source The accessible widget to start searching from.
+	 * @param SearchCriteria A predicate that takes a const TSharedRef<IAccessibleWidget>& as an argument and returns true if the passed in widget satisfies the search criteria. When an accessible widget that satisfies the SearchCriteria is found, it is returned by the function.
+	 * @return The first instance of a previous sibling from the source widget that satisfies the search criteria.
+	 */
+	template<typename PredicateType>
+	static TSharedPtr<IAccessibleWidget> SearchForPreviousSiblingFrom(const TSharedRef<IAccessibleWidget>& Source, PredicateType Predicate)
+	{
+		TSharedPtr<IAccessibleWidget> PreviousSibling = Source->GetPreviousSibling();
+		while(PreviousSibling)
+		{
+			if (Predicate(PreviousSibling.ToSharedRef()))
+			{
+				return PreviousSibling;
+			}
+			PreviousSibling = PreviousSibling->GetPreviousSibling();
+		}
+		return nullptr;
+	}
+	/**
+	 * Retrieves the logical next widget in the accessible widget hierarchy from this widget.
+	 * This is conceptually similar to tab navigation with the keyboard.
+	 * This is primarily used for mobile devices to simulate the right swipe gesture for IOS Voiceover or Android Talkback.
+	 * An example algorithm to find the next widget in the accessible hierarchy is as follows:
+	 * 1. If the current widget has children, return the first child.
+	 * 2. If the current widget has no children, return the sibling of the current widget.
+	 * 3. If the current widget has no next sibling, search for the first ancestor from the current widget with a next sibling. Return that ancestor's next sibling.
+	 *
+	 * @return The logical next widget in the accessible widget hierarchy.
+	 */
+	virtual TSharedPtr<IAccessibleWidget> GetNextWidgetInHierarchy() = 0;
+	/**
+	 * Returns the first instance of a next widget in the accessible hierarchy from a source widget that satisfies a search criteria.
+	 * If the passed in source widget does not have any next widgets or if there are no next widgets that satisfy the search criteria, nullptr is returned.
+	 * The search criteria can be either a functor or lambda that acts as a predicate. It must return a bool and take in a const TSharedRef<IAccessibleWidget>& as an argument.
+	 *
+	 * @param Source The accessible widget to start searching from.
+	 * @param SearchCriteria A predicate that takes a const TSharedRef<IAccessibleWidget>& as an argument and returns true if the passed in widget satisfies the search criteria. When an accessible widget that satisfies the SearchCriteria is found, it is returned by the function.
+	 * @return The first instance of a next widget from the source widget that satisfies the search criteria.
+	 */
+	template<typename PredicateType>
+	static TSharedPtr<IAccessibleWidget> SearchForNextWidgetInHierarchyFrom(const TSharedRef<IAccessibleWidget>& Source, PredicateType Predicate)
+	{
+		TSharedPtr<IAccessibleWidget> NextWidget = Source->GetNextWidgetInHierarchy();
+		while(NextWidget)
+		{
+			if (Predicate(NextWidget.ToSharedRef()))
+			{
+				return NextWidget;
+			}
+			NextWidget = NextWidget->GetNextWidgetInHierarchy();
+		}
+		return nullptr;
+	}
+	/**
+	 * Retrieves the logical previous widget in the accessible widget hierarchy from this current widget.
+	 * This is conceptually similar to shift tab navigation with the keyboard.
+	 * This is primarily used for mobile devices to simulate the left swipe gesture for IOS Voiceover or Android Talkback.
+	 * An example algorithm to find the previous widget in the hierarchy is as follows:
+	 * 1. Find the previous sibling of the current widget and check if it has children.
+	 * 2. If the previous sibling has children, we will call the last child of the previous sibling C. Recursively take the last child from C until a leaf is found. Return the leaf.
+	 * 3. If the previous sibling has no children, return the previous sibling.
+	 * 4. If the current widget has no previous sibling, we return the parent of the current widget.
+	 *
+	 * @return The logical previous widget in the accessible widget hierarchy.
+	 */
+	virtual TSharedPtr<IAccessibleWidget> GetPreviousWidgetInHierarchy() = 0;
+	/**
+	 * Returns the first instance of a previous widget in the accessible hierarchy from a source widget that satisfies a search criteria.
+	 * If the passed in source widget does not have any previous widgets or if there are no previous widgets that satisfy the search criteria, nullptr is returned.
+	 * The search criteria can be either a functor or lambda that acts as a predicate. It must return a bool and take in a const TSharedRef<IAccessibleWidget>& as an argument.
+	 *
+	 * @param Source The accessible widget to start searching from.
+	 * @param SearchCriteria A predicate that takes a const TSharedRef<IAccessibleWidget>& as an argument and returns true if the passed in widget satisfies the search criteria. When an accessible widget that satisfies the SearchCriteria is found, it is returned by the function.
+	 * @return The first instance of a previous widget from the source widget that satisfies the search criteria.
+	 */
+	template<typename PredicateType>
+	static TSharedPtr<IAccessibleWidget> SearchForPreviousWidgetInHierarchyFrom(const TSharedRef<IAccessibleWidget>& Source, PredicateType Predicate)
+	{
+		TSharedPtr<IAccessibleWidget> PreviousWidget = Source->GetPreviousWidgetInHierarchy();
+		while(PreviousWidget)
+		{
+			if (Predicate(PreviousWidget.ToSharedRef()))
+			{
+				return PreviousWidget;
+			}
+			PreviousWidget = PreviousWidget->GetPreviousWidgetInHierarchy();
+		}
+		return nullptr;
+	}
 	/**
 	 * Retrieves the accessible child widget at a certain index. This should return nullptr if Index < 0 or Index >= GetNumberOfChildren().
 	 *
@@ -526,6 +666,34 @@ public:
 	 * @return The number of accessible children that exist for this widget.
 	 */
 	virtual int32 GetNumberOfChildren() = 0;
+	/**
+	 * Returns the first instance of a child from a source widget that satisfies a search criteria.
+	 * If the passed in source widget does not have any children or if there are no children that satisfy the search criteria, nullptr is returned.
+	 * The search criteria can be either a functor or lambda that acts as a predicate. It must return a bool and take in a const TSharedRef<IAccessibleWidget>& as an argument.
+	 *
+	 * @param Source The accessible widget to start searching from.
+	 * @param SearchCriteria A predicate that takes a const TSharedRef<IAccessibleWidget>& as an argument and returns true if the passed in widget satisfies the search criteria. When an accessible widget that satisfies the SearchCriteria is found, it is returned by the function.
+	 * @return The first child from the source widget that satisfies the search criteria.
+	 */
+	template<typename PredicateType>
+	static TSharedPtr<IAccessibleWidget> SearchForFirstChildFrom(const TSharedRef<IAccessibleWidget>& Source, PredicateType Predicate)
+	{
+		if (Source->GetNumberOfChildren() == 0)
+		{
+			return nullptr;
+		}
+		TSharedPtr<IAccessibleWidget> Child = nullptr;
+		int32 NumChildren = Source->GetNumberOfChildren();
+		for (int32 ChildIndex = 0; ChildIndex < NumChildren; ++ChildIndex)
+		{
+			Child = Source->GetChildAt(ChildIndex);
+			if (Child && Predicate(Child.ToSharedRef()))
+			{
+				return Child;
+			}
+		}
+		return nullptr;
+	}
 	/**
 	 * What type of accessible widget the underlying widget should be treated as. A widget may be capable of presenting itself
 	 * as multiple different types of widgets, but only one can be reported back to the platform.
@@ -552,6 +720,13 @@ public:
 	 * @return A more-detailed description of what the widget is or how its used.
 	 */
 	virtual FString GetHelpText() const = 0;
+	/** Returns a string representation of this widget. Primarily used for debugging. */
+	virtual FString ToString() const
+	{
+		TStringBuilder<256> Builder;
+		Builder.Appendf(TEXT("Label: %s. Role: %s."), *GetWidgetName(), *GetClassName());
+		return Builder.ToString();
+	}
 
 	/**
 	 * Whether the widget is enabled and can be interacted with.
diff --git a/Engine/Source/Runtime/SlateCore/Private/Widgets/Accessibility/SlateCoreAccessibleWidgets.cpp b/Engine/Source/Runtime/SlateCore/Private/Widgets/Accessibility/SlateCoreAccessibleWidgets.cpp
--- a/Engine/Source/Runtime/SlateCore/Private/Widgets/Accessibility/SlateCoreAccessibleWidgets.cpp
+++ b/Engine/Source/Runtime/SlateCore/Private/Widgets/Accessibility/SlateCoreAccessibleWidgets.cpp
@@ -189,8 +189,7 @@ bool FSlateAccessibleWidget::SetUserFocus(const FAccessibleUserIndex UserIndex)
 	// Most likely a mistake to set focus on a widget that cannot currently accept focus
 	if (!CanCurrentlyAcceptAccessibleFocus())
 	{
-		// @TODOAccessibility: Log information about the widget when we create a ToString() function.
-		UE_LOG(LogAccessibility, Warning, TEXT("Attempting to set accessible focus on an accessible widget that currently cannot be focused. Focus not changed."));
+		UE_LOG(LogAccessibility, Warning, TEXT("Accessible widget %s cannot currently be focused. Focus not changed."), *ToString());
 		return false;
 	}
 	if (SupportsFocus())
@@ -203,8 +202,15 @@ bool FSlateAccessibleWidget::SetUserFocus(const FAccessibleUserIndex UserIndex)
 			FWidgetPath WidgetPath;
 			if (FSlateWindowHelper::FindPathToWidget(WindowArray, Widget.Pin().ToSharedRef(), WidgetPath))
 			{
+				FSlateApplicationBase& SlateApp = FSlateApplicationBase::Get();
 				// Focus accessible events are already raised from this function call
-				return FSlateApplicationBase::Get().SetUserFocus(UserIndex, WidgetPath, EFocusCause::SetDirectly);
+				if (!SlateApp.SetUserFocus(UserIndex, WidgetPath, EFocusCause::SetDirectly))
+				{
+					// If the function returns false, it means the SWidget corresponding to the accessible widget is already focused by the application. We will need to
+					// manually raise a focus change event to sync the accessible focus and Slate focus
+					SlateApp.GetAccessibleMessageHandler()->OnWidgetEventRaised(FSlateAccessibleMessageHandler::FSlateWidgetAccessibleEventArgs(Widget.Pin().ToSharedRef(), EAccessibleEvent::FocusChange, false, true, UserIndex));
+				}
+				return true;
 			}
 		}
 	}
@@ -261,6 +267,69 @@ TSharedPtr<IAccessibleWidget> FSlateAccessibleWidget::GetPreviousSibling()
 	return nullptr;
 }
 
+TSharedPtr<IAccessibleWidget> FSlateAccessibleWidget::GetNextWidgetInHierarchy()
+{
+	// 1. if the current widget has children, return the first child
+	// 2. If the widget has no children, return the next sibling of the current widget.
+	// 3. If the widget has no sibling, find the first ancestor with a next sibling and return that ancestor's next sibling.
+	if (GetNumberOfChildren() > 0)
+	{
+		return GetChildAt(0);
+	}
+	else if (TSharedPtr<IAccessibleWidget> NextSibling = GetNextSibling())
+	{
+		return NextSibling;
+	}
+	else if (TSharedPtr<IAccessibleWidget> CurrentAncestor = GetParent())
+	{
+		TSharedPtr<IAccessibleWidget> FoundAncestorSibling;
+		while (CurrentAncestor)
+		{
+			FoundAncestorSibling = CurrentAncestor->GetNextSibling();
+			if (FoundAncestorSibling)
+			{
+				return FoundAncestorSibling;
+			}
+			// if we get here, the current ancestor has no next sibling, move up the tree and try again
+			CurrentAncestor = CurrentAncestor->GetParent();
+		}
+		// if we get here, we are the right most leaf of the tree. There is no next accessible widget
+	}
+		return nullptr;
+}
+
+TSharedPtr<IAccessibleWidget> FSlateAccessibleWidget::GetPreviousWidgetInHierarchy()
+{
+	// 1. Find the previous sibling and see if it has children
+	// 2. If the previous sibling has children, DFS down the right most child until you find the child that is a leaf
+	// 3. If the sibling has no children, just return the previous sibling
+	// 4. If there is no previous sibling, return the parent
+	if (TSharedPtr<IAccessibleWidget> PreviousSibling = GetPreviousSibling())
+	{
+		TSharedPtr<IAccessibleWidget> ReturnWidget = PreviousSibling;
+		while (ReturnWidget)
+		{
+			int32 ReturnWidgetChildrenCount = ReturnWidget->GetNumberOfChildren();
+			if (ReturnWidgetChildrenCount > 0)
+			{
+				// keep taking the last child until we get to the leaf
+				ReturnWidget = ReturnWidget->GetChildAt(ReturnWidgetChildrenCount - 1);
+			}
+			else
+			{
+				// we are at a leaf just return
+				return ReturnWidget;
+			}
+		}
+	}
+	else
+	{
+		// note that it is possible for this widget to be the root of the tree and for this to return nullptr
+		return GetParent();
+	}
+	return nullptr;
+}
+
 TSharedPtr<IAccessibleWidget> FSlateAccessibleWidget::GetChildAt(int32 Index)
 {
 	check(Index >= 0 && Index < Children.Num());
diff --git a/Engine/Source/Runtime/SlateCore/Public/Widgets/Accessibility/SlateCoreAccessibleWidgets.h b/Engine/Source/Runtime/SlateCore/Public/Widgets/Accessibility/SlateCoreAccessibleWidgets.h
--- a/Engine/Source/Runtime/SlateCore/Public/Widgets/Accessibility/SlateCoreAccessibleWidgets.h
+++ b/Engine/Source/Runtime/SlateCore/Public/Widgets/Accessibility/SlateCoreAccessibleWidgets.h
@@ -30,6 +30,8 @@ public:
 	virtual TSharedPtr<IAccessibleWidget> GetParent() override final;
 	virtual TSharedPtr<IAccessibleWidget> GetNextSibling() override final;
 	virtual TSharedPtr<IAccessibleWidget> GetPreviousSibling() override final;
+	virtual TSharedPtr<IAccessibleWidget> GetNextWidgetInHierarchy() override final;
+	virtual TSharedPtr<IAccessibleWidget> GetPreviousWidgetInHierarchy() override final;
 	virtual TSharedPtr<IAccessibleWidget> GetChildAt(int32 Index) override final;
 	virtual int32 GetNumberOfChildren() override final;
 	virtual FString GetClassName() const override final;
