Change iOS sound suspend / resume audio logic to fix sound bug on iOS 15

diff --git a/Engine/Source/Runtime/ApplicationCore/Private/IOS/IOSAppDelegate.cpp b/Engine/Source/Runtime/ApplicationCore/Private/IOS/IOSAppDelegate.cpp
--- a/Engine/Source/Runtime/ApplicationCore/Private/IOS/IOSAppDelegate.cpp
+++ b/Engine/Source/Runtime/ApplicationCore/Private/IOS/IOSAppDelegate.cpp
@@ -547,20 +547,15 @@ static IOSAppDelegate* CachedDelegate = nil;
 		switch ([[[notification userInfo] objectForKey:AVAudioSessionInterruptionTypeKey] unsignedIntegerValue])
 		{
 			case AVAudioSessionInterruptionTypeBegan:
-				self.bAudioActive = false;
-				FAppEntry::Suspend(true);
+				// Notification that our audio was stopped by the OS. No action needed.
 				break;
 
 			case AVAudioSessionInterruptionTypeEnded:
-
 				NSNumber * interruptionOption = [[notification userInfo] objectForKey:AVAudioSessionInterruptionOptionKey];
 				if (interruptionOption != nil && interruptionOption.unsignedIntegerValue > 0)
 				{
-					FAppEntry::RestartAudio();
+					[self ToggleAudioSession:true];
 				}
-
-				FAppEntry::Resume(true);
-				[self ToggleAudioSession:true];
 				break;
 		}
 	}];
@@ -1356,6 +1351,9 @@ FCriticalSection RenderSuspend;
     [self ToggleSuspend:true];
 #endif
 
+    self.bAudioActive = false;
+    FAppEntry::Suspend(true);
+
 	FEmbeddedCommunication::KeepAwake(TEXT("Background"), false);
 
 	[FIOSAsyncTask CreateTaskWithBlock : ^ bool(void)
@@ -1401,6 +1399,9 @@ extern double GCStartTime;
 	RenderSuspend.Unlock();
 	[self ToggleSuspend : false];
     [self ToggleAudioSession:true];
+	FAppEntry::RestartAudio();
+    FAppEntry::Resume(true);
+
 
     if (bEngineInit)
     {
