Fix missing audio after phone call

diff --git a/Engine/Source/Runtime/ApplicationCore/Private/IOS/IOSAppDelegate.cpp b/Engine/Source/Runtime/ApplicationCore/Private/IOS/IOSAppDelegate.cpp
--- a/Engine/Source/Runtime/ApplicationCore/Private/IOS/IOSAppDelegate.cpp
+++ b/Engine/Source/Runtime/ApplicationCore/Private/IOS/IOSAppDelegate.cpp
@@ -554,8 +554,9 @@ static IOSAppDelegate* CachedDelegate = nil;
 				NSNumber * interruptionOption = [[notification userInfo] objectForKey:AVAudioSessionInterruptionOptionKey];
 				if (interruptionOption != nil && interruptionOption.unsignedIntegerValue > 0)
 				{
-					[self ToggleAudioSession:true];
-				}
+                    FAppEntry::RestartAudio();
+                }
+                [self ToggleAudioSession:true];
 				break;
 		}
 	}];
