Fix for Apple controller crash during connect disconnect.  Race condition on the controller array objects due to connect and disconnect callbacks running on app main thread while send controller events update is running on a different thread.   Add a deferred event array that enables all processing on the same thread as the controller events.  Retain and release controllers on connect and disconnect.  Ensure the player index is initally set correctly and on disconnect as it was all set to PlayerOne.

Tested on iOS and macOS.

diff --git a/Engine/Source/Runtime/ApplicationCore/Private/Apple/AppleControllerInterface.cpp b/Engine/Source/Runtime/ApplicationCore/Private/Apple/AppleControllerInterface.cpp
--- a/Engine/Source/Runtime/ApplicationCore/Private/Apple/AppleControllerInterface.cpp
+++ b/Engine/Source/Runtime/ApplicationCore/Private/Apple/AppleControllerInterface.cpp
@@ -18,31 +18,38 @@ FAppleControllerInterface::FAppleControllerInterface( const TSharedRef< FGeneric
 {
 	if(!IS_PROGRAM)
 	{
+		// Clear array and setup unset player index values
+		FMemory::Memzero(Controllers, sizeof(Controllers));
+		for (int32 ControllerIndex = 0; ControllerIndex < UE_ARRAY_COUNT(Controllers); ControllerIndex++)
+		{
+			FUserController& UserController = Controllers[ControllerIndex];
+			UserController.PlayerIndex = PlayerIndex::PlayerUnset;
+		}
+		
+		for (GCController* Cont in [GCController controllers])
+		{
+			HandleConnection(Cont);
+		}
+		
 		NSNotificationCenter* notificationCenter = [NSNotificationCenter defaultCenter];
-		NSOperationQueue* currentQueue = [NSOperationQueue currentQueue];
+		
+		// Not in an operation queue, [NSOperationQueue currentQueue] will return nil on macOS and iOS
+		// Notification callback will always be on the app main thread - defer events for Unreal Engine update thread
 
-		[notificationCenter addObserverForName:GCControllerDidDisconnectNotification object:nil queue:currentQueue usingBlock:^(NSNotification* Notification)
+		[notificationCenter addObserverForName:GCControllerDidDisconnectNotification object:nil queue:nil usingBlock:^(NSNotification* Notification)
 		{
-			HandleDisconnect(Notification.object);
+			SignalEvent(EAppleControllerEventType::Disconnect, Notification.object);
 		}];
 
-		[notificationCenter addObserverForName:GCControllerDidConnectNotification object:nil queue:currentQueue usingBlock:^(NSNotification* Notification)
+		[notificationCenter addObserverForName:GCControllerDidConnectNotification object:nil queue:nil usingBlock:^(NSNotification* Notification)
 		{
-			HandleConnection(Notification.object);
-            SetCurrentController(Notification.object);
+			SignalEvent(EAppleControllerEventType::Connect, Notification.object);
 		}];
 
 		dispatch_async(dispatch_get_main_queue(), ^
 		{
 		   [GCController startWirelessControllerDiscoveryWithCompletionHandler:^{ }];
 		});
-		
-		FMemory::Memzero(Controllers, sizeof(Controllers));
-		
-		for (GCController* Cont in [GCController controllers])
-		{
-			HandleConnection(Cont);
-		}
 	}
 }
 
@@ -51,9 +58,47 @@ void FAppleControllerInterface::SetMessageHandler( const TSharedRef< FGenericApp
 	MessageHandler = InMessageHandler;
 }
 
+void FAppleControllerInterface::SignalEvent(EAppleControllerEventType InEventType, GCController* InController)
+{
+	FScopeLock Lock(&DeferredEventCS);
+	DeferredEvents.Add(FDeferredAppleControllerEvent(InEventType, InController));
+}
+
 void FAppleControllerInterface::Tick( float DeltaTime )
 {
-	// NOP
+	FScopeLock Lock(&DeferredEventCS);
+	
+	for(uint32_t Index = 0;Index < DeferredEvents.Num();++Index)
+	{
+		FDeferredAppleControllerEvent& Event = DeferredEvents[Index];
+		switch(Event.EventType)
+		{
+			case EAppleControllerEventType::Connect:
+			{
+				HandleConnection(Event.Controller);
+				SetCurrentController(Event.Controller);
+				break;
+			}
+			case EAppleControllerEventType::Disconnect:
+			{
+				HandleDisconnect(Event.Controller);
+				break;
+			}
+			case EAppleControllerEventType::BecomeCurrent:
+			{
+				SetCurrentController(Event.Controller);
+				break;
+			}
+			case EAppleControllerEventType::Invalid:
+			default:
+			{
+				// NOP
+				break;
+			}
+		}
+	}
+	
+	DeferredEvents.Empty();
 }
 
 void FAppleControllerInterface::SetControllerType(uint32 ControllerIndex)
@@ -121,7 +166,7 @@ void FAppleControllerInterface::HandleConnection(GCController* Controller)
         }
         
         Controllers[ControllerIndex].PlayerIndex = (PlayerIndex)ControllerIndex;
-        Controllers[ControllerIndex].Controller = Controller;
+        Controllers[ControllerIndex].Controller = [Controller retain];
         SetControllerType(ControllerIndex);
         
         // Deprecated but buttonMenu behavior is unreliable in iOS/tvOS 14
@@ -155,12 +200,20 @@ void FAppleControllerInterface::HandleDisconnect(GCController* Controller)
 	
     for (int32 ControllerIndex = 0; ControllerIndex < UE_ARRAY_COUNT(Controllers); ControllerIndex++)
     {
-        if (Controllers[ControllerIndex].Controller == Controller)
+		FUserController& UserController = Controllers[ControllerIndex];
+        if (UserController.Controller == Controller)
         {
-            FMemory::Memzero(&Controllers[ControllerIndex], sizeof(Controllers[ControllerIndex]));
-            UE_LOG(LogAppleController, Log, TEXT("Controller for playerIndex %d was removed"), Controllers[ControllerIndex].PlayerIndex);
-            return;
+			// Player index of unset(-1) would indicate that it has become unset even though it is now trying to disconnect
+			// This can occur on iOS when bGameSupportsMultipleActiveControllers is false
+			UE_LOG(LogAppleController, Log, TEXT("Controller for playerIndex %d, controller Index %d removed"), UserController.PlayerIndex, ControllerIndex);
+			
+			[UserController.Controller release];
+			[UserController.PreviousExtendedGamepad release];
+			
+            FMemory::Memzero(&UserController, sizeof(Controllers[ControllerIndex]));
+            UserController.PlayerIndex = PlayerIndex::PlayerUnset;
             
+            return;
         }
     }
 }
@@ -170,24 +223,25 @@ void FAppleControllerInterface::SendControllerEvents()
 	@autoreleasepool{
     for(int32 i = 0; i < UE_ARRAY_COUNT(Controllers); ++i)
  	{
+		FUserController& Controller = Controllers[i];
+		
 		// make sure the connection handler has run on this guy
-		if (Controllers[i].PlayerIndex == PlayerIndex::PlayerUnset)
+		if (Controller.PlayerIndex == PlayerIndex::PlayerUnset)
 		{
             continue;
 		}
-
-		FUserController& Controller = Controllers[i];
-		GCController* Cont = Controller.Controller;
+		
+		GCController* ControllerImpl = Controller.Controller;
 
 		// Assumes iOS13, tvOS 13 & macOS 10.15
-        GCExtendedGamepad* ExtendedGamepad = [Cont capture].extendedGamepad;
-		GCMotion* Motion = Cont.motion;
+        GCExtendedGamepad* ExtendedGamepad = [ControllerImpl capture].extendedGamepad;
+		GCMotion* Motion = ControllerImpl.motion;
 		
 		// Workaround for unreliable buttonMenu behavior in iOS/tvOS 14
 		if (Controller.bPauseWasPressed)
         {
-            MessageHandler->OnControllerButtonPressed(FGamepadKeyNames::SpecialRight, Controllers[i].PlayerIndex, false);
-            MessageHandler->OnControllerButtonReleased(FGamepadKeyNames::SpecialRight, Controllers[i].PlayerIndex, false);
+            MessageHandler->OnControllerButtonPressed(FGamepadKeyNames::SpecialRight, Controller.PlayerIndex, false);
+            MessageHandler->OnControllerButtonReleased(FGamepadKeyNames::SpecialRight, Controller.PlayerIndex, false);
 
             Controller.bPauseWasPressed = false;
         }
diff --git a/Engine/Source/Runtime/ApplicationCore/Private/Apple/AppleControllerInterface.h b/Engine/Source/Runtime/ApplicationCore/Private/Apple/AppleControllerInterface.h
--- a/Engine/Source/Runtime/ApplicationCore/Private/Apple/AppleControllerInterface.h
+++ b/Engine/Source/Runtime/ApplicationCore/Private/Apple/AppleControllerInterface.h
@@ -25,7 +25,35 @@ enum PlayerIndex
 	PlayerTwo,
 	PlayerThree,
 	PlayerFour,
-	PlayerUnset
+	
+	PlayerUnset = -1
+};
+
+enum class EAppleControllerEventType : int32
+{
+    Invalid,
+    Connect,
+    Disconnect,
+    BecomeCurrent
+};
+
+struct FDeferredAppleControllerEvent
+{
+	FDeferredAppleControllerEvent(EAppleControllerEventType InEventType, GCController* InController)
+	: EventType(InEventType)
+	, Controller([InController retain])
+	{}
+	FDeferredAppleControllerEvent(const FDeferredAppleControllerEvent& Other)
+	{
+		EventType = Other.EventType;
+		Controller = [Other.Controller retain];
+	}
+	~FDeferredAppleControllerEvent()
+	{
+		[Controller release];
+	}
+    EAppleControllerEventType EventType;
+    GCController* Controller;
 };
 
 /**
@@ -76,7 +104,10 @@ public:
 protected:
 
 	FAppleControllerInterface( const TSharedRef< FGenericApplicationMessageHandler >& InMessageHandler );
+	
+	void SignalEvent(EAppleControllerEventType InEventType, GCController* InController);
 
+private:
 
 	void HandleConnection(GCController* Controller);
 	void HandleDisconnect(GCController* Controller);
@@ -110,6 +141,10 @@ protected:
         bool bPauseWasPressed;
 	};
 	
+	// Controller Event Callbacks are on the main thread - defer to tick processing
+	FCriticalSection DeferredEventCS;
+	TArray<FDeferredAppleControllerEvent> DeferredEvents;
+	
     // there is a hardcoded limit of 4 controllers in the API
 	FUserController Controllers[4];
 	
diff --git a/Engine/Source/Runtime/ApplicationCore/Private/IOS/IOSInputInterface.cpp b/Engine/Source/Runtime/ApplicationCore/Private/IOS/IOSInputInterface.cpp
--- a/Engine/Source/Runtime/ApplicationCore/Private/IOS/IOSInputInterface.cpp
+++ b/Engine/Source/Runtime/ApplicationCore/Private/IOS/IOSInputInterface.cpp
@@ -168,7 +168,7 @@ FIOSInputInterface::FIOSInputInterface( const TSharedRef< FGenericApplicationMes
         {
             [notificationCenter addObserverForName:GCControllerDidBecomeCurrentNotification object:nil queue:currentQueue usingBlock:^(NSNotification* Notification)
             {
-                SetCurrentController(Notification.object);
+                SignalEvent(EAppleControllerEventType::BecomeCurrent, Notification.object);
             }];
         }
     }
