Compile fix for DataSmithSDK - add include for FScopelock.

horde 130573

diff --git a/Engine/Source/Runtime/ApplicationCore/Private/Apple/AppleControllerInterface.cpp b/Engine/Source/Runtime/ApplicationCore/Private/Apple/AppleControllerInterface.cpp
--- a/Engine/Source/Runtime/ApplicationCore/Private/Apple/AppleControllerInterface.cpp
+++ b/Engine/Source/Runtime/ApplicationCore/Private/Apple/AppleControllerInterface.cpp
@@ -2,6 +2,7 @@
 
 #include "AppleControllerInterface.h"
 #include "HAL/PlatformTime.h"
+#include "Misc/ScopeLock.h"
 
 DEFINE_LOG_CATEGORY(LogAppleController);
 
