Unshelved from pending changelist '18530545':

Release the viewcontroller holding the splashcreen storyboard when it's no longer displayer.
# fyi wei.liu

diff --git a/Engine/Source/Runtime/ApplicationCore/Private/IOS/IOSAppDelegate.cpp b/Engine/Source/Runtime/ApplicationCore/Private/IOS/IOSAppDelegate.cpp
--- a/Engine/Source/Runtime/ApplicationCore/Private/IOS/IOSAppDelegate.cpp
+++ b/Engine/Source/Runtime/ApplicationCore/Private/IOS/IOSAppDelegate.cpp
@@ -468,6 +468,7 @@ static IOSAppDelegate* CachedDelegate = nil;
         if ([self.Window viewWithTag:200] != nil)
         {
             [[self.Window viewWithTag:200] removeFromSuperview];
+            [self.viewController release];
         }
         [timer invalidate];
     }
@@ -960,9 +961,9 @@ static FAutoConsoleVariableRef CVarGEnableThermalsReport(
     UIStoryboard *storyboard = [UIStoryboard storyboardWithName:@"LaunchScreen" bundle:nil];
     if (storyboard != nil)
     {
-        UIViewController *viewController = [storyboard instantiateViewControllerWithIdentifier:@"LaunchScreen"];
-        viewController.view.tag = 200;
-        [self.Window addSubview: viewController.view];
+        self.viewController = [storyboard instantiateViewControllerWithIdentifier:@"LaunchScreen"];
+        self.viewController.view.tag = 200;
+        [self.Window addSubview: self.viewController.view];
         GShowSplashScreen = true;
     }
 
diff --git a/Engine/Source/Runtime/ApplicationCore/Public/IOS/IOSAppDelegate.h b/Engine/Source/Runtime/ApplicationCore/Public/IOS/IOSAppDelegate.h
--- a/Engine/Source/Runtime/ApplicationCore/Public/IOS/IOSAppDelegate.h
+++ b/Engine/Source/Runtime/ApplicationCore/Public/IOS/IOSAppDelegate.h
@@ -137,6 +137,8 @@ APPLICATIONCORE_API
 
 @property bool bDeviceInPortraitMode;
 
+@property (retain) UIViewController* viewController;
+
 @property (retain) NSTimer* timer;
 
 @property (retain) NSTimer* PeakMemoryTimer;
