On Android, we automatically change rotation rate input motion depending on whether the device is in landscape mode (and which direction), but we don't actually change the other motion inputs with it, so they no longer correlate.

This change introduces the CVar Android.UnifyMotionSpace, which we can enable to have the other motion inputs changed with rotation rate so that they correlate. Enabling this CVar also changes acceleration input to be reported in g units instead of m/s^2, since that's what we're doing on most other platforms.

My hope is to add corresponding "UnifyMotionSpace" CVars to other platforms that will make motion input consistent across all platforms:
- Have acceleration axes always correlate with rotation rate axes;
- Have these axes consistent across platforms (same forward direction, same up/down direction, etc);
- Use radians per second for rotation rate and g units for acceleration;

diff --git a/Engine/Source/Runtime/ApplicationCore/Private/Android/AndroidInputInterface.cpp b/Engine/Source/Runtime/ApplicationCore/Private/Android/AndroidInputInterface.cpp
--- a/Engine/Source/Runtime/ApplicationCore/Private/Android/AndroidInputInterface.cpp
+++ b/Engine/Source/Runtime/ApplicationCore/Private/Android/AndroidInputInterface.cpp
@@ -54,6 +54,13 @@ static FAutoConsoleVariableRef CVarAndroidOldXBoxWirelessFirmware(
 	TEXT("Determines how XBox Wireless controller mapping is handled. 0 assumes new firmware, 1 will use old firmware mapping (Default: 0)"),
 	ECVF_Default);
 
+bool AndroidUnifyMotionSpace = false;
+static FAutoConsoleVariableRef CVarAndroidUnifyMotionSpace(
+	TEXT("Android.UnifyMotionSpace"),
+	UnifyMotionSpace,
+	TEXT("Motion inputs will be processed to keep rotation rate right-handed and keep other motion inputs in the same space as rotation rate. This also forces rotation rate units to be radians/s and acceleration units to be g. Default: false"),
+	ECVF_Default);
+
 TSharedRef< FAndroidInputInterface > FAndroidInputInterface::Create(const TSharedRef< FGenericApplicationMessageHandler >& InMessageHandler, const TSharedPtr< ICursor >& InCursor)
 {
 	return MakeShareable(new FAndroidInputInterface(InMessageHandler, InCursor));
@@ -1650,20 +1657,60 @@ void FAndroidInputInterface::QueueMotionData(const FVector& Tilt, const FVector&
 	EDeviceScreenOrientation ScreenOrientation = FPlatformMisc::GetDeviceOrientation();
 	FVector TempRotationRate = RotationRate;
 
-	switch (ScreenOrientation)
+	if (AndroidUnifyMotionSpace)
 	{
-		// the x tilt is inverted in LandscapeLeft.
-	case EDeviceScreenOrientation::LandscapeLeft:
-		TempRotationRate.X *= -1.0f;
-		break;
-		// the y tilt is inverted in LandscapeRight.
-	case EDeviceScreenOrientation::LandscapeRight:
-		TempRotationRate.Y *= -1.0f;
-		break;
+		FVector TempTilt = Tilt;
+		FVector TempGravity = Gravity;
+		FVector TempAcceleration = Acceleration;
+
+		auto ReorientLandscapeLeft = [](FVector InValue)
+		{
+			return FVector(-InValue.X, InValue.Y, -InValue.Z);
+		};
+
+		auto ReorientLandscapeRight = [](FVector InValue)
+		{
+			return FVector(InValue.X, -InValue.Y, -InValue.Z);
+		};
+
+		switch (ScreenOrientation)
+		{
+			// the x tilt is inverted in LandscapeLeft.
+		case EDeviceScreenOrientation::LandscapeLeft:
+			TempTilt = ReorientLandscapeLeft(TempTilt);
+			TempRotationRate = ReorientLandscapeLeft(TempRotationRate);
+			TempGravity = ReorientLandscapeLeft(TempGravity);
+			TempAcceleration = ReorientLandscapeLeft(TempAcceleration) * (1.f / 9.8f);
+			break;
+			// the y tilt is inverted in LandscapeRight.
+		case EDeviceScreenOrientation::LandscapeRight:
+			TempTilt = ReorientLandscapeRight(TempTilt);
+			TempRotationRate = ReorientLandscapeRight(TempRotationRate);
+			TempGravity = ReorientLandscapeRight(TempGravity);
+			TempAcceleration = ReorientLandscapeRight(TempAcceleration) * (1.f / 9.8f);
+			break;
+		}
+
+		FAndroidInputInterface::MotionDataStack.Push(
+			MotionData{ TempTilt, TempRotationRate, TempGravity, TempAcceleration });
 	}
+	else
+	{
+		switch (ScreenOrientation)
+		{
+			// the x tilt is inverted in LandscapeLeft.
+		case EDeviceScreenOrientation::LandscapeLeft:
+			TempRotationRate.X *= -1.0f;
+			break;
+			// the y tilt is inverted in LandscapeRight.
+		case EDeviceScreenOrientation::LandscapeRight:
+			TempRotationRate.Y *= -1.0f;
+			break;
+		}
 
-	FAndroidInputInterface::MotionDataStack.Push(
-		MotionData { Tilt, TempRotationRate, Gravity, Acceleration });
+		FAndroidInputInterface::MotionDataStack.Push(
+			MotionData{ Tilt, TempRotationRate, Gravity, Acceleration });
+	}
 }
 
 #endif
