Android.UnifyMotionSpace should apply to Portrait device orientation, too.

diff --git a/Engine/Source/Runtime/ApplicationCore/Private/Android/AndroidInputInterface.cpp b/Engine/Source/Runtime/ApplicationCore/Private/Android/AndroidInputInterface.cpp
--- a/Engine/Source/Runtime/ApplicationCore/Private/Android/AndroidInputInterface.cpp
+++ b/Engine/Source/Runtime/ApplicationCore/Private/Android/AndroidInputInterface.cpp
@@ -1673,6 +1673,11 @@ void FAndroidInputInterface::QueueMotionData(const FVector& Tilt, const FVector&
 			return FVector(InValue.X, -InValue.Y, -InValue.Z);
 		};
 
+		auto ReorientPortrait = [](FVector InValue)
+		{
+			return FVector(InValue.Y, InValue.X, -InValue.Z);
+		};
+
 		switch (ScreenOrientation)
 		{
 			// the x tilt is inverted in LandscapeLeft.
@@ -1689,6 +1694,12 @@ void FAndroidInputInterface::QueueMotionData(const FVector& Tilt, const FVector&
 			TempGravity = ReorientLandscapeRight(TempGravity);
 			TempAcceleration = ReorientLandscapeRight(TempAcceleration) * (1.f / 9.8f);
 			break;
+		case EDeviceScreenOrientation::Portrait:
+			TempTilt = ReorientPortrait(TempTilt);
+			TempRotationRate = ReorientPortrait(TempRotationRate);
+			TempGravity = ReorientPortrait(TempGravity);
+			TempAcceleration = ReorientPortrait(TempAcceleration) * (1.f / 9.8f);
+			break;
 		}
 
 		FAndroidInputInterface::MotionDataStack.Push(
