One more change to Android.UnifyMotionSpace while it's still new: Match Unreal's coordinate system. This means left-handed, left-screw, Z-up, etc.

diff --git a/Engine/Source/Runtime/ApplicationCore/Private/Android/AndroidInputInterface.cpp b/Engine/Source/Runtime/ApplicationCore/Private/Android/AndroidInputInterface.cpp
--- a/Engine/Source/Runtime/ApplicationCore/Private/Android/AndroidInputInterface.cpp
+++ b/Engine/Source/Runtime/ApplicationCore/Private/Android/AndroidInputInterface.cpp
@@ -1665,38 +1665,38 @@ void FAndroidInputInterface::QueueMotionData(const FVector& Tilt, const FVector&
 
 		auto ReorientLandscapeLeft = [](FVector InValue)
 		{
-			return FVector(-InValue.X, InValue.Y, -InValue.Z);
+			return FVector(-InValue.Z, -InValue.Y, InValue.X);
 		};
 
 		auto ReorientLandscapeRight = [](FVector InValue)
 		{
-			return FVector(InValue.X, -InValue.Y, -InValue.Z);
+			return FVector(-InValue.Z, InValue.Y, -InValue.X);
 		};
 
 		auto ReorientPortrait = [](FVector InValue)
 		{
-			return FVector(InValue.Y, InValue.X, -InValue.Z);
+			return FVector(-InValue.Z, -InValue.X, -InValue.Y);
 		};
 
 		switch (ScreenOrientation)
 		{
 			// the x tilt is inverted in LandscapeLeft.
 		case EDeviceScreenOrientation::LandscapeLeft:
-			TempTilt = ReorientLandscapeLeft(TempTilt);
-			TempRotationRate = ReorientLandscapeLeft(TempRotationRate);
+			TempTilt = -ReorientLandscapeLeft(TempTilt);
+			TempRotationRate = -ReorientLandscapeLeft(TempRotationRate);
 			TempGravity = ReorientLandscapeLeft(TempGravity);
 			TempAcceleration = ReorientLandscapeLeft(TempAcceleration) * (1.f / 9.8f);
 			break;
 			// the y tilt is inverted in LandscapeRight.
 		case EDeviceScreenOrientation::LandscapeRight:
-			TempTilt = ReorientLandscapeRight(TempTilt);
-			TempRotationRate = ReorientLandscapeRight(TempRotationRate);
+			TempTilt = -ReorientLandscapeRight(TempTilt);
+			TempRotationRate = -ReorientLandscapeRight(TempRotationRate);
 			TempGravity = ReorientLandscapeRight(TempGravity);
 			TempAcceleration = ReorientLandscapeRight(TempAcceleration) * (1.f / 9.8f);
 			break;
 		case EDeviceScreenOrientation::Portrait:
-			TempTilt = ReorientPortrait(TempTilt);
-			TempRotationRate = ReorientPortrait(TempRotationRate);
+			TempTilt = -ReorientPortrait(TempTilt);
+			TempRotationRate = -ReorientPortrait(TempRotationRate);
 			TempGravity = ReorientPortrait(TempGravity);
 			TempAcceleration = ReorientPortrait(TempAcceleration) * (1.f / 9.8f);
 			break;
