Tooltips can now accept input at times. So if we are a Tooltip type ignore checking the AcceptsInput def

diff --git a/Engine/Source/Runtime/ApplicationCore/Private/Linux/LinuxWindow.cpp b/Engine/Source/Runtime/ApplicationCore/Private/Linux/LinuxWindow.cpp
--- a/Engine/Source/Runtime/ApplicationCore/Private/Linux/LinuxWindow.cpp
+++ b/Engine/Source/Runtime/ApplicationCore/Private/Linux/LinuxWindow.cpp
@@ -133,10 +133,10 @@ void FLinuxWindow::Initialize( FLinuxApplication* const Application, const TShar
 
 	// This is a tool tip window.
 	if (!InParent.IsValid() && !Definition->HasOSWindowBorder &&
-		!Definition->AcceptsInput && Definition->IsTopmostWindow && 
-		!Definition->AppearsInTaskbar && !Definition->HasSizingFrame &&
-		!Definition->IsModalWindow && !Definition->IsRegularWindow &&
-		Definition->SizeWillChangeOften)
+		(Definition->Type == EWindowType::ToolTip || !Definition->AcceptsInput) && // tooltips can now be interactive which means they can accept input at times
+		Definition->IsTopmostWindow && !Definition->AppearsInTaskbar &&
+		!Definition->HasSizingFrame && !Definition->IsModalWindow &&
+		!Definition->IsRegularWindow && Definition->SizeWillChangeOften)
 	{
 		WindowStyle |= SDL_WINDOW_TOOLTIP;
 		bIsTooltipWindow = true;
