Portrait to match landscape axes on Android

diff --git a/Engine/Source/Runtime/ApplicationCore/Private/Android/AndroidInputInterface.cpp b/Engine/Source/Runtime/ApplicationCore/Private/Android/AndroidInputInterface.cpp
--- a/Engine/Source/Runtime/ApplicationCore/Private/Android/AndroidInputInterface.cpp
+++ b/Engine/Source/Runtime/ApplicationCore/Private/Android/AndroidInputInterface.cpp
@@ -1675,30 +1675,32 @@ void FAndroidInputInterface::QueueMotionData(const FVector& Tilt, const FVector&
 
 		auto ReorientPortrait = [](FVector InValue)
 		{
-			return FVector(-InValue.Z, -InValue.X, -InValue.Y);
+			return FVector(-InValue.Z, InValue.X, InValue.Y);
 		};
 
+		const float ToG = 1.f / 9.8f;
+
 		switch (ScreenOrientation)
 		{
 			// the x tilt is inverted in LandscapeLeft.
 		case EDeviceScreenOrientation::LandscapeLeft:
 			TempTilt = -ReorientLandscapeLeft(TempTilt);
 			TempRotationRate = -ReorientLandscapeLeft(TempRotationRate);
-			TempGravity = ReorientLandscapeLeft(TempGravity);
-			TempAcceleration = ReorientLandscapeLeft(TempAcceleration) * (1.f / 9.8f);
+			TempGravity = ReorientLandscapeLeft(TempGravity) * ToG;
+			TempAcceleration = ReorientLandscapeLeft(TempAcceleration) * ToG;
 			break;
 			// the y tilt is inverted in LandscapeRight.
 		case EDeviceScreenOrientation::LandscapeRight:
 			TempTilt = -ReorientLandscapeRight(TempTilt);
 			TempRotationRate = -ReorientLandscapeRight(TempRotationRate);
-			TempGravity = ReorientLandscapeRight(TempGravity);
-			TempAcceleration = ReorientLandscapeRight(TempAcceleration) * (1.f / 9.8f);
+			TempGravity = ReorientLandscapeRight(TempGravity) * ToG;
+			TempAcceleration = ReorientLandscapeRight(TempAcceleration) * ToG;
 			break;
 		case EDeviceScreenOrientation::Portrait:
 			TempTilt = -ReorientPortrait(TempTilt);
 			TempRotationRate = -ReorientPortrait(TempRotationRate);
-			TempGravity = ReorientPortrait(TempGravity);
-			TempAcceleration = ReorientPortrait(TempAcceleration) * (1.f / 9.8f);
+			TempGravity = ReorientPortrait(TempGravity) * ToG;
+			TempAcceleration = ReorientPortrait(TempAcceleration) * ToG;
 			break;
 		}
 
