Fix FLinuxWindow::IsPointInWindow to include X=0, Y=0 points

diff --git a/Engine/Source/Runtime/ApplicationCore/Private/Linux/LinuxWindow.cpp b/Engine/Source/Runtime/ApplicationCore/Private/Linux/LinuxWindow.cpp
--- a/Engine/Source/Runtime/ApplicationCore/Private/Linux/LinuxWindow.cpp
+++ b/Engine/Source/Runtime/ApplicationCore/Private/Linux/LinuxWindow.cpp
@@ -761,7 +761,7 @@ bool FLinuxWindow::IsPointInWindow( int32 X, int32 Y ) const
 
 	SDL_GetWindowSize( HWnd, &width, &height );
 	
-	return X > 0 && Y > 0 && X < width && Y < height;
+	return X >= 0 && Y >= 0 && X < width && Y < height;
 }
 
 int32 FLinuxWindow::GetWindowBorderSize() const
