Menus are not main windows (Mac)

I've added a boolean AllowMainWindow to CocoaWindow which is now set to true unless the window is a menu.  This status is reported to the OS via the canBecomeMainWindow method.  It is also used by the MacApplication which will now not activate a window when a window that's being closed is not a main window (via AllowMainWindow).

diff --git a/Engine/Source/Runtime/ApplicationCore/Private/Mac/CocoaWindow.cpp b/Engine/Source/Runtime/ApplicationCore/Private/Mac/CocoaWindow.cpp
--- a/Engine/Source/Runtime/ApplicationCore/Private/Mac/CocoaWindow.cpp
+++ b/Engine/Source/Runtime/ApplicationCore/Private/Mac/CocoaWindow.cpp
@@ -128,7 +128,7 @@ NSString* NSPerformDragOperation = @"NSPerformDragOperation";
 - (BOOL)canBecomeMainWindow
 {
 	SCOPED_AUTORELEASE_POOL;
-	return bAcceptsInput && ![self ignoresMouseEvents];
+	return bAcceptsInput && ![self ignoresMouseEvents] && self.AllowMainWindow;
 }
 
 - (BOOL)canBecomeKeyWindow
diff --git a/Engine/Source/Runtime/ApplicationCore/Private/Mac/MacApplication.cpp b/Engine/Source/Runtime/ApplicationCore/Private/Mac/MacApplication.cpp
--- a/Engine/Source/Runtime/ApplicationCore/Private/Mac/MacApplication.cpp
+++ b/Engine/Source/Runtime/ApplicationCore/Private/Mac/MacApplication.cpp
@@ -1254,6 +1254,8 @@ bool FMacApplication::OnWindowDestroyed(TSharedRef<FMacWindow> DestroyedWindow)
 	FCocoaWindow* WindowHandle = DestroyedWindow->GetWindowHandle();
 	const bool bDestroyingMainWindow = DestroyedWindow == ActiveWindow;
 
+	const bool bAllowMainWindow = WindowHandle.AllowMainWindow;
+	
 	if (bDestroyingMainWindow)
 	{
 		OnWindowActivationChanged(DestroyedWindow, EWindowActivation::Deactivate);
@@ -1296,10 +1298,10 @@ bool FMacApplication::OnWindowDestroyed(TSharedRef<FMacWindow> DestroyedWindow)
 		}
 	}
 
-	if (WindowToActivate.IsValid())
+	if (WindowToActivate.IsValid() && bAllowMainWindow)
 	{
 		WindowToActivate->SetWindowFocus();
-			}
+	}
 
 	MessageHandler->OnCursorSet();
 
diff --git a/Engine/Source/Runtime/ApplicationCore/Private/Mac/MacWindow.cpp b/Engine/Source/Runtime/ApplicationCore/Private/Mac/MacWindow.cpp
--- a/Engine/Source/Runtime/ApplicationCore/Private/Mac/MacWindow.cpp
+++ b/Engine/Source/Runtime/ApplicationCore/Private/Mac/MacWindow.cpp
@@ -85,6 +85,7 @@ void FMacWindow::Initialize( FMacApplication* const Application, const TSharedRe
 	MainThreadCall(^{
 		SCOPED_AUTORELEASE_POOL;
 		WindowHandle = [[FCocoaWindow alloc] initWithContentRect: ViewRect styleMask: WindowStyle backing: NSBackingStoreBuffered defer: NO];
+		WindowHandle.AllowMainWindow = ( Definition->Type != EWindowType::Menu);
 		
 		if( WindowHandle != nullptr )
 		{
diff --git a/Engine/Source/Runtime/ApplicationCore/Public/Mac/CocoaWindow.h b/Engine/Source/Runtime/ApplicationCore/Public/Mac/CocoaWindow.h
--- a/Engine/Source/Runtime/ApplicationCore/Public/Mac/CocoaWindow.h
+++ b/Engine/Source/Runtime/ApplicationCore/Public/Mac/CocoaWindow.h
@@ -26,6 +26,7 @@
 }
 
 @property (assign) EWindowMode::Type TargetWindowMode;
+@property (assign) BOOL AllowMainWindow;
 
 #if WITH_ACCESSIBILITY
 /** Adds the accessible element representing this window to the content view and exposes the window to accessibility interactions */
