Fix UnrealLightmass spewing UE_LOGs to console for Linux

Add bIsWindowShown option to Linux Console Output Device
Default to true to match old behavior.
Add -nostdout by default in UnrealLightmass

diff --git a/Engine/Source/Programs/UnrealLightmass/Private/Launch/UnrealLightmass.cpp b/Engine/Source/Programs/UnrealLightmass/Private/Launch/UnrealLightmass.cpp
--- a/Engine/Source/Programs/UnrealLightmass/Private/Launch/UnrealLightmass.cpp
+++ b/Engine/Source/Programs/UnrealLightmass/Private/Launch/UnrealLightmass.cpp
@@ -84,6 +84,12 @@ int LightmassMain(int argc, ANSICHAR* argv[])
 		CommandLine += TEXT(" -Messaging");
 	}
 
+	// Don't spew all log messages to console (unless user has specified "-stdout" on command line)
+	if (!FParse::Param(*CommandLine, TEXT("-stdout")))
+	{
+		CommandLine += TEXT(" -nostdout");
+	}
+
 	GEngineLoop.PreInit(*CommandLine);
 
 	// Tell the module manager is may now process newly-loaded UObjects when new C++ modules are loaded
diff --git a/Engine/Source/Programs/UnrealLightmass/Private/LightmassCore/LMCore.cpp b/Engine/Source/Programs/UnrealLightmass/Private/LightmassCore/LMCore.cpp
--- a/Engine/Source/Programs/UnrealLightmass/Private/LightmassCore/LMCore.cpp
+++ b/Engine/Source/Programs/UnrealLightmass/Private/LightmassCore/LMCore.cpp
@@ -81,12 +81,16 @@ void FLightmassLog::Serialize( const TCHAR* V, ELogVerbosity::Type Verbosity, co
 #endif
 	fflush( stdout );
 
+	// FUnixPlatformMisc::LowLevelOutputDebugString does a "fprintf(stderr..."
+	//  Don't print these messages twice when running under debugger on Linux.
+#if !PLATFORM_LINUX
 	if( FPlatformMisc::IsDebuggerPresent() )
 	{
 		FPlatformMisc::LowLevelOutputDebugString( V );
 		FPlatformMisc::LowLevelOutputDebugString( TEXT("\n") );
 		fflush( stderr );
 	}
+#endif
 }
 
 
diff --git a/Engine/Source/Runtime/ApplicationCore/Private/Linux/LinuxConsoleOutputDevice.cpp b/Engine/Source/Runtime/ApplicationCore/Private/Linux/LinuxConsoleOutputDevice.cpp
--- a/Engine/Source/Runtime/ApplicationCore/Private/Linux/LinuxConsoleOutputDevice.cpp
+++ b/Engine/Source/Runtime/ApplicationCore/Private/Linux/LinuxConsoleOutputDevice.cpp
@@ -20,8 +20,19 @@
 
 FLinuxConsoleOutputDevice::FLinuxConsoleOutputDevice()
 	: bOverrideColorSet(false),
-	  bOutputtingToTerminal(isatty(STDOUT_FILENO))
+	  bOutputtingToTerminal(isatty(STDOUT_FILENO)),
+	  bIsWindowShown(true)
 {
+	FString CommandLine = FCommandLine::Get();
+
+	// If -nostdout is specified and not -stdout, default to not spewing log messages.
+	// This is useful on apps like UnrealLightmass so we don't overwhelm console output
+	//  with duplicate entries and all the UE_LOG messages.
+	if (FParse::Param(*CommandLine, TEXT("nostdout")) &&
+		!FParse::Param(*CommandLine, TEXT("stdout")))
+	{
+		bIsWindowShown = false;
+	}
 }
 
 FLinuxConsoleOutputDevice::~FLinuxConsoleOutputDevice()
@@ -30,15 +41,19 @@ FLinuxConsoleOutputDevice::~FLinuxConsoleOutputDevice()
 
 void FLinuxConsoleOutputDevice::Show(bool bShowWindow)
 {
+	bIsWindowShown = bShowWindow;
 }
 
 bool FLinuxConsoleOutputDevice::IsShown()
 {
-	return true;
+	return bIsWindowShown;
 }
 
 void FLinuxConsoleOutputDevice::Serialize(const TCHAR* Data, ELogVerbosity::Type Verbosity, const class FName& Category)
 {
+	if (!bIsWindowShown)
+		return;
+
 	static bool bEntry=false;
 	if (!GIsCriticalError || bEntry)
 	{
diff --git a/Engine/Source/Runtime/ApplicationCore/Public/Linux/LinuxConsoleOutputDevice.h b/Engine/Source/Runtime/ApplicationCore/Public/Linux/LinuxConsoleOutputDevice.h
--- a/Engine/Source/Runtime/ApplicationCore/Public/Linux/LinuxConsoleOutputDevice.h
+++ b/Engine/Source/Runtime/ApplicationCore/Public/Linux/LinuxConsoleOutputDevice.h
@@ -19,6 +19,9 @@ private:
 	/** true if we are outputting to a terminal */
 	bool bOutputtingToTerminal;
 
+	/** true if window is visible */
+	bool bIsWindowShown;
+
 public:
 
 	/**
