Fix taskbar item for Splash Screen.

Or-ing in WS_EX_TOOLWINDOW into the window style prevents the splash screen's taskbar item from appearing until its visible on screen.

Since it takes several seconds before the splash screen appears, if you Alt-Tab to another application before it appears it causes the splash screen to not be shown, and since it's not shown there is no task bar icon, so you cannot see the progress bar or select the splash screen via the taskbar.

Also, the splash screen appears to go to the bottom of the z-order, causing the user to have to minimize/close a lot of windows to click it to allocate the taskbar item.

#rnx

diff --git a/Engine/Source/Runtime/ApplicationCore/Private/Windows/WindowsPlatformSplash.cpp b/Engine/Source/Runtime/ApplicationCore/Private/Windows/WindowsPlatformSplash.cpp
--- a/Engine/Source/Runtime/ApplicationCore/Private/Windows/WindowsPlatformSplash.cpp
+++ b/Engine/Source/Runtime/ApplicationCore/Private/Windows/WindowsPlatformSplash.cpp
@@ -416,7 +416,7 @@ uint32 WINAPI StartSplashScreenThread( LPVOID unused )
 		const bool bAllowFading = false;
 
 		// Force the editor splash screen to show up in the taskbar and alt-tab lists
-		uint32 dwWindowStyle = (GIsEditor ? WS_EX_APPWINDOW : 0) | WS_EX_TOOLWINDOW;
+		uint32 dwWindowStyle = GIsEditor ? WS_EX_APPWINDOW : WS_EX_TOOLWINDOW;
 		if( bAllowFading )
 		{
 			dwWindowStyle |= WS_EX_LAYERED;
