Fix menu's flashing white when application scale is not default (100%).

diff --git a/Engine/Source/Runtime/ApplicationCore/Private/Windows/WindowsWindow.cpp b/Engine/Source/Runtime/ApplicationCore/Private/Windows/WindowsWindow.cpp
--- a/Engine/Source/Runtime/ApplicationCore/Private/Windows/WindowsWindow.cpp
+++ b/Engine/Source/Runtime/ApplicationCore/Private/Windows/WindowsWindow.cpp
@@ -540,8 +540,11 @@ void FWindowsWindow::ReshapeWindow( int32 NewX, int32 NewY, int32 NewWidth, int3
 	// We use SWP_NOSENDCHANGING when in fullscreen mode to prevent Windows limiting our window size to the current resolution, as that 
 	// prevents us being able to change to a higher resolution while in fullscreen mode
 	::SetWindowPos( HWnd, nullptr, WindowX, WindowY, NewWidth, NewHeight, SWP_NOZORDER | SWP_NOACTIVATE | ((WindowMode == EWindowMode::Fullscreen) ? SWP_NOSENDCHANGING : 0) );
+	
+	bool bAdjustSizeChange = Definition->SizeWillChangeOften || bVirtualSizeChanged;
+	bool bAdjustCorners = Definition->Type != EWindowType::Menu && Definition->CornerRadius > 0;
 
-	if( !Definition->HasOSWindowBorder && (((Definition->SizeWillChangeOften && bVirtualSizeChanged) || Definition->CornerRadius > 0)) )
+	if (!Definition->HasOSWindowBorder && (bAdjustSizeChange || bAdjustCorners))
 	{
 		AdjustWindowRegion( VirtualWidth, VirtualHeight );
 	}
