allow setting the console title of spawned output consoles on Windows

#rnx

diff --git a/Engine/Source/Runtime/ApplicationCore/Private/Windows/WindowsConsoleOutputDevice.cpp b/Engine/Source/Runtime/ApplicationCore/Private/Windows/WindowsConsoleOutputDevice.cpp
--- a/Engine/Source/Runtime/ApplicationCore/Private/Windows/WindowsConsoleOutputDevice.cpp
+++ b/Engine/Source/Runtime/ApplicationCore/Private/Windows/WindowsConsoleOutputDevice.cpp
@@ -137,6 +137,12 @@ void FWindowsConsoleOutputDevice::Show( bool ShowWindow )
 					ConsolePosY = WindowRect.top;
 				}
 
+				FString ConsoleTitle;
+				if (FParse::Value(FCommandLine::Get(), TEXT("ConsoleTitle="), ConsoleTitle))
+				{
+					SetConsoleTitleW(*ConsoleTitle);
+				}
+
 				FDisplayMetrics DisplayMetrics;
 				FDisplayMetrics::RebuildDisplayMetrics(DisplayMetrics);
 
