Linux windows now have a reasonable minimum size, and no longer can be scaled down to 1 pixel wide/tall

diff --git a/Engine/Source/Runtime/ApplicationCore/Private/Linux/LinuxWindow.cpp b/Engine/Source/Runtime/ApplicationCore/Private/Linux/LinuxWindow.cpp
--- a/Engine/Source/Runtime/ApplicationCore/Private/Linux/LinuxWindow.cpp
+++ b/Engine/Source/Runtime/ApplicationCore/Private/Linux/LinuxWindow.cpp
@@ -22,6 +22,12 @@ DEFINE_LOG_CATEGORY( LogLinuxWindowEvent );
 // might need to be a function in case SDL gets overriden at runtime
 #define UE_USING_BORDERS_AWARE_SDL					1
 
+namespace 
+{
+	constexpr int DefaultMinWindowWidth = 100;
+	constexpr int DefaultMinWindowHeight = 50;
+}
+
 FLinuxWindow::~FLinuxWindow()
 {
 	// NOTE: The HWnd is invalid here!
@@ -259,6 +265,26 @@ void FLinuxWindow::Initialize( FLinuxApplication* const Application, const TShar
 		return;
 	}
 
+	// SDL_SetWindow[Minimum/Maximum]Size() requires a valid value for BOTH width and height, we can't do one or the other without magic numbers
+	const bool bIsMinSizeSet = Definition->SizeLimits.GetMinWidth().IsSet() && Definition->SizeLimits.GetMinHeight().IsSet();
+	const bool bIsMaxSizeSet = Definition->SizeLimits.GetMaxWidth().IsSet() && Definition->SizeLimits.GetMaxHeight().IsSet();
+
+	if ((WindowStyle & SDL_WINDOW_RESIZABLE) && !bIsMinSizeSet)
+	{
+		// Set a reasonable minimum size if none are set
+		// Vulkan will fail to create a swap chain if the back buffer is 0x0
+		SDL_SetWindowMinimumSize(HWnd, DefaultMinWindowWidth, DefaultMinWindowHeight);
+	}
+	else if (bIsMinSizeSet)
+	{
+		SDL_SetWindowMinimumSize(HWnd, *Definition->SizeLimits.GetMinWidth(), *Definition->SizeLimits.GetMinHeight());
+	}
+
+	if (bIsMaxSizeSet)
+	{
+		SDL_SetWindowMaximumSize(HWnd, *Definition->SizeLimits.GetMaxWidth(), *Definition->SizeLimits.GetMaxHeight());
+	}
+
 	if (Definition->AppearsInTaskbar)
 	{
 		// Try to find an icon for the window
