Some fixes for pathing with Gauntlet and EngineTest
Fix some test/asserts when running offscreen rendering. We use a GenericWindow which was hitting some fatal logs

diff --git a/Engine/Source/Programs/AutomationTool/Gauntlet/Unreal/Base/Gauntlet.UnrealTestConfiguration.cs b/Engine/Source/Programs/AutomationTool/Gauntlet/Unreal/Base/Gauntlet.UnrealTestConfiguration.cs
--- a/Engine/Source/Programs/AutomationTool/Gauntlet/Unreal/Base/Gauntlet.UnrealTestConfiguration.cs
+++ b/Engine/Source/Programs/AutomationTool/Gauntlet/Unreal/Base/Gauntlet.UnrealTestConfiguration.cs
@@ -874,9 +874,11 @@ namespace Gauntlet
 
 				// if we are unattended but still may need access to Vulkan passing renderoffscreen to allow not depending on
 				// the X11/Wayland display server to be around and use a dummy/offscreen rendering mode
+				//
+				// As well as disable sound as there are no audio devices when running through horde
 				if (AppConfig.Platform == UnrealTargetPlatform.Linux)
 				{
-					AppConfig.CommandLine += " -renderoffscreen";
+					AppConfig.CommandLine += " -renderoffscreen -nosound";
 				}
 			}
 
diff --git a/Engine/Source/Programs/AutomationTool/Gauntlet/Unreal/Config/Gauntlet.EpicConfig.cs b/Engine/Source/Programs/AutomationTool/Gauntlet/Unreal/Config/Gauntlet.EpicConfig.cs
--- a/Engine/Source/Programs/AutomationTool/Gauntlet/Unreal/Config/Gauntlet.EpicConfig.cs
+++ b/Engine/Source/Programs/AutomationTool/Gauntlet/Unreal/Config/Gauntlet.EpicConfig.cs
@@ -25,6 +25,10 @@ namespace EpicConfig
 				{
 					ConfigFilePath = @"\\epicgames.net\root\Builds\Automation\UE5\Config\ue-automation-telemetry-mysql-config.txt";
 				}
+				else if (HostPlatform.Current.HostEditorPlatform == UnrealTargetPlatform.Linux)
+				{
+					ConfigFilePath = "/mnt/Builds/Automation/UE5/Config/ue-automation-telemetry-mysql-config.txt";
+				}
 				else
 				{
 					ConfigFilePath = "/Volumes/Root/Builds/Automation/UE5/Config/ue-automation-telemetry-mysql-config.txt";
@@ -125,4 +129,4 @@ namespace EpicConfig
 		protected override string TableName { get { return "test_records_staging"; } }
 		protected override string BuildTableName { get { return "test_builds_staging"; } }
 	}
-}
\ No newline at end of file
+}
diff --git a/Engine/Source/Runtime/ApplicationCore/Private/GenericPlatform/GenericWindow.cpp b/Engine/Source/Runtime/ApplicationCore/Private/GenericPlatform/GenericWindow.cpp
--- a/Engine/Source/Runtime/ApplicationCore/Private/GenericPlatform/GenericWindow.cpp
+++ b/Engine/Source/Runtime/ApplicationCore/Private/GenericPlatform/GenericWindow.cpp
@@ -21,8 +21,11 @@ void FGenericWindow::ReshapeWindow( int32 X, int32 Y, int32 Width, int32 Height
 
 bool FGenericWindow::GetFullScreenInfo( int32& X, int32& Y, int32& Width, int32& Height ) const
 {
-	// this function cannot return valid results, and should not be needed on consoles, etc
-	UE_LOG(LogGenericPlatformWindow, Fatal, TEXT("GetFullScreenInfo is not expected to be called on this platform"));
+	if (!FParse::Param(FCommandLine::Get(), TEXT("RenderOffScreen")))
+	{
+		// this function cannot return valid results, and should not be needed on consoles, etc
+		UE_LOG(LogGenericPlatformWindow, Fatal, TEXT("GetFullScreenInfo is not expected to be called on this platform"));
+	}
 	return false;
 }
 
@@ -102,8 +105,11 @@ bool FGenericWindow::IsVisible() const
 
 bool FGenericWindow::GetRestoredDimensions(int32& X, int32& Y, int32& Width, int32& Height)
 {
-	// this function cannot return valid results, and should not be needed on consoles, etc
-	UE_LOG(LogGenericPlatformWindow, Fatal, TEXT("GetRestoredDimensions is not expected to be called on this platform"));
+	if (!FParse::Param(FCommandLine::Get(), TEXT("RenderOffScreen")))
+	{
+		// this function cannot return valid results, and should not be needed on consoles, etc
+		UE_LOG(LogGenericPlatformWindow, Fatal, TEXT("GetRestoredDimensions is not expected to be called on this platform"));
+	}
 	return false;
 }
 
diff --git a/Engine/Source/Runtime/VulkanRHI/Private/Linux/VulkanLinuxPlatform.cpp b/Engine/Source/Runtime/VulkanRHI/Private/Linux/VulkanLinuxPlatform.cpp
--- a/Engine/Source/Runtime/VulkanRHI/Private/Linux/VulkanLinuxPlatform.cpp
+++ b/Engine/Source/Runtime/VulkanRHI/Private/Linux/VulkanLinuxPlatform.cpp
@@ -179,7 +179,8 @@ void FVulkanLinuxPlatform::GetInstanceExtensions(TArray<const ANSICHAR*>& OutExt
 	{
 		OutExtensions.Add("VK_KHR_wayland_surface");
 	}
-	else
+	// dummy is when we render offscreen, so ignore warning here
+	else if (strcmp(SDLDriver, "dummy") != 0)
 	{
 		UE_LOG(LogRHI, Warning, TEXT("Could not detect SDL video driver!"));
 	}
@@ -277,4 +278,4 @@ void FVulkanLinuxPlatform::CheckDeviceDriver(uint32 DeviceIndex, EGpuVendorId Ve
 			GRHIAllow64bitShaderAtomicsCvar->SetWithCurrentPriority(0);
 		}
 	}
-}
\ No newline at end of file
+}
