Fix Non-Unity issues from missing header

diff --git a/Engine/Source/Runtime/ApplicationCore/Private/GenericPlatform/GenericWindow.cpp b/Engine/Source/Runtime/ApplicationCore/Private/GenericPlatform/GenericWindow.cpp
--- a/Engine/Source/Runtime/ApplicationCore/Private/GenericPlatform/GenericWindow.cpp
+++ b/Engine/Source/Runtime/ApplicationCore/Private/GenericPlatform/GenericWindow.cpp
@@ -1,6 +1,7 @@
 // Copyright Epic Games, Inc. All Rights Reserved.
 
 #include "GenericPlatform/GenericWindow.h"
+#include "Misc/CommandLine.h"
 
 DEFINE_LOG_CATEGORY_STATIC(LogGenericPlatformWindow, Log, All);
 
