Fix NonUnity issue

This needs to be refactored out to a DummyWindow for Offscreen rendering vs doing this

diff --git a/Engine/Source/Runtime/ApplicationCore/Private/GenericPlatform/GenericWindow.cpp b/Engine/Source/Runtime/ApplicationCore/Private/GenericPlatform/GenericWindow.cpp
--- a/Engine/Source/Runtime/ApplicationCore/Private/GenericPlatform/GenericWindow.cpp
+++ b/Engine/Source/Runtime/ApplicationCore/Private/GenericPlatform/GenericWindow.cpp
@@ -2,6 +2,7 @@
 
 #include "GenericPlatform/GenericWindow.h"
 #include "Misc/CommandLine.h"
+#include "Misc/Parse.h"
 
 DEFINE_LOG_CATEGORY_STATIC(LogGenericPlatformWindow, Log, All);
 
