Avoid double output when running with -stdout enabled

This does use LaunchEngineLoops output over Linux which means with -stdout no color coded warning/error messages

diff --git a/Engine/Source/Runtime/ApplicationCore/Private/Linux/LinuxConsoleOutputDevice.cpp b/Engine/Source/Runtime/ApplicationCore/Private/Linux/LinuxConsoleOutputDevice.cpp
--- a/Engine/Source/Runtime/ApplicationCore/Private/Linux/LinuxConsoleOutputDevice.cpp
+++ b/Engine/Source/Runtime/ApplicationCore/Private/Linux/LinuxConsoleOutputDevice.cpp
@@ -21,7 +21,8 @@
 FLinuxConsoleOutputDevice::FLinuxConsoleOutputDevice()
 	: bOverrideColorSet(false),
 	  bOutputtingToTerminal(isatty(STDOUT_FILENO)),
-	  bIsWindowShown(true)
+	  bIsWindowShown(true),
+	  bIsStdoutSet(false)
 {
 	FString CommandLine = FCommandLine::Get();
 
@@ -33,6 +34,9 @@ FLinuxConsoleOutputDevice::FLinuxConsoleOutputDevice()
 	{
 		bIsWindowShown = false;
 	}
+
+	// If -stdout is getting set we are doing a printf in LaunchEngineLoop.cpp, so lets avoid double printing here
+	bIsStdoutSet = FParse::Param(*CommandLine, TEXT("stdout"));
 }
 
 FLinuxConsoleOutputDevice::~FLinuxConsoleOutputDevice()
@@ -51,7 +55,7 @@ bool FLinuxConsoleOutputDevice::IsShown()
 
 void FLinuxConsoleOutputDevice::Serialize(const TCHAR* Data, ELogVerbosity::Type Verbosity, const class FName& Category)
 {
-	if (!bIsWindowShown)
+	if (!bIsWindowShown || bIsStdoutSet)
 		return;
 
 	static bool bEntry=false;
diff --git a/Engine/Source/Runtime/ApplicationCore/Public/Linux/LinuxConsoleOutputDevice.h b/Engine/Source/Runtime/ApplicationCore/Public/Linux/LinuxConsoleOutputDevice.h
--- a/Engine/Source/Runtime/ApplicationCore/Public/Linux/LinuxConsoleOutputDevice.h
+++ b/Engine/Source/Runtime/ApplicationCore/Public/Linux/LinuxConsoleOutputDevice.h
@@ -22,6 +22,9 @@ private:
 	/** true if window is visible */
 	bool bIsWindowShown;
 
+	/** true if stdout is set via the command line */
+	bool bIsStdoutSet;
+
 public:
 
 	/**
