Implement PS5 DualSense controller support on iOS/tvOS and fix crash when setting Siri Remote as Unset on tvOS

diff --git a/Engine/Source/Runtime/ApplicationCore/Private/Apple/AppleControllerInterface.cpp b/Engine/Source/Runtime/ApplicationCore/Private/Apple/AppleControllerInterface.cpp
--- a/Engine/Source/Runtime/ApplicationCore/Private/Apple/AppleControllerInterface.cpp
+++ b/Engine/Source/Runtime/ApplicationCore/Private/Apple/AppleControllerInterface.cpp
@@ -114,6 +114,11 @@ void FAppleControllerInterface::SetControllerType(uint32 ControllerIndex)
     {
         Controllers[ControllerIndex].ControllerType = ControllerType::XboxGamepad;
     }
+    else if ([Controller.productCategory isEqualToString:@"DualSense"])
+    {
+        Controllers[ControllerIndex].ControllerType = ControllerType::DualSenseGamepad;
+    }
+
     else if (Controller.extendedGamepad != nil)
     {
         Controllers[ControllerIndex].ControllerType = ControllerType::ExtendedGamepad;
@@ -447,6 +452,7 @@ bIsPressed = Gamepad.GCButton.pressed; \
         case ControllerType::ExtendedGamepad:
         case ControllerType::DualShockGamepad:
         case ControllerType::XboxGamepad:
+        case ControllerType::DualSenseGamepad:
             
             ExtendedGamepad = Cont.extendedGamepad;
             ExtendedPreviousGamepad = Controllers[ControllerIndex].PreviousExtendedGamepad;
@@ -508,6 +514,7 @@ void FAppleControllerInterface::HandleAnalogGamepad(const FGamepadKeyNames::Type
         case ControllerType::ExtendedGamepad:
         case ControllerType::DualShockGamepad:
         case ControllerType::XboxGamepad:
+        case ControllerType::DualSenseGamepad:
             
             if (UEAxis == FGamepadKeyNames::LeftAnalogX){if ((ExtendedPreviousGamepad != nil && ExtendedGamepad.leftThumbstick.xAxis.value != ExtendedPreviousGamepad.leftThumbstick.xAxis.value) ||
                                                              (ExtendedGamepad.leftThumbstick.xAxis.value < -RepeatDeadzone || ExtendedGamepad.leftThumbstick.xAxis.value > RepeatDeadzone)){axisValue = ExtendedGamepad.leftThumbstick.xAxis.value;}}
diff --git a/Engine/Source/Runtime/ApplicationCore/Private/Apple/AppleControllerInterface.h b/Engine/Source/Runtime/ApplicationCore/Private/Apple/AppleControllerInterface.h
--- a/Engine/Source/Runtime/ApplicationCore/Private/Apple/AppleControllerInterface.h
+++ b/Engine/Source/Runtime/ApplicationCore/Private/Apple/AppleControllerInterface.h
@@ -16,7 +16,8 @@ enum ControllerType
 	SiriRemote,
 	ExtendedGamepad,
 	XboxGamepad,
-	DualShockGamepad
+	DualShockGamepad,
+    DualSenseGamepad
 };
 
 enum PlayerIndex
diff --git a/Engine/Source/Runtime/ApplicationCore/Private/IOS/IOSInputInterface.cpp b/Engine/Source/Runtime/ApplicationCore/Private/IOS/IOSInputInterface.cpp
--- a/Engine/Source/Runtime/ApplicationCore/Private/IOS/IOSInputInterface.cpp
+++ b/Engine/Source/Runtime/ApplicationCore/Private/IOS/IOSInputInterface.cpp
@@ -532,7 +532,7 @@ void FIOSInputInterface::SendControllerEvents()
         {
             GCController* Cont = Controller.Controller;
             GCMicroGamepad* MicroGamepad = [Cont capture].microGamepad;
-            if (MicroGamepad != nil)
+            if (MicroGamepad != nil && Controller.PlayerIndex != PlayerIndex::PlayerUnset)
             {
                 const GCMicroGamepad* PreviousMicroGamepad = Controller.PreviousMicroGamepad;
 
