Fix window activation when menus close

CocoaWindows now store their Type.  This is used so that we can make sure that we can call OnWindowActivationChanged
when the menu closes.  This is needed to restore focus in the main window after the menu is triggered.  Several other
utility windows are now not treated as main windows.

#rnx

diff --git a/Engine/Source/Runtime/ApplicationCore/Private/Mac/CocoaWindow.cpp b/Engine/Source/Runtime/ApplicationCore/Private/Mac/CocoaWindow.cpp
--- a/Engine/Source/Runtime/ApplicationCore/Private/Mac/CocoaWindow.cpp
+++ b/Engine/Source/Runtime/ApplicationCore/Private/Mac/CocoaWindow.cpp
@@ -579,4 +579,18 @@ NSString* NSPerformDragOperation = @"NSPerformDragOperation";
 }
 #endif
 
+- (BOOL)AllowMainWindow
+{
+	switch (self.Type)
+	{
+		case EWindowType::Menu:
+		case EWindowType::ToolTip:
+		case EWindowType::Notification:
+		case EWindowType::CursorDecorator:
+			return false;
+		default:
+			return true;
+	}
+}
+
 @end
diff --git a/Engine/Source/Runtime/ApplicationCore/Private/Mac/MacApplication.cpp b/Engine/Source/Runtime/ApplicationCore/Private/Mac/MacApplication.cpp
--- a/Engine/Source/Runtime/ApplicationCore/Private/Mac/MacApplication.cpp
+++ b/Engine/Source/Runtime/ApplicationCore/Private/Mac/MacApplication.cpp
@@ -1255,7 +1255,8 @@ bool FMacApplication::OnWindowDestroyed(TSharedRef<FMacWindow> DestroyedWindow)
 	const bool bDestroyingMainWindow = DestroyedWindow == ActiveWindow;
 
 	const bool bAllowMainWindow = WindowHandle.AllowMainWindow;
-	
+	const EWindowType WindowType =WindowHandle.Type;
+    
 	if (bDestroyingMainWindow)
 	{
 		OnWindowActivationChanged(DestroyedWindow, EWindowActivation::Deactivate);
@@ -1298,9 +1299,43 @@ bool FMacApplication::OnWindowDestroyed(TSharedRef<FMacWindow> DestroyedWindow)
 		}
 	}
 
-	if (WindowToActivate.IsValid() && bAllowMainWindow)
+	if (WindowToActivate.IsValid())
 	{
-		WindowToActivate->SetWindowFocus();
+		if (bAllowMainWindow)
+		{
+			WindowToActivate->SetWindowFocus();
+		}
+		
+		FCocoaWindow* ActivateWindowHandle = WindowToActivate->GetWindowHandle();
+		bool bActivateAllowMainWindow = ActivateWindowHandle.AllowMainWindow;
+				
+		if (WindowType == EWindowType::Menu && bActivateAllowMainWindow)
+		{
+			// For some reason when submenus are getting closed a main window is getting activated rather then the previous menu.  I think a
+			// better solution would be to investigate the ordering of Windows[], or possibly just enumerating the NSWindows via z-order.
+			// I'm worried about the consequences of that kind of change so I'm opting for this since it seems less risky.  We count all
+			// the open menus and only activate a normal window if the menu closing is the last one.
+
+			int32 NumMenus=0;
+			for (int32 Index = 0; Index < Windows.Num(); ++Index)
+			{
+				TSharedPtr<FMacWindow> IsMenuWindowRef = Windows[Index];
+				if (IsMenuWindowRef.IsValid())
+				{
+					if (IsMenuWindowRef->GetWindowHandle().Type == EWindowType::Menu)
+					{
+						++NumMenus;
+					}
+				}
+			}
+			
+			if (NumMenus<=0)
+			{
+				GameThreadCall(^{
+					OnWindowActivationChanged(WindowToActivate.ToSharedRef(), EWindowActivation::Activate);
+				}, @[ NSDefaultRunLoopMode, UnrealResizeEventMode, UnrealShowEventMode, UnrealFullscreenEventMode, UnrealCloseEventMode ], true);
+			}
+		}
 	}
 
 	MessageHandler->OnCursorSet();
diff --git a/Engine/Source/Runtime/ApplicationCore/Private/Mac/MacWindow.cpp b/Engine/Source/Runtime/ApplicationCore/Private/Mac/MacWindow.cpp
--- a/Engine/Source/Runtime/ApplicationCore/Private/Mac/MacWindow.cpp
+++ b/Engine/Source/Runtime/ApplicationCore/Private/Mac/MacWindow.cpp
@@ -85,7 +85,7 @@ void FMacWindow::Initialize( FMacApplication* const Application, const TSharedRe
 	MainThreadCall(^{
 		SCOPED_AUTORELEASE_POOL;
 		WindowHandle = [[FCocoaWindow alloc] initWithContentRect: ViewRect styleMask: WindowStyle backing: NSBackingStoreBuffered defer: NO];
-		WindowHandle.AllowMainWindow = ( Definition->Type != EWindowType::Menu);
+		WindowHandle.Type = Definition->Type;
 		
 		if( WindowHandle != nullptr )
 		{
diff --git a/Engine/Source/Runtime/ApplicationCore/Public/Mac/CocoaWindow.h b/Engine/Source/Runtime/ApplicationCore/Public/Mac/CocoaWindow.h
--- a/Engine/Source/Runtime/ApplicationCore/Public/Mac/CocoaWindow.h
+++ b/Engine/Source/Runtime/ApplicationCore/Public/Mac/CocoaWindow.h
@@ -3,6 +3,8 @@
 #pragma once
 
 #include "GenericPlatform/GenericWindow.h"
+#include "GenericPlatform/GenericWindowDefinition.h"
+
 #if WITH_ACCESSIBILITY
 #include "GenericPlatform/Accessibility/GenericAccessibleInterfaces.h"
 #endif
@@ -26,7 +28,8 @@
 }
 
 @property (assign) EWindowMode::Type TargetWindowMode;
-@property (assign) BOOL AllowMainWindow;
+@property (assign,readonly,getter=AllowMainWindow) BOOL AllowMainWindow;
+@property (assign) EWindowType Type;
 
 #if WITH_ACCESSIBILITY
 /** Adds the accessible element representing this window to the content view and exposes the window to accessibility interactions */
