Fix Mac Menus

Sometimes the active window seems to be confused on Mac.  This works around
the issue so that we will activate another window when a menu closes.

#rnx

diff --git a/Engine/Source/Runtime/ApplicationCore/Private/Mac/MacApplication.cpp b/Engine/Source/Runtime/ApplicationCore/Private/Mac/MacApplication.cpp
--- a/Engine/Source/Runtime/ApplicationCore/Private/Mac/MacApplication.cpp
+++ b/Engine/Source/Runtime/ApplicationCore/Private/Mac/MacApplication.cpp
@@ -1282,7 +1282,7 @@ bool FMacApplication::OnWindowDestroyed(TSharedRef<FMacWindow> DestroyedWindow)
 
 	TSharedPtr<FMacWindow> WindowToActivate;
 
-	if (bDestroyingMainWindow)
+	if (bDestroyingMainWindow || DestroyedWindow->GetWindowHandle().Type == EWindowType::Menu)
 	{
 		FScopeLock Lock(&WindowsMutex);
 		// Figure out which window will now become active and let Slate know without waiting for Cocoa events.
