Fix r.MobileContentScaleFactor being ignored and always defaulting to 1 on iOS

diff --git a/Engine/Source/Runtime/ApplicationCore/Private/IOS/IOSAppDelegate.cpp b/Engine/Source/Runtime/ApplicationCore/Private/IOS/IOSAppDelegate.cpp
--- a/Engine/Source/Runtime/ApplicationCore/Private/IOS/IOSAppDelegate.cpp
+++ b/Engine/Source/Runtime/ApplicationCore/Private/IOS/IOSAppDelegate.cpp
@@ -331,15 +331,15 @@ static IOSAppDelegate* CachedDelegate = nil;
 #endif
 
 	bEngineInit = true;
-    
-    // put a render thread job to turn off the splash screen after the first render flip
-    if (GShowSplashScreen)
-    {
-        FGraphEventRef SplashTask = FFunctionGraphTask::CreateAndDispatchWhenReady([]()
-        {
-            GShowSplashScreen = false;
-        }, TStatId(), NULL, ENamedThreads::ActualRenderingThread);
-    }
+
+	// put a render thread job to turn off the splash screen after the first render flip
+	if (GShowSplashScreen)
+	{
+		FGraphEventRef SplashTask = FFunctionGraphTask::CreateAndDispatchWhenReady([]()
+		{
+			GShowSplashScreen = false;
+		}, TStatId(), NULL, ENamedThreads::ActualRenderingThread);
+	}
 
 	for (NSDictionary* openUrlParameter in self.savedOpenUrlParameters)
 	{
@@ -825,6 +825,25 @@ static IOSAppDelegate* CachedDelegate = nil;
 #endif
 }
 
+- (void)LoadMobileContentScaleFactor
+{
+	// need to cache the MobileContentScaleFactor for framebuffer creation.
+	static IConsoleVariable* CVar = IConsoleManager::Get().FindConsoleVariable(TEXT("r.MobileContentScaleFactor"));
+	self.MobileContentScaleFactor = CVar ? CVar->GetFloat() : 0;
+
+	// Can also be overridden from the commandline using "mcsf="
+	FString CmdLineCSF;
+	if (FParse::Value(FCommandLine::Get(), TEXT("mcsf="), CmdLineCSF, false))
+	{
+		self.MobileContentScaleFactor = FCString::Atof(*CmdLineCSF);
+	}
+}
+
+- (float)GetMobileContentScaleFactor
+{
+	return self.MobileContentScaleFactor;
+}
+
 - (void)CheckForZoomAccessibility
 {
 #if !PLATFORM_TVOS
diff --git a/Engine/Source/Runtime/ApplicationCore/Private/IOS/IOSView.cpp b/Engine/Source/Runtime/ApplicationCore/Private/IOS/IOSView.cpp
--- a/Engine/Source/Runtime/ApplicationCore/Private/IOS/IOSView.cpp
+++ b/Engine/Source/Runtime/ApplicationCore/Private/IOS/IOSView.cpp
@@ -272,34 +272,16 @@ id<MTLDevice> GMetalDevice = nil;
 	if (!bIsInitialized)
 	{
 		// look up what the device can support
-		const float NativeScale = [[UIScreen mainScreen] scale];
-
-		float RequestedContentScaleFactor = 1.0;
-		[IOSAppDelegate WaitAndRunOnGameThread:[&RequestedContentScaleFactor]()
-		{
-			// look up the CVar for the scale factor
-			static IConsoleVariable* CVar = IConsoleManager::Get().FindConsoleVariable(TEXT("r.MobileContentScaleFactor"));
-			RequestedContentScaleFactor = CVar->GetFloat();
-
-			FString CmdLineCSF;
-			if (FParse::Value(FCommandLine::Get(), TEXT("mcsf="), CmdLineCSF, false))
-			{
-				RequestedContentScaleFactor = FCString::Atof(*CmdLineCSF);
-			}
-		}];
-
+		const float Scale = [[UIScreen mainScreen] scale];
+		const float NativeScale = self.window.screen.nativeScale;
+		const float RequestedContentScaleFactor = [[IOSAppDelegate GetDelegate] GetMobileContentScaleFactor];
+		
+		UE_LOG(LogIOS, Log, TEXT("RequestedContentScaleFactor %f to nativeScale which is = (s:%f, ns:%f, csf:%f"), RequestedContentScaleFactor, Scale, NativeScale, self.contentScaleFactor);
 		// 0 means to leave the scale alone, use native
 		if (RequestedContentScaleFactor == 0.0f)
 		{
-            if ([self.window.screen respondsToSelector:@selector(nativeScale)])
-            {
-                self.contentScaleFactor = self.window.screen.nativeScale;
-                UE_LOG(LogIOS, Log, TEXT("Setting contentScaleFactor to nativeScale which is = %f"), self.contentScaleFactor);
-            }
-            else
-            {
-                UE_LOG(LogIOS, Log, TEXT("Leaving contentScaleFactor alone, with scale = %f"), NativeScale);
-            }
+			self.contentScaleFactor = NativeScale;
+			UE_LOG(LogIOS, Log, TEXT("Setting contentScaleFactor to nativeScale which is = %f"), self.contentScaleFactor);
 		}
 		else
 		{
@@ -308,9 +290,7 @@ id<MTLDevice> GMetalDevice = nil;
 			UE_LOG(LogIOS, Log, TEXT("Setting contentScaleFactor to %0.4f (optimal = %0.4f)"), self.contentScaleFactor, NativeScale);
 		}
 
-
-		// handle Metal or GL sizing
-#if HAS_METAL
+		// TODO: bIsUsingMetal is always true. Remove it.
 		if (bIsUsingMetal)
 		{
 			CAMetalLayer* MetalLayer = (CAMetalLayer*)self.layer;
@@ -319,10 +299,9 @@ id<MTLDevice> GMetalDevice = nil;
 			DrawableSize.height *= self.contentScaleFactor;
 			MetalLayer.drawableSize = DrawableSize;
 		}
-#endif
 
 		bIsInitialized = true;
-	}    
+	}
 	return true;
 }
 
diff --git a/Engine/Source/Runtime/ApplicationCore/Public/IOS/IOSAppDelegate.h b/Engine/Source/Runtime/ApplicationCore/Public/IOS/IOSAppDelegate.h
--- a/Engine/Source/Runtime/ApplicationCore/Public/IOS/IOSAppDelegate.h
+++ b/Engine/Source/Runtime/ApplicationCore/Public/IOS/IOSAppDelegate.h
@@ -180,6 +180,7 @@ APPLICATIONCORE_API
 @property (assign) bool bBatteryState;
 @property (assign) int BatteryLevel;
 
+@property (assign) float MobileContentScaleFactor;
 @property (assign) bool bUpdateAvailable;
 
 /**
@@ -194,6 +195,9 @@ APPLICATIONCORE_API
 +(bool)WaitAndRunOnGameThread:(TUniqueFunction<void()>)Function;
 -(void)NoUrlCommandLine;
 
+-(void)LoadMobileContentScaleFactor;
+-(float)GetMobileContentScaleFactor;
+
 -(int)GetAudioVolume;
 -(bool)AreHeadphonesPluggedIn;
 -(int)GetBatteryLevel;
diff --git a/Engine/Source/Runtime/Core/Private/IOS/IOSPlatformMisc.cpp b/Engine/Source/Runtime/Core/Private/IOS/IOSPlatformMisc.cpp
--- a/Engine/Source/Runtime/Core/Private/IOS/IOSPlatformMisc.cpp
+++ b/Engine/Source/Runtime/Core/Private/IOS/IOSPlatformMisc.cpp
@@ -73,6 +73,10 @@ static int32 GetFreeMemoryMB()
 
 void FIOSPlatformMisc::PlatformInit()
 {
+	// PlatformInit() starts the UI thread which creates the framebuffer and it requires
+	// "r.MobileContentScaleFactor" to be available before it's creation, so need to cache that value now.
+	[[IOSAppDelegate GetDelegate] LoadMobileContentScaleFactor];
+		
 	FAppEntry::PlatformInit();
 
 	// Increase the maximum number of simultaneously open files
