UIA accessibility is available for all Windows-based platforms.

#rnx

diff --git a/Engine/Source/Runtime/ApplicationCore/ApplicationCore.Build.cs b/Engine/Source/Runtime/ApplicationCore/ApplicationCore.Build.cs
--- a/Engine/Source/Runtime/ApplicationCore/ApplicationCore.Build.cs
+++ b/Engine/Source/Runtime/ApplicationCore/ApplicationCore.Build.cs
@@ -31,8 +31,7 @@ public class ApplicationCore : ModuleRules
 			AddEngineThirdPartyPrivateStaticDependencies(Target,
 				"XInput"
 			);
-			// We only enable UIA with Win64
-			if (Target.bCompileWithAccessibilitySupport && !Target.bIsBuildingConsoleApplication && Target.Platform == UnrealTargetPlatform.Win64)
+			if (Target.bCompileWithAccessibilitySupport && !Target.bIsBuildingConsoleApplication)
 			{
 				PublicSystemLibraries.Add("uiautomationcore.lib");
 				PublicDefinitions.Add("UE_WINDOWS_USING_UIA=1");
