CVar Android.UnifyMotionSpace can now be set to 1 to match Unreal space, 2 to swap X and Y and make it right-handed (since most systems are right handed)

diff --git a/Engine/Source/Runtime/ApplicationCore/Private/Android/AndroidInputInterface.cpp b/Engine/Source/Runtime/ApplicationCore/Private/Android/AndroidInputInterface.cpp
--- a/Engine/Source/Runtime/ApplicationCore/Private/Android/AndroidInputInterface.cpp
+++ b/Engine/Source/Runtime/ApplicationCore/Private/Android/AndroidInputInterface.cpp
@@ -54,11 +54,11 @@ static FAutoConsoleVariableRef CVarAndroidOldXBoxWirelessFirmware(
 	TEXT("Determines how XBox Wireless controller mapping is handled. 0 assumes new firmware, 1 will use old firmware mapping (Default: 0)"),
 	ECVF_Default);
 
-bool AndroidUnifyMotionSpace = false;
+int32 AndroidUnifyMotionSpace = 0;
 static FAutoConsoleVariableRef CVarAndroidUnifyMotionSpace(
 	TEXT("Android.UnifyMotionSpace"),
 	AndroidUnifyMotionSpace,
-	TEXT("Motion inputs will be processed to keep rotation rate right-handed and keep other motion inputs in the same space as rotation rate. This also forces rotation rate units to be radians/s and acceleration units to be g. Default: false"),
+	TEXT("If set to non-zero, acceleration, gravity, and rotation rate will all be in the same coordinate space. 0 for legacy behaviour. 1 will match Unreal's coordinate space (left-handed, z-up, etc). 2 will be right-handed by swapping x and y. Non-zero also forces rotation rate units to be radians/s and acceleration units to be g."),
 	ECVF_Default);
 
 TSharedRef< FAndroidInputInterface > FAndroidInputInterface::Create(const TSharedRef< FGenericApplicationMessageHandler >& InMessageHandler, const TSharedPtr< ICursor >& InCursor)
@@ -1657,7 +1657,7 @@ void FAndroidInputInterface::QueueMotionData(const FVector& Tilt, const FVector&
 	EDeviceScreenOrientation ScreenOrientation = FPlatformMisc::GetDeviceOrientation();
 	FVector TempRotationRate = RotationRate;
 
-	if (AndroidUnifyMotionSpace)
+	if (AndroidUnifyMotionSpace != 0)
 	{
 		FVector TempTilt = Tilt;
 		FVector TempGravity = Gravity;
@@ -1665,17 +1665,17 @@ void FAndroidInputInterface::QueueMotionData(const FVector& Tilt, const FVector&
 
 		auto ReorientLandscapeLeft = [](FVector InValue)
 		{
-			return FVector(-InValue.Z, -InValue.Y, InValue.X);
+			return AndroidUnifyMotionSpace == 1 ? FVector(-InValue.Z, -InValue.Y, InValue.X) : FVector(-InValue.Y, -InValue.Z, InValue.X);
 		};
 
 		auto ReorientLandscapeRight = [](FVector InValue)
 		{
-			return FVector(-InValue.Z, InValue.Y, -InValue.X);
+			return AndroidUnifyMotionSpace == 1 ? FVector(-InValue.Z, InValue.Y, -InValue.X) : FVector(InValue.Y, -InValue.Z, -InValue.X);
 		};
 
 		auto ReorientPortrait = [](FVector InValue)
 		{
-			return FVector(-InValue.Z, InValue.X, InValue.Y);
+			return AndroidUnifyMotionSpace == 1 ? FVector(-InValue.Z, InValue.X, InValue.Y) : FVector(InValue.X, -InValue.Z, InValue.Y);
 		};
 
 		const float ToG = 1.f / 9.8f;
