When Android.UnifyMotionSpace is set to 2, make sure rotation rate is right-screw

diff --git a/Engine/Source/Runtime/ApplicationCore/Private/Android/AndroidInputInterface.cpp b/Engine/Source/Runtime/ApplicationCore/Private/Android/AndroidInputInterface.cpp
--- a/Engine/Source/Runtime/ApplicationCore/Private/Android/AndroidInputInterface.cpp
+++ b/Engine/Source/Runtime/ApplicationCore/Private/Android/AndroidInputInterface.cpp
@@ -1704,6 +1704,11 @@ void FAndroidInputInterface::QueueMotionData(const FVector& Tilt, const FVector&
 			break;
 		}
 
+		if (AndroidUnifyMotionSpace == 2)
+		{
+			TempRotationRate = -TempRotationRate;
+		}
+
 		FAndroidInputInterface::MotionDataStack.Push(
 			MotionData{ TempTilt, TempRotationRate, TempGravity, TempAcceleration });
 	}
