Make FLinuxConsoleOutputDevice::Serialize Thread Safe

Set CanBeUsedOnAnyThread() to true and move to single printf in Serialize

#rnx

diff --git a/Engine/Source/Runtime/ApplicationCore/Private/Linux/LinuxConsoleOutputDevice.cpp b/Engine/Source/Runtime/ApplicationCore/Private/Linux/LinuxConsoleOutputDevice.cpp
--- a/Engine/Source/Runtime/ApplicationCore/Private/Linux/LinuxConsoleOutputDevice.cpp
+++ b/Engine/Source/Runtime/ApplicationCore/Private/Linux/LinuxConsoleOutputDevice.cpp
@@ -67,28 +67,24 @@ void FLinuxConsoleOutputDevice::Serialize(const TCHAR* Data, ELogVerbosity::Type
 		}
 		else
 		{
-			bool bNeedToResetColor = false;
+			const ANSICHAR *ConsoleColorStr = "";
 
 			if (bOutputtingToTerminal && !bOverrideColorSet)
 			{
 				if (Verbosity == ELogVerbosity::Error)
 				{
-					printf(CONSOLE_RED);
-					bNeedToResetColor = true;
+					ConsoleColorStr = CONSOLE_RED;
 				}
 				else if (Verbosity == ELogVerbosity::Warning)
 				{
-					printf(CONSOLE_YELLOW);
-					bNeedToResetColor = true;
+					ConsoleColorStr = CONSOLE_YELLOW;
 				}
 			}
 
-			printf("%s\n", TCHAR_TO_UTF8(*FOutputDeviceHelper::FormatLogLine(Verbosity, Category, Data, GPrintLogTimes)));
-
-			if (bNeedToResetColor)
-			{
-				printf(CONSOLE_NONE);
-			}
+			printf("%s%s%s\n",
+				ConsoleColorStr,
+				TCHAR_TO_UTF8(*FOutputDeviceHelper::FormatLogLine(Verbosity, Category, Data, GPrintLogTimes)),
+				ConsoleColorStr[0] ? CONSOLE_NONE : "");
 		}
 	}
 	else
@@ -109,3 +105,8 @@ void FLinuxConsoleOutputDevice::Serialize(const TCHAR* Data, ELogVerbosity::Type
 		bEntry = false;
 	}
 }
+
+bool FLinuxConsoleOutputDevice::CanBeUsedOnAnyThread() const 
+{
+	return true;
+}
diff --git a/Engine/Source/Runtime/ApplicationCore/Public/Linux/LinuxConsoleOutputDevice.h b/Engine/Source/Runtime/ApplicationCore/Public/Linux/LinuxConsoleOutputDevice.h
--- a/Engine/Source/Runtime/ApplicationCore/Public/Linux/LinuxConsoleOutputDevice.h
+++ b/Engine/Source/Runtime/ApplicationCore/Public/Linux/LinuxConsoleOutputDevice.h
@@ -49,6 +49,8 @@ public:
 
 	virtual bool IsAttached() {return false;}
 
+	virtual bool CanBeUsedOnAnyThread() const override;
+
 	/**
 	 * Displays text on the console and scrolls if necessary.
 	 *
