Siri Remote touch panel is a dpad and add Virtual Keyboard support to tvOS

diff --git a/Engine/Config/BaseEngine.ini b/Engine/Config/BaseEngine.ini
--- a/Engine/Config/BaseEngine.ini
+++ b/Engine/Config/BaseEngine.ini
@@ -2577,7 +2577,6 @@ bSupportsMetalMRT=False
 bShipForBitcode=True
 bGameSupportsMultipleActiveControllers=False
 bAllowRemoteRotation=True
-bUseRemoteAsVirtualJoystick=True
 bUseAbsoluteDpadValues=False
 bAllowControllers=True
 bBuildAsFramework=False
diff --git a/Engine/Source/Runtime/ApplicationCore/Private/IOS/IOSInputInterface.cpp b/Engine/Source/Runtime/ApplicationCore/Private/IOS/IOSInputInterface.cpp
--- a/Engine/Source/Runtime/ApplicationCore/Private/IOS/IOSInputInterface.cpp
+++ b/Engine/Source/Runtime/ApplicationCore/Private/IOS/IOSInputInterface.cpp
@@ -101,8 +101,6 @@ FIOSInputInterface::FIOSInputInterface( const TSharedRef< FGenericApplicationMes
 	: FAppleControllerInterface( InMessageHandler )
 	, bAllowRemoteRotation(false)
 	, bGameSupportsMultipleActiveControllers(false)
-	, bUseRemoteAsVirtualJoystick_DEPRECATED(true)
-	, bUseRemoteAbsoluteDpadValues(false)
     , LastHapticValue(0.0f)
     , MouseDeltaX(0)
     , MouseDeltaY(0)
@@ -120,8 +118,6 @@ FIOSInputInterface::FIOSInputInterface( const TSharedRef< FGenericApplicationMes
     
     GConfig->GetBool(TEXT("/Script/IOSRuntimeSettings.IOSRuntimeSettings"), TEXT("bGameSupportsMultipleActiveControllers"), bGameSupportsMultipleActiveControllers, GEngineIni);
     GConfig->GetBool(TEXT("/Script/IOSRuntimeSettings.IOSRuntimeSettings"), TEXT("bAllowRemoteRotation"), bAllowRemoteRotation, GEngineIni);
-    GConfig->GetBool(TEXT("/Script/IOSRuntimeSettings.IOSRuntimeSettings"), TEXT("bUseRemoteAsVirtualJoystick"), bUseRemoteAsVirtualJoystick_DEPRECATED, GEngineIni);
-    GConfig->GetBool(TEXT("/Script/IOSRuntimeSettings.IOSRuntimeSettings"), TEXT("bUseRemoteAbsoluteDpadValues"), bUseRemoteAbsoluteDpadValues, GEngineIni);
     GConfig->GetBool(TEXT("/Script/IOSRuntimeSettings.IOSRuntimeSettings"), TEXT("bAllowControllers"), bAllowControllers, GEngineIni);
     GConfig->GetBool(TEXT("/Script/IOSRuntimeSettings.IOSRuntimeSettings"), TEXT("bControllersBlockDeviceFeedback"), bControllersBlockDeviceFeedback, GEngineIni);
     
@@ -540,23 +536,13 @@ void FIOSInputInterface::SendControllerEvents()
                 HandleButtonGamepad(FGamepadKeyNames::FaceButtonLeft, i);
                 HandleButtonGamepad(FGamepadKeyNames::SpecialRight, i);
                 
-                // if we want virtual joysticks, then use the dpad values (and drain the touch queue to not leak memory)
-                if (bUseRemoteAsVirtualJoystick_DEPRECATED)
-                {
-                    HandleAnalogGamepad(FGamepadKeyNames::LeftAnalogX, i);
-                    HandleAnalogGamepad(FGamepadKeyNames::LeftAnalogY, i);
-
-                    HandleButtonGamepad(FGamepadKeyNames::LeftStickUp, i);
-                    HandleButtonGamepad(FGamepadKeyNames::LeftStickDown, i);
-                    HandleButtonGamepad(FGamepadKeyNames::LeftStickRight, i);
-                    HandleButtonGamepad(FGamepadKeyNames::LeftStickLeft, i);
-                }
-                // otherwise, process touches like ios for the remote's index
-                else
-                {
-                    ProcessTouchesAndKeys(Cont.playerIndex, LocalTouchInputStack, LocalKeyInputStack);
-                }
+                HandleAnalogGamepad(FGamepadKeyNames::LeftAnalogX, i);
+                HandleAnalogGamepad(FGamepadKeyNames::LeftAnalogY, i);
 
+                HandleButtonGamepad(FGamepadKeyNames::LeftStickUp, i);
+                HandleButtonGamepad(FGamepadKeyNames::LeftStickDown, i);
+                HandleButtonGamepad(FGamepadKeyNames::LeftStickRight, i);
+                HandleButtonGamepad(FGamepadKeyNames::LeftStickLeft, i);
                 
                 [Controller.PreviousMicroGamepad release];
                 Controller.PreviousMicroGamepad = MicroGamepad;
diff --git a/Engine/Source/Runtime/ApplicationCore/Private/IOS/IOSPlatformApplicationMisc.cpp b/Engine/Source/Runtime/ApplicationCore/Private/IOS/IOSPlatformApplicationMisc.cpp
--- a/Engine/Source/Runtime/ApplicationCore/Private/IOS/IOSPlatformApplicationMisc.cpp
+++ b/Engine/Source/Runtime/ApplicationCore/Private/IOS/IOSPlatformApplicationMisc.cpp
@@ -299,3 +299,7 @@ EScreenPhysicalAccuracy FIOSPlatformApplicationMisc::ComputePhysicalScreenDensit
 	return EScreenPhysicalAccuracy::Approximation;
 }
 
+bool FIOSPlatformApplicationMisc::RequiresVirtualKeyboard()
+{
+    return PLATFORM_HAS_TOUCH_MAIN_SCREEN || PLATFORM_TVOS;
+}
diff --git a/Engine/Source/Runtime/ApplicationCore/Public/IOS/IOSInputInterface.h b/Engine/Source/Runtime/ApplicationCore/Public/IOS/IOSInputInterface.h
--- a/Engine/Source/Runtime/ApplicationCore/Public/IOS/IOSInputInterface.h
+++ b/Engine/Source/Runtime/ApplicationCore/Public/IOS/IOSInputInterface.h
@@ -140,12 +140,6 @@ private:
 	// can the game handle multiple gamepads at the same time (siri remote is a gamepad) ?
 	bool bGameSupportsMultipleActiveControllers;
 
-	// should the remote be used as virtual joystick vs touch events
-	bool bUseRemoteAsVirtualJoystick_DEPRECATED;
-
-	// should the tracking use the pad center as the virtual joystick center?
-	bool bUseRemoteAbsoluteDpadValues;
-
 	// bluetooth connected controllers will block force feedback.
 	bool bControllersBlockDeviceFeedback;
 	
diff --git a/Engine/Source/Runtime/ApplicationCore/Public/IOS/IOSPlatformApplicationMisc.h b/Engine/Source/Runtime/ApplicationCore/Public/IOS/IOSPlatformApplicationMisc.h
--- a/Engine/Source/Runtime/ApplicationCore/Public/IOS/IOSPlatformApplicationMisc.h
+++ b/Engine/Source/Runtime/ApplicationCore/Public/IOS/IOSPlatformApplicationMisc.h
@@ -32,6 +32,8 @@ struct APPLICATIONCORE_API FIOSPlatformApplicationMisc : public FGenericPlatform
 
 	static EScreenPhysicalAccuracy ComputePhysicalScreenDensity(int32& ScreenDensity);
 
+    static bool RequiresVirtualKeyboard();
+    
 private:
 	static class FIOSApplication* CachedApplication;
 
diff --git a/Engine/Source/Runtime/IOS/IOSRuntimeSettings/Classes/IOSRuntimeSettings.h b/Engine/Source/Runtime/IOS/IOSRuntimeSettings/Classes/IOSRuntimeSettings.h
--- a/Engine/Source/Runtime/IOS/IOSRuntimeSettings/Classes/IOSRuntimeSettings.h
+++ b/Engine/Source/Runtime/IOS/IOSRuntimeSettings/Classes/IOSRuntimeSettings.h
@@ -294,15 +294,7 @@ public:
 	// If checked, the Siri Remote can be rotated to landscape view
 	UPROPERTY(GlobalConfig, EditAnywhere, Category = Input, meta = (DisplayName = "Allow AppleTV Remote landscape mode"))
 	bool bAllowRemoteRotation;
-	
-	// If checked, the trackpad is a virtual joystick (acts like the left stick of a controller). If unchecked, the trackpad will send touch events
-	UPROPERTY(config, meta = (Deprecated, DeprecationMessage = "Use AppleTV trackpad as virtual joystick. Deprecated. Siri Remote shouls always behave as a joystick"))
-	bool bUseRemoteAsVirtualJoystick_DEPRECATED;
-	
-	// If checked, the center of the trackpad is 0,0 (center) for the virtual joystick. If unchecked, the location the user taps becomes 0,0
-	UPROPERTY(GlobalConfig, EditAnywhere, Category = Input, meta = (DisplayName = "Use AppleTV Remote absolute trackpad values"))
-	bool bUseRemoteAbsoluteDpadValues;
-	
+		
 	// If checked, Bluetooth connected controllers will send input
 	UPROPERTY(GlobalConfig, EditAnywhere, Category = Input, meta = (DisplayName = "Allow MFi (Bluetooth) controllers"))
 	bool bAllowControllers;
diff --git a/Engine/Source/Runtime/IOS/IOSRuntimeSettings/Private/IOSRuntimeSettings.cpp b/Engine/Source/Runtime/IOS/IOSRuntimeSettings/Private/IOSRuntimeSettings.cpp
--- a/Engine/Source/Runtime/IOS/IOSRuntimeSettings/Private/IOSRuntimeSettings.cpp
+++ b/Engine/Source/Runtime/IOS/IOSRuntimeSettings/Private/IOSRuntimeSettings.cpp
@@ -44,8 +44,6 @@ UIOSRuntimeSettings::UIOSRuntimeSettings(const FObjectInitializer& ObjectInitial
 	AdditionalShippingLinkerFlags = TEXT("");
     bGameSupportsMultipleActiveControllers = false;
 	bAllowRemoteRotation = true;
-    bUseRemoteAsVirtualJoystick_DEPRECATED = true;
-	bUseRemoteAbsoluteDpadValues = false;
 	bDisableMotionData = false;
     bEnableRemoteNotificationsSupport = false;
     bEnableBackgroundFetch = false;
diff --git a/Engine/Source/Runtime/Slate/Private/Framework/Text/IOS/IOSPlatformTextField.cpp b/Engine/Source/Runtime/Slate/Private/Framework/Text/IOS/IOSPlatformTextField.cpp
--- a/Engine/Source/Runtime/Slate/Private/Framework/Text/IOS/IOSPlatformTextField.cpp
+++ b/Engine/Source/Runtime/Slate/Private/Framework/Text/IOS/IOSPlatformTextField.cpp
@@ -56,20 +56,16 @@ FIOSPlatformTextField::~FIOSPlatformTextField()
         TextField = nullptr;
 		dispatch_async(dispatch_get_main_queue(), ^{
             NSLog(@"Finally releasing text field %@", LocalTextField);
-#if !PLATFORM_TVOS
 			if (LocalTextField != nullptr && [LocalTextField respondsToSelector:@selector(hide:)])
 			{
             	[LocalTextField hide];
 			}
-#endif
 		});
 	}
 }
 
 void FIOSPlatformTextField::ShowVirtualKeyboard(bool bShow, int32 UserIndex, TSharedPtr<IVirtualKeyboardEntry> TextEntryWidget)
 {
-#if !PLATFORM_TVOS
-	
 	FIOSView* View = [IOSAppDelegate GetDelegate].IOSView;
 	if (View->bIsUsingIntegratedKeyboard)
 	{
@@ -120,11 +116,8 @@ void FIOSPlatformTextField::ShowVirtualKeyboard(bool bShow, int32 UserIndex, TSh
 			}
         }
 	}
-#endif
 }
 
-#if !PLATFORM_TVOS
-
 @implementation SlateTextField
 
 -(id)init
@@ -326,6 +319,3 @@ void FIOSPlatformTextField::ShowVirtualKeyboard(bool bShow, int32 UserIndex, TSh
 }
 
 @end
-
-#endif
-
diff --git a/Engine/Source/Runtime/Slate/Public/Framework/Text/IOS/IOSPlatformTextField.h b/Engine/Source/Runtime/Slate/Public/Framework/Text/IOS/IOSPlatformTextField.h
--- a/Engine/Source/Runtime/Slate/Public/Framework/Text/IOS/IOSPlatformTextField.h
+++ b/Engine/Source/Runtime/Slate/Public/Framework/Text/IOS/IOSPlatformTextField.h
@@ -28,8 +28,6 @@ private:
 
 typedef FIOSPlatformTextField FPlatformTextField;
 
-#if !PLATFORM_TVOS
-
 @interface SlateTextField : UIAlertController
 {
 	TWeakPtr<IVirtualKeyboardEntry> TextWidget;
@@ -50,4 +48,3 @@ typedef FIOSPlatformTextField FPlatformTextField;
 -(bool)hasTextWidget;
 
 @end
-#endif
