Added log highlighting, based on log category and log search string.

Commandline usage:
-LogHighlights="LogNet Cyan, LogTemp Green"
-LogStringHighlights="UNetConnection::Close=Purple, NotifyAcceptingConnection accepted from=DarkGreen"

diff --git a/Engine/Source/Runtime/ApplicationCore/Private/Windows/WindowsConsoleOutputDevice.cpp b/Engine/Source/Runtime/ApplicationCore/Private/Windows/WindowsConsoleOutputDevice.cpp
--- a/Engine/Source/Runtime/ApplicationCore/Private/Windows/WindowsConsoleOutputDevice.cpp
+++ b/Engine/Source/Runtime/ApplicationCore/Private/Windows/WindowsConsoleOutputDevice.cpp
@@ -21,6 +21,7 @@
 #include "HAL/ThreadHeartBeat.h"
 #include "Windows/WindowsHWrapper.h"
 #include "Windows/AllowWindowsPlatformTypes.h"
+#include "String/Find.h"
 
 namespace OutputDeviceConstants
 {
@@ -35,6 +36,131 @@ FWindowsConsoleOutputDevice::FWindowsConsoleOutputDevice()
 	, OverrideColorSet(false)
 	, bIsAttached(false)
 {
+#if !UE_BUILD_SHIPPING
+	auto ParseColorStr = [](const FString& ColorStr) -> const TCHAR*
+	{
+		const TCHAR* ReturnVal = nullptr;
+
+		if (ColorStr == TEXTVIEW("Black"))
+		{
+			ReturnVal = COLOR_BLACK;
+		}
+		else if (ColorStr == TEXTVIEW("DarkRed"))
+		{
+			ReturnVal = COLOR_DARK_RED;
+		}
+		else if (ColorStr == TEXTVIEW("DarkGreen"))
+		{
+			ReturnVal = COLOR_DARK_GREEN;
+		}
+		else if (ColorStr == TEXTVIEW("DarkBlue"))
+		{
+			ReturnVal = COLOR_DARK_BLUE;
+		}
+		else if (ColorStr == TEXTVIEW("DarkYellow"))
+		{
+			ReturnVal = COLOR_DARK_YELLOW;
+		}
+		else if (ColorStr == TEXTVIEW("DarkCyan"))
+		{
+			ReturnVal = COLOR_DARK_CYAN;
+		}
+		else if (ColorStr == TEXTVIEW("DarkPurple"))
+		{
+			ReturnVal = COLOR_DARK_PURPLE;
+		}
+		else if (ColorStr == TEXTVIEW("Gray"))
+		{
+			ReturnVal = COLOR_DARK_WHITE;
+		}
+		else if (ColorStr == TEXTVIEW("Red"))
+		{
+			ReturnVal = COLOR_RED;
+		}
+		else if (ColorStr == TEXTVIEW("Green"))
+		{
+			ReturnVal = COLOR_GREEN;
+		}
+		else if (ColorStr == TEXTVIEW("Blue"))
+		{
+			ReturnVal = COLOR_BLUE;
+		}
+		else if (ColorStr == TEXTVIEW("Yellow"))
+		{
+			ReturnVal = COLOR_YELLOW;
+		}
+		else if (ColorStr == TEXTVIEW("Cyan"))
+		{
+			ReturnVal = COLOR_CYAN;
+		}
+		else if (ColorStr == TEXTVIEW("Purple"))
+		{
+			ReturnVal = COLOR_PURPLE;
+		}
+		else if (ColorStr == TEXTVIEW("White"))
+		{
+			ReturnVal = COLOR_WHITE;
+		}
+
+		return ReturnVal;
+	};
+
+	FString HighlightsStr;
+
+	// -LogHighlights="LogNet Cyan, LogTemp Green"
+	if (FParse::Value(FCommandLine::Get(), TEXT("LogHighlights="), HighlightsStr))
+	{
+		TArray<FString> HighlightsList;
+
+		HighlightsStr.ParseIntoArray(HighlightsList, TEXT(","));
+
+		for (const FString& CurHighlightEntry : HighlightsList)
+		{
+			TArray<FString> CategoryAndColor;
+
+			if (CurHighlightEntry.TrimStartAndEnd().ParseIntoArray(CategoryAndColor, TEXT(" ")) && CategoryAndColor.Num() == 2)
+			{
+				const TCHAR* ColorStr = ParseColorStr(CategoryAndColor[1].TrimStartAndEnd());
+
+				if (ColorStr != nullptr)
+				{
+					FLogHighlight& NewEntry = LogHighlights.AddDefaulted_GetRef();
+
+					NewEntry.Category = FName(CategoryAndColor[0]);
+					NewEntry.Color = ColorStr;
+				}
+			}
+		}
+	}
+
+	FString StringHighlights;
+
+	// -LogStringHighlights="UNetConnection::Close=Purple, NotifyAcceptingConnection accepted from=DarkGreen"
+	if (FParse::Value(FCommandLine::Get(), TEXT("LogStringHighlights="), StringHighlights))
+	{
+		TArray<FString> StringHighlightsList;
+
+		StringHighlights.ParseIntoArray(StringHighlightsList, TEXT(","));
+
+		for (const FString& CurStringHighlightEntry : StringHighlightsList)
+		{
+			TArray<FString> StringAndColor;
+
+			if (CurStringHighlightEntry.ParseIntoArray(StringAndColor, TEXT("=")) && StringAndColor.Num() == 2)
+			{
+				const TCHAR* ColorStr = ParseColorStr(StringAndColor[1].TrimStartAndEnd());
+
+				if (ColorStr != nullptr)
+				{
+					FLogStringHighlight& NewEntry = LogStringHighlights.AddDefaulted_GetRef();
+
+					NewEntry.SearchString = StringAndColor[0].GetCharArray();
+					NewEntry.Color = ColorStr;
+				}
+			}
+		}
+	}
+#endif
 }
 
 FWindowsConsoleOutputDevice::~FWindowsConsoleOutputDevice()
@@ -180,6 +306,8 @@ bool FWindowsConsoleOutputDevice::IsShown()
 
 void FWindowsConsoleOutputDevice::Serialize( const TCHAR* Data, ELogVerbosity::Type Verbosity, const class FName& Category, const double Time )
 {
+	using namespace UE::String;
+
 	if( ConsoleHandle )
 	{
 		const double RealTime = Time == -1.0f ? FPlatformTime::Seconds() - GStartTime : Time;
@@ -207,6 +335,38 @@ void FWindowsConsoleOutputDevice::Serialize( const TCHAR* Data, ELogVerbosity::T
 						SetColor( COLOR_YELLOW );
 						bNeedToResetColor = true;
 					}
+#if !UE_BUILD_SHIPPING
+					else
+					{
+						if (LogHighlights.Num() > 0)
+						{
+							if (const FLogHighlight* CurHighlight = LogHighlights.FindByKey(Category))
+							{
+								SetColor(CurHighlight->Color);
+								bNeedToResetColor = true;
+							}
+						}
+
+						if (LogStringHighlights.Num() > 0)
+						{
+							FStringView DataView = TStringView(Data);
+
+							for (const FLogStringHighlight& CurStringHighlight : LogStringHighlights)
+							{
+								FStringView SearchView = TStringView(CurStringHighlight.SearchString.GetData(),
+																		CurStringHighlight.SearchString.Num()-1);
+
+								if (FindFirst(DataView, SearchView, ESearchCase::IgnoreCase) != INDEX_NONE)
+								{
+									SetColor(CurStringHighlight.Color);
+									bNeedToResetColor = true;
+
+									break;
+								}
+							}
+						}
+					}
+#endif
 				}
 				TCHAR OutputString[MAX_SPRINTF]=TEXT(""); //@warning: this is safe as FCString::Sprintf only use 1024 characters max
 				FCString::Sprintf(OutputString,TEXT("%s%s"),*FOutputDeviceHelper::FormatLogLine(Verbosity, Category, Data, GPrintLogTimes,RealTime),LINE_TERMINATOR);
diff --git a/Engine/Source/Runtime/ApplicationCore/Public/Windows/WindowsConsoleOutputDevice.h b/Engine/Source/Runtime/ApplicationCore/Public/Windows/WindowsConsoleOutputDevice.h
--- a/Engine/Source/Runtime/ApplicationCore/Public/Windows/WindowsConsoleOutputDevice.h
+++ b/Engine/Source/Runtime/ApplicationCore/Public/Windows/WindowsConsoleOutputDevice.h
@@ -19,6 +19,40 @@ private:
 
 	/** true if this console device is attached to an existing console window. */
 	bool bIsAttached;
+
+#if !UE_BUILD_SHIPPING
+	/** An entry for log category highlighting */
+	struct FLogHighlight
+	{
+		/** The category to highlight */
+		FName Category;
+
+		/** The 4 digit color code to highlight with */
+		const TCHAR* Color = nullptr;
+
+		bool operator == (FName InCategory) const
+		{
+			return Category == InCategory;
+		}
+	};
+
+	/** Log categories to be highlighted */
+	TArray<FLogHighlight, TInlineAllocator<8>> LogHighlights;
+
+	/** An entry for log string highlighting */
+	struct FLogStringHighlight
+	{
+		/** The string to search for and highlight */
+		TArray<TCHAR, TInlineAllocator<128>> SearchString;
+
+		/** The 4 digit color code to highlight with */
+		const TCHAR* Color = nullptr;
+	};
+
+	/** Log strings to be highlighted */
+	TArray<FLogStringHighlight, TInlineAllocator<8>> LogStringHighlights;
+#endif
+
 	/**
 	 * Saves the console window's position and size to the game .ini
 	 */
