CIS fix.

diff --git a/Engine/Source/Runtime/ApplicationCore/Public/Windows/WindowsConsoleOutputDevice.h b/Engine/Source/Runtime/ApplicationCore/Public/Windows/WindowsConsoleOutputDevice.h
--- a/Engine/Source/Runtime/ApplicationCore/Public/Windows/WindowsConsoleOutputDevice.h
+++ b/Engine/Source/Runtime/ApplicationCore/Public/Windows/WindowsConsoleOutputDevice.h
@@ -4,6 +4,7 @@
 
 #include "CoreTypes.h"
 #include "Misc/OutputDeviceConsole.h"
+#include "UObject/NameTypes.h"
 
 /**
  * Windows implementation of console log window, utilizing the Win32 console API
