Fix UngrabImpl not tearing down barries, and not disabling relative mode

diff --git a/Engine/Source/Runtime/ApplicationCore/Private/Linux/LinuxPlatformApplicationMisc.cpp b/Engine/Source/Runtime/ApplicationCore/Private/Linux/LinuxPlatformApplicationMisc.cpp
--- a/Engine/Source/Runtime/ApplicationCore/Private/Linux/LinuxPlatformApplicationMisc.cpp
+++ b/Engine/Source/Runtime/ApplicationCore/Private/Linux/LinuxPlatformApplicationMisc.cpp
@@ -232,13 +232,21 @@ void UngrabAllInputImpl()
 {
 	if (GInitializedSDL)
 	{
-		SDL_Window * GrabbedWindow = SDL_GetGrabbedWindow();
+		SDL_Window* GrabbedWindow = SDL_GetGrabbedWindow();
 		if (GrabbedWindow)
 		{
 			SDL_SetWindowGrab(GrabbedWindow, SDL_FALSE);
 			SDL_SetWindowKeyboardGrab(GrabbedWindow, SDL_FALSE);
 		}
-		SDL_SetWindowMouseRect(nullptr, nullptr);
+
+		SDL_Window* MouseFocusedWindow = SDL_GetMouseFocus();
+		if (MouseFocusedWindow)
+		{
+			SDL_SetWindowMouseRect(MouseFocusedWindow, nullptr);
+		}
+
+		SDL_SetRelativeMouseMode(SDL_FALSE);
+		SDL_ShowCursor(SDL_ENABLE);
 		SDL_CaptureMouse(SDL_FALSE);
 	}
 }
