Added support for IsGamepadAttached on Linux

diff --git a/Engine/Source/Runtime/ApplicationCore/Private/Linux/LinuxApplication.cpp b/Engine/Source/Runtime/ApplicationCore/Private/Linux/LinuxApplication.cpp
--- a/Engine/Source/Runtime/ApplicationCore/Private/Linux/LinuxApplication.cpp
+++ b/Engine/Source/Runtime/ApplicationCore/Private/Linux/LinuxApplication.cpp
@@ -1366,19 +1366,30 @@ FModifierKeysState FLinuxApplication::GetModifierKeys() const
 	return FModifierKeysState( bIsLeftShiftDown, bIsRightShiftDown, bIsLeftControlDown, bIsRightControlDown, bIsLeftAltDown, bIsRightAltDown, false, false, bAreCapsLocked );
 }
 
-
 void FLinuxApplication::SetCapture( const TSharedPtr< FGenericWindow >& InWindow )
 {
 	bIsMouseCaptureEnabled = InWindow.IsValid();
 	UpdateMouseCaptureWindow( bIsMouseCaptureEnabled ? ((FLinuxWindow*)InWindow.Get())->GetHWnd() : NULL );
 }
 
-
 void* FLinuxApplication::GetCapture( void ) const
 {
 	return ( bIsMouseCaptureEnabled && MouseCaptureWindow ) ? MouseCaptureWindow : NULL;
 }
 
+bool FLinuxApplication::IsGamepadAttached() const
+{
+	for (const TPair<SDL_JoystickID, SDLControllerState>& ControllerState : ControllerStates)
+	{
+		if (SDL_GameControllerGetAttached(ControllerState.Value.Controller))
+		{
+			return true;
+		}
+	}
+	
+	return false;
+}
+
 void FLinuxApplication::UpdateMouseCaptureWindow(SDL_HWindow TargetWindow)
 {
 	const bool bEnable = bIsMouseCaptureEnabled || bIsMouseCursorLocked;
diff --git a/Engine/Source/Runtime/ApplicationCore/Public/Linux/LinuxApplication.h b/Engine/Source/Runtime/ApplicationCore/Public/Linux/LinuxApplication.h
--- a/Engine/Source/Runtime/ApplicationCore/Public/Linux/LinuxApplication.h
+++ b/Engine/Source/Runtime/ApplicationCore/Public/Linux/LinuxApplication.h
@@ -71,6 +71,8 @@ public:
 
 	virtual bool IsUsingHighPrecisionMouseMode() const override { return bUsingHighPrecisionMouseInput; }
 
+	virtual bool IsGamepadAttached() const override;
+	
 	virtual FModifierKeysState GetModifierKeys() const override;
 
 	virtual FPlatformRect GetWorkArea( const FPlatformRect& CurrentWindow ) const override;
@@ -303,6 +305,9 @@ private:
 	/** Using high precision mouse input */
 	bool bUsingHighPrecisionMouseInput;
 
+	/** Gamepad is attached */
+	bool bIsGamepadAttached;
+	
 	/** TODO: describe */
 	bool bIsMouseCursorLocked;
 
