Removed an unused bool related to Gamepads

diff --git a/Engine/Source/Runtime/ApplicationCore/Public/Linux/LinuxApplication.h b/Engine/Source/Runtime/ApplicationCore/Public/Linux/LinuxApplication.h
--- a/Engine/Source/Runtime/ApplicationCore/Public/Linux/LinuxApplication.h
+++ b/Engine/Source/Runtime/ApplicationCore/Public/Linux/LinuxApplication.h
@@ -305,9 +305,6 @@ private:
 	/** Using high precision mouse input */
 	bool bUsingHighPrecisionMouseInput;
 
-	/** Gamepad is attached */
-	bool bIsGamepadAttached;
-	
 	/** TODO: describe */
 	bool bIsMouseCursorLocked;
 
