Create the Platform Input Device Mapper. This provides an interface that platforms can use to create a map of Input Devices to their owning PlatformUser. This will allow for granular control of how many input devices are mapped to a single user, which some platforms want to do differently than others. The API can be used to query what devices are mapped to which user, so that you can easily determine which controller should recieve a force feedback event, or which user has pushed a certain button.

- Create the FInputDeviceID struct, which represents a single input device such as a gamepad, keyboard, or mouse. These are globally unique, and allocated by the Platform Input Device Mapper. This struct will replace the old "int32 ControllerID" that is used in various places to provide some type safety and maintainability.

- Functions to query all the currently connected devices on a given platform, as well as get their current state and owning platform users. These can be exposed to blueprints via the local player, game instance, or even a subsystem in the future. See UE-147801 for progress on Blueprint Support.

- New delegates for the ConnectionChange and PairingChange of a gamepad device that utilize the new FInputDeviceId and FPlatformUserId

- Add functions to the GenericApplicationMessageHandler for the "OnController_X" functions. These will be implemented later when I can replace the old uses of the "int32 ControllerId". See UE-147796.

diff --git a/Engine/Source/Runtime/ApplicationCore/Private/GenericPlatform/GenericPlatformApplicationMisc.cpp b/Engine/Source/Runtime/ApplicationCore/Private/GenericPlatform/GenericPlatformApplicationMisc.cpp
--- a/Engine/Source/Runtime/ApplicationCore/Private/GenericPlatform/GenericPlatformApplicationMisc.cpp
+++ b/Engine/Source/Runtime/ApplicationCore/Private/GenericPlatform/GenericPlatformApplicationMisc.cpp
@@ -2,6 +2,7 @@
 
 #include "GenericPlatform/GenericPlatformApplicationMisc.h"
 #include "GenericPlatform/GenericApplication.h"
+#include "GenericPlatform/GenericPlatformInputDeviceMapper.h"
 #include "Misc/OutputDeviceAnsiError.h"
 #include "HAL/FeedbackContextAnsi.h"
 #include "Math/Color.h"
@@ -61,6 +62,11 @@ class FFeedbackContext* FGenericPlatformApplicationMisc::GetFeedbackContext()
 	return FPlatformOutputDevices::GetFeedbackContext();
 }
 
+IPlatformInputDeviceMapper* FGenericPlatformApplicationMisc::CreatePlatformInputDeviceManager()
+{
+	return new FGenericPlatformInputDeviceMapper(/* bUsingControllerIdAsUserId = */ true, /* bShouldBroadcastLegacyDelegates = */ true);
+}
+
 GenericApplication* FGenericPlatformApplicationMisc::CreateApplication()
 {
 	return new GenericApplication( nullptr );
diff --git a/Engine/Source/Runtime/ApplicationCore/Private/GenericPlatform/GenericPlatformInputDeviceMapper.cpp b/Engine/Source/Runtime/ApplicationCore/Private/GenericPlatform/GenericPlatformInputDeviceMapper.cpp
new file mode 100644
--- /dev/null
+++ b/Engine/Source/Runtime/ApplicationCore/Private/GenericPlatform/GenericPlatformInputDeviceMapper.cpp
@@ -0,0 +1,326 @@
+// Copyright Epic Games, Inc. All Rights Reserved.
+
+#include "GenericPlatform/GenericPlatformInputDeviceMapper.h"
+#include "GenericPlatform/GenericPlatformApplicationMisc.h"
+#include "Misc/CoreDelegates.h"
+
+///////////////////////////////////////////////////////////////////////////
+// IPlatformInputDeviceMapper
+
+DEFINE_LOG_CATEGORY(LogInputDeviceMapper);
+
+IPlatformInputDeviceMapper::FOnUserInputDeviceConnectionChange IPlatformInputDeviceMapper::OnInputDeviceConnectionChange;
+IPlatformInputDeviceMapper::FOnUserInputDevicePairingChange IPlatformInputDeviceMapper::OnInputDevicePairingChange;
+
+IPlatformInputDeviceMapper& IPlatformInputDeviceMapper::Get()
+{
+	static IPlatformInputDeviceMapper* StaticManager = nullptr;
+	if (!StaticManager)
+	{
+		StaticManager = FGenericPlatformApplicationMisc::CreatePlatformInputDeviceManager();
+		check(StaticManager);
+	}
+
+	return *StaticManager;
+}
+
+int32 IPlatformInputDeviceMapper::GetAllInputDevicesForUser(const FPlatformUserId UserId, TArray<FInputDeviceId>& OutInputDevices) const
+{
+	OutInputDevices.Reset();
+	
+	for (TPair<FInputDeviceId, UE::Input::FInputDeviceState> MappedDevice : MappedInputDevices)
+	{
+		if (MappedDevice.Value.OwningPlatformUser == UserId)
+		{
+			OutInputDevices.AddUnique(MappedDevice.Key);
+		}
+	}
+	
+	return OutInputDevices.Num();
+}
+
+int32 IPlatformInputDeviceMapper::GetAllInputDevices(TArray<FInputDeviceId>& OutInputDevices) const
+{
+	return MappedInputDevices.GetKeys(OutInputDevices);
+}
+
+int32 IPlatformInputDeviceMapper::GetAllConnectedInputDevices(TArray<FInputDeviceId>& OutInputDevices) const
+{
+	using namespace UE::Input;
+	OutInputDevices.Reset();
+	
+	for (TPair<FInputDeviceId, FInputDeviceState> MappedDevice : MappedInputDevices)
+	{
+		if(MappedDevice.Value.ConnectionState == EInputDeviceConnectionState::Connected)
+		{
+			OutInputDevices.AddUnique(MappedDevice.Key);
+		}
+	}
+	
+	return OutInputDevices.Num();
+}
+
+int32 IPlatformInputDeviceMapper::GetAllActiveUsers(TArray<FPlatformUserId>& OutUsers) const
+{
+	OutUsers.Reset();
+
+	// Add the owning platform user for each input device
+	for (TPair<FInputDeviceId, UE::Input::FInputDeviceState> MappedDevice : MappedInputDevices)
+	{
+		OutUsers.AddUnique(MappedDevice.Value.OwningPlatformUser);
+	}
+	
+	return OutUsers.Num();
+}
+
+bool IPlatformInputDeviceMapper::IsUnpairedUserId(const FPlatformUserId PlatformId) const
+{
+	return PlatformId == GetUserForUnpairedInputDevices();
+}
+
+bool IPlatformInputDeviceMapper::IsInputDeviceMappedToUnpairedUser(const FInputDeviceId InputDevice) const
+{
+	if (const UE::Input::FInputDeviceState* DeviceState = MappedInputDevices.Find(InputDevice))
+	{
+		return IsUnpairedUserId(DeviceState->OwningPlatformUser);
+	}
+	return false;
+}
+
+FPlatformUserId IPlatformInputDeviceMapper::GetUserForInputDevice(FInputDeviceId DeviceId) const
+{
+	if (const UE::Input::FInputDeviceState* FoundState = MappedInputDevices.Find(DeviceId))
+	{
+		return FoundState->OwningPlatformUser;
+	}
+	return PLATFORMUSERID_NONE;
+}
+
+FInputDeviceId IPlatformInputDeviceMapper::GetPrimaryInputDeviceForUser(FPlatformUserId UserId) const
+{
+	FInputDeviceId FoundDevice = INPUTDEVICEID_NONE;
+
+	// By default look for the lowest input device mapped to this user
+	for (const TPair<FInputDeviceId, UE::Input::FInputDeviceState>& DeviceMapping : MappedInputDevices)
+	{
+		if (DeviceMapping.Value.OwningPlatformUser == UserId)
+		{
+			if (FoundDevice == INPUTDEVICEID_NONE || DeviceMapping.Key < FoundDevice)
+			{
+				FoundDevice = DeviceMapping.Key;
+			}
+		}
+	}
+
+	return FoundDevice;
+}
+
+bool IPlatformInputDeviceMapper::Internal_SetInputDeviceConnectionState(FInputDeviceId DeviceId, UE::Input::EInputDeviceConnectionState NewState)
+{
+	if (!DeviceId.IsValid())
+	{
+		UE_LOG(LogInputDeviceMapper, Error, TEXT("IPlatformInputDeviceMapper::Internal_SetInputDeviceConnectionState was called with an invalid DeviceId of '%d'"), DeviceId.GetId());
+		return false;
+	}
+
+	// Determine the owning user for this input device
+	FPlatformUserId OwningUser = GetUserForInputDevice(DeviceId);
+
+	// If the user is invalid, then fallback to being "Unpaired" user on this platform (which may still be PLATFORMUSERID_NONE)
+	if (!OwningUser.IsValid())
+	{
+		OwningUser = GetUserForUnpairedInputDevices();
+	}
+
+	// Mapping the input device to the user will ensure that it is correctly mapped to the given user.
+	// This covers the case where someone has called this function with a new input device that is not
+	// yet mapped, as well as broadcasting the delegates we want.
+	return Internal_MapInputDeviceToUser(DeviceId, OwningUser, NewState);
+}
+
+UE::Input::EInputDeviceConnectionState IPlatformInputDeviceMapper::GetInputDeviceConnectionState(const FInputDeviceId DeviceId) const
+{
+	using namespace UE::Input;
+	EInputDeviceConnectionState State = EInputDeviceConnectionState::Unknown;
+
+	if (!DeviceId.IsValid())
+	{
+		State = EInputDeviceConnectionState::Invalid;
+	}
+	else if (const FInputDeviceState* MappedDeviceState = MappedInputDevices.Find(DeviceId))
+	{
+		State = MappedDeviceState->ConnectionState;
+	}
+
+	return State;
+}
+
+bool IPlatformInputDeviceMapper::Internal_MapInputDeviceToUser(FInputDeviceId DeviceId, FPlatformUserId UserId, UE::Input::EInputDeviceConnectionState ConnectionState)
+{
+	if (!DeviceId.IsValid())
+	{
+		UE_LOG(LogInputDeviceMapper, Error, TEXT("IPlatformInputDeviceMapper::Internal_MapInputDeviceToUser was called with an invalid DeviceId of '%d'"), DeviceId.GetId());
+		return false;
+	}
+
+	// Some platforms could validate it had been allocated before, but we allocate on demand if needed
+	{
+		if (DeviceId > LastInputDeviceId)
+		{
+			LastInputDeviceId = DeviceId;
+		}
+		
+		if (UserId > LastPlatformUserId)
+		{
+			LastPlatformUserId = UserId;
+		}
+	}
+	
+	// Store the connection state of the input device
+	using namespace UE::Input;
+	FInputDeviceState& InputDeviceState = MappedInputDevices.FindOrAdd(DeviceId);
+	InputDeviceState.OwningPlatformUser = UserId;
+	InputDeviceState.ConnectionState = ConnectionState;
+	
+	// Broadcast delegates to let listeners know that the platform user has had an input device change
+	OnInputDeviceConnectionChange.Broadcast(ConnectionState, UserId, DeviceId);
+
+	if (ShouldBroadcastLegacyDelegates())
+	{
+		const bool bIsConnected = (ConnectionState == EInputDeviceConnectionState::Connected);
+		FCoreDelegates::OnControllerConnectionChange.Broadcast(bIsConnected, UserId, DeviceId.GetId());
+	}
+	
+	return true;
+}
+
+bool IPlatformInputDeviceMapper::Internal_ChangeInputDeviceUserMapping(FInputDeviceId DeviceId, FPlatformUserId NewUserId, FPlatformUserId OldUserId)
+{
+	if (!DeviceId.IsValid())
+	{
+		UE_LOG(LogInputDeviceMapper, Error, TEXT("IPlatformInputDeviceMapper::Internal_ChangeInputDeviceUserMapping was called with an invalid DeviceId of '%d'"), DeviceId.GetId());
+		return false;
+	}
+
+	using namespace UE::Input;
+
+	// Update the existing device state to be the new owning platform user
+	if (FInputDeviceState* ExistingDeviceState = MappedInputDevices.Find(DeviceId))
+	{
+		// Only change the platform user of this device if the old user matches up with the one that was given
+		if (ensureMsgf(ExistingDeviceState->OwningPlatformUser == OldUserId, TEXT("Attempting to change the Input Device User Mapping with a mismatched OldUserId!")))
+		{
+			ExistingDeviceState->OwningPlatformUser = NewUserId;
+		}
+	}
+	else
+	{
+		UE_LOG(LogInputDeviceMapper, Error, TEXT("IPlatformInputDeviceMapper::Internal_ChangeInputDeviceUserMapping: DeviceID '%d' is not mapped! Call Internal_MapInputDeviceToUser to map it to a user first!"), DeviceId.GetId());
+		return false;
+	}
+
+	// Broadcast the delegates letting listeners know that the input device has changed owners
+	OnInputDevicePairingChange.Broadcast(DeviceId, NewUserId, OldUserId);
+	
+	if (ShouldBroadcastLegacyDelegates())
+	{
+		// Remap the DeviceId to the older int32 "ControllerId" format for the legacy delegates
+		int32 LegacyControllerId = INDEX_NONE;
+		RemapUserAndDeviceToControllerId(NewUserId, LegacyControllerId, DeviceId);
+		
+		FCoreDelegates::OnControllerPairingChange.Broadcast(LegacyControllerId, NewUserId, OldUserId);
+	}
+	return true;
+}
+
+///////////////////////////////////////////////////////////////////////////
+// FGenericPlatformInputDeviceMapper
+
+FGenericPlatformInputDeviceMapper::FGenericPlatformInputDeviceMapper(const bool InbUsingControllerIdAsUserId, const bool InbShouldBroadcastLegacyDelegates)
+	: bUsingControllerIdAsUserId(InbUsingControllerIdAsUserId)
+	, bShouldBroadcastLegacyDelegates(InbShouldBroadcastLegacyDelegates)
+{
+	// Set the last input device id to be the default of 0, that way any new devices will have
+	// an index of 1 or higher and we can use the Default Input Device as a fallback for any
+	// unpaired input devices without an owning PlatformUserId
+	LastInputDeviceId = GetDefaultInputDevice();
+}
+
+FPlatformUserId FGenericPlatformInputDeviceMapper::GetUserForUnpairedInputDevices() const
+{
+	// Not supported by default. If a platform wanted to support this, then it is recommended that
+	// you create a static const FPlatformUserID with a value of 0 to start out as the "Unpaired"
+	// user that input devices can then map to
+	return PLATFORMUSERID_NONE;
+}
+
+FInputDeviceId FGenericPlatformInputDeviceMapper::GetDefaultInputDevice() const
+{
+	static const FInputDeviceId DefaultInputDeviceId = FInputDeviceId::CreateFromInternalId(0);
+	return DefaultInputDeviceId;
+}
+
+bool FGenericPlatformInputDeviceMapper::RemapControllerIdToPlatformUserAndDevice(int32 ControllerId, FPlatformUserId& InOutUserId, FInputDeviceId& OutInputDeviceId)
+{
+	if (IsUsingControllerIdAsUserId())
+	{
+		if (InOutUserId.GetInternalId() >= 0 && ControllerId >= 0 && InOutUserId.GetInternalId() != ControllerId)
+		{
+			// Both are valid so use them
+			OutInputDeviceId = FInputDeviceId::CreateFromInternalId(ControllerId);
+			return true;
+		}
+		else if (ControllerId >= 0)
+		{
+			// Just use controller id, and copy over device
+			InOutUserId = FPlatformUserId::CreateFromInternalId(ControllerId);
+			OutInputDeviceId = FInputDeviceId::CreateFromInternalId(ControllerId);
+			return true;
+		}
+		else if (InOutUserId.GetInternalId() >= 0)
+		{
+			// Ignore controller id
+			OutInputDeviceId = FInputDeviceId::CreateFromInternalId(InOutUserId.GetInternalId());
+			return true;
+		}
+	}
+	
+	return false;
+}
+
+bool FGenericPlatformInputDeviceMapper::RemapUserAndDeviceToControllerId(FPlatformUserId UserId, int32& OutControllerId, FInputDeviceId OptionalDeviceId /* = INPUTDEVICEID_NONE */)
+{
+	// It's just a 1:1 mapping of the old ControllerId to PlatformId if this is true
+	if (IsUsingControllerIdAsUserId())
+	{
+		OutControllerId = UserId;
+		return true;
+	}
+	return false;
+}
+
+bool FGenericPlatformInputDeviceMapper::IsUsingControllerIdAsUserId() const
+{
+	return bUsingControllerIdAsUserId;
+}
+
+bool FGenericPlatformInputDeviceMapper::ShouldBroadcastLegacyDelegates() const
+{
+	return bShouldBroadcastLegacyDelegates;
+}
+
+FPlatformUserId FGenericPlatformInputDeviceMapper::AllocateNewUserId()
+{
+	// Create a new platform user ID that is 1 higher than the last one
+	LastPlatformUserId = FPlatformUserId::CreateFromInternalId(LastPlatformUserId.GetInternalId() + 1);
+	
+	return LastPlatformUserId;
+}
+
+FInputDeviceId FGenericPlatformInputDeviceMapper::AllocateNewInputDeviceId()
+{
+	// Create a new platform user ID that is 1 higher than the last one
+	LastInputDeviceId = FInputDeviceId::CreateFromInternalId(LastInputDeviceId.GetId() + 1);
+	
+	return LastInputDeviceId;
+}
\ No newline at end of file
diff --git a/Engine/Source/Runtime/ApplicationCore/Public/GenericPlatform/GenericPlatformApplicationMisc.h b/Engine/Source/Runtime/ApplicationCore/Public/GenericPlatform/GenericPlatformApplicationMisc.h
--- a/Engine/Source/Runtime/ApplicationCore/Public/GenericPlatform/GenericPlatformApplicationMisc.h
+++ b/Engine/Source/Runtime/ApplicationCore/Public/GenericPlatform/GenericPlatformApplicationMisc.h
@@ -7,6 +7,7 @@
 #include "HAL/IConsoleManager.h"
 
 UE_DECLARE_LWC_TYPE(Vector2,,FVector2D);
+class IPlatformInputDeviceMapper;
 
 /**
  * The accuracy when dealing with physical characteristics of the monitor/screen of the device we're running on.
@@ -57,6 +58,11 @@ struct APPLICATIONCORE_API FGenericPlatformApplicationMisc
 	 */
 	static FFeedbackContext* GetFeedbackContext();
 
+	/**
+	 * Gets a pointer to the default platform input device manager.
+	 */
+	static IPlatformInputDeviceMapper* CreatePlatformInputDeviceManager();
+
 	/**
 	 * Creates an application instance.
 	 */
diff --git a/Engine/Source/Runtime/ApplicationCore/Public/GenericPlatform/GenericPlatformInputDeviceMapper.h b/Engine/Source/Runtime/ApplicationCore/Public/GenericPlatform/GenericPlatformInputDeviceMapper.h
new file mode 100644
--- /dev/null
+++ b/Engine/Source/Runtime/ApplicationCore/Public/GenericPlatform/GenericPlatformInputDeviceMapper.h
@@ -0,0 +1,279 @@
+// Copyright Epic Games, Inc. All Rights Reserved.
+
+#pragma once
+
+#include "CoreTypes.h"
+#include "Containers/Array.h"
+#include "Containers/Map.h"
+#include "Delegates/Delegate.h"
+#include "Logging/LogMacros.h"
+
+namespace UE
+{
+namespace Input
+{
+
+	/** Represents the connection status of a given FInputDeviceId */
+	enum class APPLICATIONCORE_API EInputDeviceConnectionState : uint8
+	{
+		/** This is not a valid input device */
+		Invalid,
+
+		/** It is not known if this device is connected or not */
+		Unknown,
+
+		/** Device is definitely connected */
+		Disconnected,
+
+		/** Definitely connected and powered on */
+		Connected
+	};
+
+	/** Data about an input device's current state */
+	struct APPLICATIONCORE_API FInputDeviceState
+	{
+		/** The platform user that this input device belongs to */
+		FPlatformUserId OwningPlatformUser = PLATFORMUSERID_NONE;
+
+		/** The connection state of this input device */
+		EInputDeviceConnectionState ConnectionState = EInputDeviceConnectionState::Invalid;
+	};
+	
+}	// namespace Input
+}	// namespace UE
+
+DECLARE_LOG_CATEGORY_EXTERN(LogInputDeviceMapper, Log, All);
+
+/**
+ * Base class to private a mapping of Platform Users (FPlatformUserID)
+ * to their associated available input devices (FInputDeviceID).
+ *
+ * This will handle the allocation of the globally unique identifier
+ * of the FInputDeviceID, and allow overrides of how each platform
+ * maps input devices to their users. Some platforms may desire to
+ * have each new input device assigned to a different user, while others
+ * may want multiple input devices associated with a single user.
+ */
+class APPLICATIONCORE_API IPlatformInputDeviceMapper
+{
+public:
+	
+	/** Get the platform input device mapper */
+	static IPlatformInputDeviceMapper& Get();
+	
+	IPlatformInputDeviceMapper() = default;
+	
+	/** Virtual destructor */
+	virtual ~IPlatformInputDeviceMapper() = default;
+	
+	/**
+	 * Populates the OutInputDevices array with any InputDeviceID's that are mapped to the given platform user
+	 *
+	 * @param UserId				The Platform User to gather the input devices of.
+	 * @param OutInputDevices		Array of input device ID's that will be populated with the mapped devices.
+	 * @return						The number of mapped devices, INDEX_NONE if the user was not found.
+	 */
+	virtual int32 GetAllInputDevicesForUser(const FPlatformUserId UserId, TArray<FInputDeviceId>& OutInputDevices) const;
+
+	/**
+	 * Get all mapped input devices on this platform regardless of their connection state.
+	 * 
+	 * @param OutInputDevices	Array of input devices to populate
+	 * @return					The number of connected input devices
+	 */
+	virtual int32 GetAllInputDevices(TArray<FInputDeviceId>& OutInputDevices) const;
+	
+	/**
+	 * Gather all currently connected input devices
+	 * 
+	 * @param OutInputDevices	Array of input devices to populate
+	 * @return					The number of connected input devices
+	 */
+	virtual int32 GetAllConnectedInputDevices(TArray<FInputDeviceId>& OutInputDevices) const;
+
+	/**
+	 * Get all currently active platform ids, anyone who has a mapped input device
+	 *
+	 * @param OutUsers		Array that will be populated with the platform users.
+	 * @return				The number of active platform users
+	 */
+	virtual int32 GetAllActiveUsers(TArray<FPlatformUserId>& OutUsers) const;
+
+	/**
+	 * Returns the platform user id that is being used for unmapped input devices.
+	 * Will be PLATFORMUSERID_NONE if platform does not support this (this is the default behavior)
+	 */
+	virtual FPlatformUserId GetUserForUnpairedInputDevices() const = 0;
+
+	/** Returns true if the given Platform User Id is the user for unpaired input devices on this platform. */
+	virtual bool IsUnpairedUserId(const FPlatformUserId PlatformId) const;
+
+	/** Returns true if the given input device is mapped to the unpaired platform user id. */
+	virtual bool IsInputDeviceMappedToUnpairedUser(const FInputDeviceId InputDevice) const;
+	
+	/** Returns the default device id used for things like keyboard/mouse input */
+	virtual FInputDeviceId GetDefaultInputDevice() const = 0;
+
+	/** Returns the platform user attached to this input device, or PLATFORMUSERID_NONE if invalid */
+	virtual FPlatformUserId GetUserForInputDevice(FInputDeviceId DeviceId) const;
+
+	/** Returns the primary input device used by a specific player, or INPUTDEVICEID_NONE if invalid */
+	virtual FInputDeviceId GetPrimaryInputDeviceForUser(FPlatformUserId UserId) const;
+
+	/**
+	 * Set the connection state of a given Input Device to something new. This will
+	 * broadcast the OnInputDeviceConnectionChange delegate.
+	 * This should be called by the platform's implementation.
+	 *
+	 * @param DeviceId		The device id that has had a connection change
+	 * @param NewState		The new connection state of the given device
+	 * @return				True if the connection state was set successfully
+	 */
+	virtual bool Internal_SetInputDeviceConnectionState(FInputDeviceId DeviceId, UE::Input::EInputDeviceConnectionState NewState);
+	
+	/**
+	 * Gets the connection state of the given input device.
+	 * 
+	 * @param DeviceId		The device to get the connection state of
+	 * @return				The connection state of the given device. EInputDeviceConnectionState::Unknown if the device is not mapped
+	 */
+	virtual UE::Input::EInputDeviceConnectionState GetInputDeviceConnectionState(const FInputDeviceId DeviceId) const;
+
+	/**
+	 * Maps the given Input Device to the given userID. This will broadcast the OnInputDeviceConnectionChange delegate.
+	 * This should be called by the platform's implementation.
+	 *
+	 * @param DeviceId			The device id to map
+	 * @param UserId			The Platform User that owns the given device
+	 * @param ConnectionState	The connection state of the device
+	 * @return					True if the input device was successfully mapped
+	 */
+	virtual bool Internal_MapInputDeviceToUser(FInputDeviceId DeviceId, FPlatformUserId UserId, UE::Input::EInputDeviceConnectionState ConnectionState);
+
+	/**
+	 * Change the user mapping of the given input device from an old user to a new one.
+	 * This will broadcast the OnInputDevicePairingChange delegate.
+	 * Use this when you know that an input device is already mapped, but it has changed platform users
+	 * This should be called by the platform's implementation.
+	 *
+	 * @param DeviceId		The input device to change the owner on
+	 * @param NewUserId		The new platform user that this input device should be mapped to
+	 * @param OldUserId		The old platform user that this input device is currently mapped to
+	 * @return				True if the device was successfully remapped
+	 */
+	virtual bool Internal_ChangeInputDeviceUserMapping(FInputDeviceId DeviceId, FPlatformUserId NewUserId, FPlatformUserId OldUserId);
+
+	//////////////////////////////////////////////////////////////////////////////
+	// Delegates for listening to input device changes
+	
+	/**
+	 * Callback for handling an Input Device's connection state change.
+	 * 
+	 * @param NewConnectionState	The new connection state of this device
+	 * @param FPlatformUserId		The User ID whose input device has changed
+	 * @param FInputDeviceId		The Input Device ID that has changed connection
+	 */
+	DECLARE_MULTICAST_DELEGATE_ThreeParams(FOnUserInputDeviceConnectionChange, UE::Input::EInputDeviceConnectionState /* NewConnectionState */, FPlatformUserId /* PlatformUserId */, FInputDeviceId /* InputDeviceId */);
+
+	/**
+	 * Callback for handling an Input Device pairing change.
+	 * 
+	 * @param FInputDeviceId	Input device ID
+	 * @param FPlatformUserId	The NewUserPlatformId
+	 * @param FPlatformUserId	The OldUserPlatformId
+	 */
+	DECLARE_MULTICAST_DELEGATE_ThreeParams(FOnUserInputDevicePairingChange, FInputDeviceId /* InputDeviceId */, FPlatformUserId /* NewUserPlatformId */, FPlatformUserId /* OldUserPlatformId */);
+
+	FOnUserInputDeviceConnectionChange& GetOnInputDeviceConnectionChange() const { return OnInputDeviceConnectionChange; }
+	FOnUserInputDevicePairingChange& GetOnInputDevicePairingChange() const { return OnInputDevicePairingChange; }
+	
+	//////////////////////////////////////////////////////////////////////////////
+	// Functions to provide compatibility between the old "int32 ControllerId"
+	// and the new FPlatformUserId and FInputDeviceId structs.
+	
+	virtual bool RemapUserAndDeviceToControllerId(FPlatformUserId UserId, int32& OutControllerId, FInputDeviceId OptionalDeviceId = INPUTDEVICEID_NONE) = 0;
+	
+	/**
+	 * Remap the legacy "int32 ControllerId" to the updated FPlatformUserId and FInputDeviceId. Use this function to
+	 * add compatibility to platforms that may not have implemented this device mapper yet.
+	 *
+	 * This is useful for functions such as FGenericApplicationMessageHandler::OnControllerAnalog that used to
+	 * use the "int32 ControllerId" as a parameter so that you can call the new FGenericApplicationMessageHandler
+	 * that take in a PlatformUserId and an InputDeviceId.
+	 * 
+	 * @param ControllerId			The old plain "int32" that represented gamepad id or user id depending on the context.
+	 * @param InOutUserId			If the old function provides a PlatformId then pass it here, otherwise pass PLATFORMUSERID_NONE.
+	 * @param OutInputDeviceId		The best guess for an InputDeviceId based on the legacy int32 ControllerId. This may be INPUTDEVICEID_NONE.
+	 * @return						True if this maps to a real user
+	 */
+	virtual bool RemapControllerIdToPlatformUserAndDevice(int32 ControllerId, FPlatformUserId& InOutUserId, FInputDeviceId& OutInputDeviceId) = 0;
+
+protected:
+
+	/**
+	 * If true, this device mapper is operating in a backward compatible mode where there is a
+	 * 1:1 mapping between controller id and user id
+	 */
+	virtual bool IsUsingControllerIdAsUserId() const = 0;
+
+	/**
+	 * If true, than this device mapper will broadcast the older 
+	 * CoreDelegates as well as the new delegates. Set this to
+	 * true if your platform needs calls from OnControllerConnectionChange or
+	 * OnControllerPairingChange
+	 */
+	virtual bool ShouldBroadcastLegacyDelegates() const = 0;
+	
+	/** Allocates a new user id when a user becomes active, will return none if no more can be created */
+	virtual FPlatformUserId AllocateNewUserId() = 0;
+
+	/** Returns the next available input device id. This ID should be globally unique! */
+	virtual FInputDeviceId AllocateNewInputDeviceId() = 0;
+
+	/** Callback when input devices are disconnected/reconnected */
+	static FOnUserInputDeviceConnectionChange OnInputDeviceConnectionChange;
+
+	/** Callback when an input device's owning platform user pairing changes */
+	static FOnUserInputDevicePairingChange OnInputDevicePairingChange;
+	
+	/** A map of all input devices to their current state */
+	TMap<FInputDeviceId, UE::Input::FInputDeviceState> MappedInputDevices;
+
+	/** Highest used platform user id. Incremented in AllocateNewUserId and Internal_MapInputDeviceToUser by default. */
+	FPlatformUserId LastPlatformUserId = PLATFORMUSERID_NONE;
+
+	/** Highest used input device id. Incremented in AllocateNewInputDeviceId and Internal_MapInputDeviceToUser by default. */
+	FInputDeviceId LastInputDeviceId = INPUTDEVICEID_NONE;
+};
+
+/**
+ * Generic implementation of the IPlatformInputDeviceMapper.
+ * This provides the base functionality that can be used on most platforms.
+ */
+class APPLICATIONCORE_API FGenericPlatformInputDeviceMapper : public IPlatformInputDeviceMapper
+{
+public:
+
+	FGenericPlatformInputDeviceMapper(const bool InbUsingControllerIdAsUserId, const bool InbShouldBroadcastLegacyDelegates);
+	
+	/** This is unsupported by default and will return PLATFORMUSERID_NONE on the generic platform */  
+	virtual FPlatformUserId GetUserForUnpairedInputDevices() const override;
+	virtual FInputDeviceId GetDefaultInputDevice() const override;
+
+	virtual bool RemapControllerIdToPlatformUserAndDevice(int32 ControllerId, FPlatformUserId& InOutUserId, FInputDeviceId& OutInputDeviceId) override;
+	virtual bool RemapUserAndDeviceToControllerId(FPlatformUserId UserId, int32& OutControllerId, FInputDeviceId OptionalDeviceId = INPUTDEVICEID_NONE) override;
+	virtual bool IsUsingControllerIdAsUserId() const override;
+	virtual bool ShouldBroadcastLegacyDelegates() const override;
+	
+protected:
+	
+	/** Allocates a new user id when a user becomes active, will return none if no more can be created */
+	virtual FPlatformUserId AllocateNewUserId() override;
+
+	/** Returns the next available input device id */
+	virtual FInputDeviceId AllocateNewInputDeviceId() override;
+
+	/** Flags for backwards compatibility with the older "int32 ControllerId" implementation */
+	const bool bUsingControllerIdAsUserId = true;
+	const bool bShouldBroadcastLegacyDelegates = true;
+};
\ No newline at end of file
diff --git a/Engine/Source/Runtime/Core/Public/Misc/CoreMiscDefines.h b/Engine/Source/Runtime/Core/Public/Misc/CoreMiscDefines.h
--- a/Engine/Source/Runtime/Core/Public/Misc/CoreMiscDefines.h
+++ b/Engine/Source/Runtime/Core/Public/Misc/CoreMiscDefines.h
@@ -349,3 +349,83 @@ private:
 
 /** Static invalid platform user */
 const FPlatformUserId PLATFORMUSERID_NONE;
+
+/**
+ * Represents a single input device such as a gamepad, keyboard, or mouse.
+ *
+ * Has a globally unique identifier that is assigned by the IPlatformInputDeviceMapper
+ */
+struct FInputDeviceId
+{
+	FORCEINLINE explicit FInputDeviceId()
+		: InternalId(INDEX_NONE)
+	{}
+	
+	FInputDeviceId(const FInputDeviceId&) = default;
+	
+	/** Explicit function to create from an internal id */
+	FORCEINLINE static FInputDeviceId CreateFromInternalId(int32 InInternalId)
+	{
+		FInputDeviceId IdToReturn;
+		IdToReturn.InternalId = InInternalId;
+		return IdToReturn;
+	}
+	
+	FORCEINLINE int32 GetId() const
+	{
+		return InternalId;
+	}
+
+	/** Sees if this is a valid input device */
+	FORCEINLINE bool IsValid() const
+	{
+		return InternalId >= 0;
+	}
+
+	FORCEINLINE bool operator==(const FInputDeviceId& Other) const
+	{
+		return InternalId == Other.InternalId;
+	}
+
+	FORCEINLINE bool operator!=(const FInputDeviceId& Other) const
+	{
+		return InternalId != Other.InternalId;
+	}
+	
+	FORCEINLINE bool operator<(const FInputDeviceId& Other) const
+	{
+		return InternalId < Other.InternalId;
+	}
+
+	FORCEINLINE bool operator<=(const FInputDeviceId& Other) const
+	{
+		return InternalId <= Other.InternalId;
+	}
+
+	FORCEINLINE bool operator>(const FInputDeviceId& Other) const
+	{
+		return InternalId > Other.InternalId;
+	}
+
+	FORCEINLINE bool operator>=(const FInputDeviceId& Other) const
+	{
+		return InternalId >= Other.InternalId;
+	}
+
+	FORCEINLINE friend uint32 GetTypeHash(const FInputDeviceId& InputId)
+	{
+		return InputId.InternalId;
+	}
+	
+private:
+	
+	/**
+	 * Raw id, will be allocated by application layer
+	 * 
+	 * @see IPlatformInputDeviceMapper::AllocateNewInputDeviceId
+	 */
+	int32 InternalId;
+};
+
+/** Static invalid input device. */
+const FInputDeviceId INPUTDEVICEID_NONE;
\ No newline at end of file
