Remove the UE::Input namespace because the types won't be able to be BP exposed if we use namespaces. This is because UHT does not support any namespacing outside of "::Type" for some enums, which wouldn't work for the Device State struct. This renames the FInputDeviceState struct to FPlatformInputDeviceState because without the namespaces it conflicts with an existing slate type now that it outside the namespace.

Move the FPlatformInputDeviceState and EInputDeviceConnectionState declarations to the CoreMiscDefines file so that they can be correctly exported in the NoExports.h file. We need them to be in a different file than the base GenericPlatformInputDeviceMapper.h because we need to keep the APPLICATIONCORE_API on that class, which is not available in the NoExports.h file.

#rnx

diff --git a/Engine/Source/Runtime/ApplicationCore/Private/GenericPlatform/GenericPlatformInputDeviceMapper.cpp b/Engine/Source/Runtime/ApplicationCore/Private/GenericPlatform/GenericPlatformInputDeviceMapper.cpp
--- a/Engine/Source/Runtime/ApplicationCore/Private/GenericPlatform/GenericPlatformInputDeviceMapper.cpp
+++ b/Engine/Source/Runtime/ApplicationCore/Private/GenericPlatform/GenericPlatformInputDeviceMapper.cpp
@@ -28,7 +28,7 @@ int32 IPlatformInputDeviceMapper::GetAllInputDevicesForUser(const FPlatformUserI
 {
 	OutInputDevices.Reset();
 	
-	for (TPair<FInputDeviceId, UE::Input::FInputDeviceState> MappedDevice : MappedInputDevices)
+	for (TPair<FInputDeviceId, FPlatformInputDeviceState> MappedDevice : MappedInputDevices)
 	{
 		if (MappedDevice.Value.OwningPlatformUser == UserId)
 		{
@@ -46,10 +46,9 @@ int32 IPlatformInputDeviceMapper::GetAllInputDevices(TArray<FInputDeviceId>& Out
 
 int32 IPlatformInputDeviceMapper::GetAllConnectedInputDevices(TArray<FInputDeviceId>& OutInputDevices) const
 {
-	using namespace UE::Input;
 	OutInputDevices.Reset();
 	
-	for (TPair<FInputDeviceId, FInputDeviceState> MappedDevice : MappedInputDevices)
+	for (TPair<FInputDeviceId, FPlatformInputDeviceState> MappedDevice : MappedInputDevices)
 	{
 		if(MappedDevice.Value.ConnectionState == EInputDeviceConnectionState::Connected)
 		{
@@ -65,7 +64,7 @@ int32 IPlatformInputDeviceMapper::GetAllActiveUsers(TArray<FPlatformUserId>& Out
 	OutUsers.Reset();
 
 	// Add the owning platform user for each input device
-	for (TPair<FInputDeviceId, UE::Input::FInputDeviceState> MappedDevice : MappedInputDevices)
+	for (TPair<FInputDeviceId, FPlatformInputDeviceState> MappedDevice : MappedInputDevices)
 	{
 		OutUsers.AddUnique(MappedDevice.Value.OwningPlatformUser);
 	}
@@ -80,7 +79,7 @@ bool IPlatformInputDeviceMapper::IsUnpairedUserId(const FPlatformUserId Platform
 
 bool IPlatformInputDeviceMapper::IsInputDeviceMappedToUnpairedUser(const FInputDeviceId InputDevice) const
 {
-	if (const UE::Input::FInputDeviceState* DeviceState = MappedInputDevices.Find(InputDevice))
+	if (const FPlatformInputDeviceState* DeviceState = MappedInputDevices.Find(InputDevice))
 	{
 		return IsUnpairedUserId(DeviceState->OwningPlatformUser);
 	}
@@ -89,7 +88,7 @@ bool IPlatformInputDeviceMapper::IsInputDeviceMappedToUnpairedUser(const FInputD
 
 FPlatformUserId IPlatformInputDeviceMapper::GetUserForInputDevice(FInputDeviceId DeviceId) const
 {
-	if (const UE::Input::FInputDeviceState* FoundState = MappedInputDevices.Find(DeviceId))
+	if (const FPlatformInputDeviceState* FoundState = MappedInputDevices.Find(DeviceId))
 	{
 		return FoundState->OwningPlatformUser;
 	}
@@ -101,7 +100,7 @@ FInputDeviceId IPlatformInputDeviceMapper::GetPrimaryInputDeviceForUser(FPlatfor
 	FInputDeviceId FoundDevice = INPUTDEVICEID_NONE;
 
 	// By default look for the lowest input device mapped to this user
-	for (const TPair<FInputDeviceId, UE::Input::FInputDeviceState>& DeviceMapping : MappedInputDevices)
+	for (const TPair<FInputDeviceId, FPlatformInputDeviceState>& DeviceMapping : MappedInputDevices)
 	{
 		if (DeviceMapping.Value.OwningPlatformUser == UserId)
 		{
@@ -115,7 +114,7 @@ FInputDeviceId IPlatformInputDeviceMapper::GetPrimaryInputDeviceForUser(FPlatfor
 	return FoundDevice;
 }
 
-bool IPlatformInputDeviceMapper::Internal_SetInputDeviceConnectionState(FInputDeviceId DeviceId, UE::Input::EInputDeviceConnectionState NewState)
+bool IPlatformInputDeviceMapper::Internal_SetInputDeviceConnectionState(FInputDeviceId DeviceId, EInputDeviceConnectionState NewState)
 {
 	if (!DeviceId.IsValid())
 	{
@@ -138,16 +137,15 @@ bool IPlatformInputDeviceMapper::Internal_SetInputDeviceConnectionState(FInputDe
 	return Internal_MapInputDeviceToUser(DeviceId, OwningUser, NewState);
 }
 
-UE::Input::EInputDeviceConnectionState IPlatformInputDeviceMapper::GetInputDeviceConnectionState(const FInputDeviceId DeviceId) const
+EInputDeviceConnectionState IPlatformInputDeviceMapper::GetInputDeviceConnectionState(const FInputDeviceId DeviceId) const
 {
-	using namespace UE::Input;
 	EInputDeviceConnectionState State = EInputDeviceConnectionState::Unknown;
 
 	if (!DeviceId.IsValid())
 	{
 		State = EInputDeviceConnectionState::Invalid;
 	}
-	else if (const FInputDeviceState* MappedDeviceState = MappedInputDevices.Find(DeviceId))
+	else if (const FPlatformInputDeviceState* MappedDeviceState = MappedInputDevices.Find(DeviceId))
 	{
 		State = MappedDeviceState->ConnectionState;
 	}
@@ -155,7 +153,7 @@ UE::Input::EInputDeviceConnectionState IPlatformInputDeviceMapper::GetInputDevic
 	return State;
 }
 
-bool IPlatformInputDeviceMapper::Internal_MapInputDeviceToUser(FInputDeviceId DeviceId, FPlatformUserId UserId, UE::Input::EInputDeviceConnectionState ConnectionState)
+bool IPlatformInputDeviceMapper::Internal_MapInputDeviceToUser(FInputDeviceId DeviceId, FPlatformUserId UserId, EInputDeviceConnectionState ConnectionState)
 {
 	if (!DeviceId.IsValid())
 	{
@@ -177,8 +175,7 @@ bool IPlatformInputDeviceMapper::Internal_MapInputDeviceToUser(FInputDeviceId De
 	}
 	
 	// Store the connection state of the input device
-	using namespace UE::Input;
-	FInputDeviceState& InputDeviceState = MappedInputDevices.FindOrAdd(DeviceId);
+	FPlatformInputDeviceState& InputDeviceState = MappedInputDevices.FindOrAdd(DeviceId);
 	InputDeviceState.OwningPlatformUser = UserId;
 	InputDeviceState.ConnectionState = ConnectionState;
 	
@@ -202,10 +199,8 @@ bool IPlatformInputDeviceMapper::Internal_ChangeInputDeviceUserMapping(FInputDev
 		return false;
 	}
 
-	using namespace UE::Input;
-
 	// Update the existing device state to be the new owning platform user
-	if (FInputDeviceState* ExistingDeviceState = MappedInputDevices.Find(DeviceId))
+	if (FPlatformInputDeviceState* ExistingDeviceState = MappedInputDevices.Find(DeviceId))
 	{
 		// Only change the platform user of this device if the old user matches up with the one that was given
 		if (ensureMsgf(ExistingDeviceState->OwningPlatformUser == OldUserId, TEXT("Attempting to change the Input Device User Mapping with a mismatched OldUserId!")))
diff --git a/Engine/Source/Runtime/ApplicationCore/Public/GenericPlatform/GenericPlatformInputDeviceMapper.h b/Engine/Source/Runtime/ApplicationCore/Public/GenericPlatform/GenericPlatformInputDeviceMapper.h
--- a/Engine/Source/Runtime/ApplicationCore/Public/GenericPlatform/GenericPlatformInputDeviceMapper.h
+++ b/Engine/Source/Runtime/ApplicationCore/Public/GenericPlatform/GenericPlatformInputDeviceMapper.h
@@ -8,40 +8,6 @@
 #include "Delegates/Delegate.h"
 #include "Logging/LogMacros.h"
 
-namespace UE
-{
-namespace Input
-{
-
-	/** Represents the connection status of a given FInputDeviceId */
-	enum class APPLICATIONCORE_API EInputDeviceConnectionState : uint8
-	{
-		/** This is not a valid input device */
-		Invalid,
-
-		/** It is not known if this device is connected or not */
-		Unknown,
-
-		/** Device is definitely connected */
-		Disconnected,
-
-		/** Definitely connected and powered on */
-		Connected
-	};
-
-	/** Data about an input device's current state */
-	struct APPLICATIONCORE_API FInputDeviceState
-	{
-		/** The platform user that this input device belongs to */
-		FPlatformUserId OwningPlatformUser = PLATFORMUSERID_NONE;
-
-		/** The connection state of this input device */
-		EInputDeviceConnectionState ConnectionState = EInputDeviceConnectionState::Invalid;
-	};
-	
-}	// namespace Input
-}	// namespace UE
-
 DECLARE_LOG_CATEGORY_EXTERN(LogInputDeviceMapper, Log, All);
 
 /**
@@ -129,7 +95,7 @@ public:
 	 * @param NewState		The new connection state of the given device
 	 * @return				True if the connection state was set successfully
 	 */
-	virtual bool Internal_SetInputDeviceConnectionState(FInputDeviceId DeviceId, UE::Input::EInputDeviceConnectionState NewState);
+	virtual bool Internal_SetInputDeviceConnectionState(FInputDeviceId DeviceId, EInputDeviceConnectionState NewState);
 	
 	/**
 	 * Gets the connection state of the given input device.
@@ -137,7 +103,7 @@ public:
 	 * @param DeviceId		The device to get the connection state of
 	 * @return				The connection state of the given device. EInputDeviceConnectionState::Unknown if the device is not mapped
 	 */
-	virtual UE::Input::EInputDeviceConnectionState GetInputDeviceConnectionState(const FInputDeviceId DeviceId) const;
+	virtual EInputDeviceConnectionState GetInputDeviceConnectionState(const FInputDeviceId DeviceId) const;
 
 	/**
 	 * Maps the given Input Device to the given userID. This will broadcast the OnInputDeviceConnectionChange delegate.
@@ -148,7 +114,7 @@ public:
 	 * @param ConnectionState	The connection state of the device
 	 * @return					True if the input device was successfully mapped
 	 */
-	virtual bool Internal_MapInputDeviceToUser(FInputDeviceId DeviceId, FPlatformUserId UserId, UE::Input::EInputDeviceConnectionState ConnectionState);
+	virtual bool Internal_MapInputDeviceToUser(FInputDeviceId DeviceId, FPlatformUserId UserId, EInputDeviceConnectionState ConnectionState);
 
 	/**
 	 * Change the user mapping of the given input device from an old user to a new one.
@@ -173,7 +139,7 @@ public:
 	 * @param FPlatformUserId		The User ID whose input device has changed
 	 * @param FInputDeviceId		The Input Device ID that has changed connection
 	 */
-	DECLARE_MULTICAST_DELEGATE_ThreeParams(FOnUserInputDeviceConnectionChange, UE::Input::EInputDeviceConnectionState /* NewConnectionState */, FPlatformUserId /* PlatformUserId */, FInputDeviceId /* InputDeviceId */);
+	DECLARE_MULTICAST_DELEGATE_ThreeParams(FOnUserInputDeviceConnectionChange, EInputDeviceConnectionState /* NewConnectionState */, FPlatformUserId /* PlatformUserId */, FInputDeviceId /* InputDeviceId */);
 
 	/**
 	 * Callback for handling an Input Device pairing change.
@@ -237,7 +203,7 @@ protected:
 	static FOnUserInputDevicePairingChange OnInputDevicePairingChange;
 	
 	/** A map of all input devices to their current state */
-	TMap<FInputDeviceId, UE::Input::FInputDeviceState> MappedInputDevices;
+	TMap<FInputDeviceId, FPlatformInputDeviceState> MappedInputDevices;
 
 	/** Highest used platform user id. Incremented in AllocateNewUserId and Internal_MapInputDeviceToUser by default. */
 	FPlatformUserId LastPlatformUserId = PLATFORMUSERID_NONE;
diff --git a/Engine/Source/Runtime/Core/Public/Misc/CoreMiscDefines.h b/Engine/Source/Runtime/Core/Public/Misc/CoreMiscDefines.h
--- a/Engine/Source/Runtime/Core/Public/Misc/CoreMiscDefines.h
+++ b/Engine/Source/Runtime/Core/Public/Misc/CoreMiscDefines.h
@@ -428,4 +428,30 @@ private:
 };
 
 /** Static invalid input device. */
-const FInputDeviceId INPUTDEVICEID_NONE;
\ No newline at end of file
+const FInputDeviceId INPUTDEVICEID_NONE;
+
+/** Represents the connection status of a given FInputDeviceId */
+enum class EInputDeviceConnectionState : uint8
+{
+	/** This is not a valid input device */
+	Invalid,
+
+	/** It is not known if this device is connected or not */
+	Unknown,
+
+	/** Device is definitely connected */
+	Disconnected,
+
+	/** Definitely connected and powered on */
+	Connected
+};
+
+/** Data about an input device's current state */
+struct FPlatformInputDeviceState
+{
+	/** The platform user that this input device belongs to */
+	FPlatformUserId OwningPlatformUser = PLATFORMUSERID_NONE;
+
+	/** The connection state of this input device */
+	EInputDeviceConnectionState ConnectionState = EInputDeviceConnectionState::Invalid;
+};
\ No newline at end of file
