Update comment in the GenericPlatformInputDeviceMapper to keep the BP lib up to date

#rnx

diff --git a/Engine/Source/Runtime/ApplicationCore/Public/GenericPlatform/GenericPlatformInputDeviceMapper.h b/Engine/Source/Runtime/ApplicationCore/Public/GenericPlatform/GenericPlatformInputDeviceMapper.h
--- a/Engine/Source/Runtime/ApplicationCore/Public/GenericPlatform/GenericPlatformInputDeviceMapper.h
+++ b/Engine/Source/Runtime/ApplicationCore/Public/GenericPlatform/GenericPlatformInputDeviceMapper.h
@@ -19,6 +19,8 @@ DECLARE_LOG_CATEGORY_EXTERN(LogInputDeviceMapper, Log, All);
  * maps input devices to their users. Some platforms may desire to
  * have each new input device assigned to a different user, while others
  * may want multiple input devices associated with a single user.
+ * 
+ * @note Keep any function comments up to date with those in PlatformInputDeviceMapperLibrary.h!
  */
 class APPLICATIONCORE_API IPlatformInputDeviceMapper
 {
