Add r.Mobile.DesiredResX/Y for Android and -mobileresx and -mobileresy commandline overrides

#android

diff --git a/Engine/Source/Runtime/ApplicationCore/Private/Android/AndroidWindow.cpp b/Engine/Source/Runtime/ApplicationCore/Private/Android/AndroidWindow.cpp
--- a/Engine/Source/Runtime/ApplicationCore/Private/Android/AndroidWindow.cpp
+++ b/Engine/Source/Runtime/ApplicationCore/Private/Android/AndroidWindow.cpp
@@ -18,7 +18,7 @@ int GAndroidPropagateAlpha = 0;
 
 struct FCachedWindowRect
 {
-	FCachedWindowRect() : WindowWidth(-1), WindowHeight(-1), WindowInit(false), ContentScaleFactor(-1.0f), Window_EventThread(nullptr)
+	FCachedWindowRect() : WindowWidth(-1), WindowHeight(-1), WindowInit(false), ContentScaleFactor(-1.0f), MobileResX(-1), MobileResY(-1), Window_EventThread(nullptr)
 	{
 	}
 
@@ -26,6 +26,8 @@ struct FCachedWindowRect
 	int32 WindowHeight;
 	bool WindowInit;
 	float ContentScaleFactor;
+	int32 MobileResX;
+	int32 MobileResY;
 	ANativeWindow* Window_EventThread;
 };
 
@@ -269,7 +271,7 @@ void* FAndroidWindow::WaitForHardwareWindow()
 extern bool AndroidThunkCpp_IsOculusMobileApplication();
 #endif
 
-bool FAndroidWindow::IsCachedRectValid(bool bUseEventThreadWindow, const float RequestedContentScaleFactor, ANativeWindow* Window)
+bool FAndroidWindow::IsCachedRectValid(bool bUseEventThreadWindow, const float RequestedContentScaleFactor, const int32 RequestedMobileResX, const int32 RequestedMobileResY, ANativeWindow* Window)
 {
 	// window must be valid when bUseEventThreadWindow and null when !bUseEventThreadWindow.
 	check((Window != nullptr) == bUseEventThreadWindow);
@@ -289,6 +291,18 @@ bool FAndroidWindow::IsCachedRectValid(bool bUseEventThreadWindow, const float R
 		bValidCache = false;
 	}
 
+	if (CachedRect.MobileResX != RequestedMobileResX)
+	{
+		FPlatformMisc::LowLevelOutputDebugStringf(TEXT("***** RequestedMobileResX different %d != %d, not using res cache (%d)"), RequestedMobileResX, CachedWindowRect.MobileResX, (int32)bUseEventThreadWindow);
+		bValidCache = false;
+	}
+
+	if (CachedRect.MobileResY != RequestedMobileResY)
+	{
+		FPlatformMisc::LowLevelOutputDebugStringf(TEXT("***** RequestedMobileResY different %d != %d, not using res cache (%d)"), RequestedMobileResY, CachedWindowRect.MobileResY, (int32)bUseEventThreadWindow);
+		bValidCache = false;
+	}
+
 	if (CachedRect.Window_EventThread != Window)
 	{
 		FPlatformMisc::LowLevelOutputDebugStringf(TEXT("***** Window different, not using res cache (%d)"), (int32)bUseEventThreadWindow);
@@ -304,7 +318,7 @@ bool FAndroidWindow::IsCachedRectValid(bool bUseEventThreadWindow, const float R
 	return bValidCache;
 }
 
-void FAndroidWindow::CacheRect(bool bUseEventThreadWindow, const int32 Width, const int32 Height, const float RequestedContentScaleFactor, ANativeWindow* Window)
+void FAndroidWindow::CacheRect(bool bUseEventThreadWindow, const int32 Width, const int32 Height, const float RequestedContentScaleFactor, const int32 RequestedMobileResX, const int32 RequestedMobileResY, ANativeWindow* Window)
 {
 	check(Window != nullptr || !bUseEventThreadWindow);
 
@@ -314,6 +328,8 @@ void FAndroidWindow::CacheRect(bool bUseEventThreadWindow, const int32 Width, co
 	CachedRect.WindowHeight = Height;
 	CachedRect.WindowInit = true;
 	CachedRect.ContentScaleFactor = RequestedContentScaleFactor;
+	CachedRect.MobileResX = RequestedMobileResX;
+	CachedRect.MobileResY = RequestedMobileResY;
 	CachedRect.Window_EventThread = Window;
 }
 
@@ -341,9 +357,9 @@ FPlatformRect FAndroidWindow::GetScreenRect(bool bUseEventThreadWindow)
 	static const bool bIsOculusMobileApp = AndroidThunkCpp_IsOculusMobileApplication();
 
 	// CSF is a multiplier to 1280x720
-	static IConsoleVariable* CVar = IConsoleManager::Get().FindConsoleVariable(TEXT("r.MobileContentScaleFactor"));
+	static IConsoleVariable* CVarScale = IConsoleManager::Get().FindConsoleVariable(TEXT("r.MobileContentScaleFactor"));
 	// If the app is for Oculus Mobile then always use 0 as ScaleFactor (to match window size).
-	float RequestedContentScaleFactor = bIsOculusMobileApp ? 0.0f : CVar->GetFloat();
+	float RequestedContentScaleFactor = bIsOculusMobileApp ? 0.0f : CVarScale->GetFloat();
 
 	FString CmdLineCSF;
 	if (FParse::Value(FCommandLine::Get(), TEXT("mcsf="), CmdLineCSF, false))
@@ -351,9 +367,24 @@ FPlatformRect FAndroidWindow::GetScreenRect(bool bUseEventThreadWindow)
 		RequestedContentScaleFactor = FCString::Atof(*CmdLineCSF);
 	}
 
+	static IConsoleVariable* CVarResX = IConsoleManager::Get().FindConsoleVariable(TEXT("r.Mobile.DesiredResX"));
+	static IConsoleVariable* CVarResY = IConsoleManager::Get().FindConsoleVariable(TEXT("r.Mobile.DesiredResY"));
+	int32 RequestedResX = bIsOculusMobileApp ? 0 : CVarResX->GetInt();
+	int32 RequestedResY = bIsOculusMobileApp ? 0 : CVarResY->GetInt();
+
+	FString CmdLineMDRes;
+	if (FParse::Value(FCommandLine::Get(), TEXT("mobileresx="), CmdLineMDRes, false))
+	{
+		RequestedResX = FCString::Atoi(*CmdLineMDRes);
+	}
+	if (FParse::Value(FCommandLine::Get(), TEXT("mobileresy="), CmdLineMDRes, false))
+	{
+		RequestedResY = FCString::Atoi(*CmdLineMDRes);
+	}
+
 	// since orientation won't change on Android, use cached results if still valid. Different cache is maintained for event_thread flavor.
 	ANativeWindow* Window = bUseEventThreadWindow ? (ANativeWindow*)FAndroidWindow::GetHardwareWindow_EventThread() : nullptr;
-	bool bComputeRect = !IsCachedRectValid(bUseEventThreadWindow, RequestedContentScaleFactor, Window);
+	bool bComputeRect = !IsCachedRectValid(bUseEventThreadWindow, RequestedContentScaleFactor, RequestedResX, RequestedResY, Window);
 	if (bComputeRect)
 	{
 		// currently hardcoding resolution
@@ -372,7 +403,7 @@ FPlatformRect FAndroidWindow::GetScreenRect(bool bUseEventThreadWindow)
 		}
 
 		// save for future calls
-		CacheRect(bUseEventThreadWindow, ScreenWidth, ScreenHeight, RequestedContentScaleFactor, Window);
+		CacheRect(bUseEventThreadWindow, ScreenWidth, ScreenHeight, RequestedContentScaleFactor, RequestedResX, RequestedResY, Window);
 	}
 
 	const FCachedWindowRect& CachedRect = bUseEventThreadWindow ? CachedWindowRect_EventThread : CachedWindowRect;
@@ -399,7 +430,7 @@ void FAndroidWindow::CalculateSurfaceSize(int32_t& SurfaceWidth, int32_t& Surfac
 	// too much of the following code needs JNI things, just assume override
 #if !USE_ANDROID_JNI
 
-	UE_LOG(LogAndroid, Fatal, TEXT("FAndroidWindow::CalculateSurfaceSize currently expedcts non-JNI platforms to override resolution"));
+	UE_LOG(LogAndroid, Fatal, TEXT("FAndroidWindow::CalculateSurfaceSize currently expects non-JNI platforms to override resolution"));
 
 #else
 
diff --git a/Engine/Source/Runtime/ApplicationCore/Public/Android/AndroidWindow.h b/Engine/Source/Runtime/ApplicationCore/Public/Android/AndroidWindow.h
--- a/Engine/Source/Runtime/ApplicationCore/Public/Android/AndroidWindow.h
+++ b/Engine/Source/Runtime/ApplicationCore/Public/Android/AndroidWindow.h
@@ -75,9 +75,9 @@ protected:
 private:
 	/** called from GetScreenRect function */
 	/** test cached values from the latest computations stored by CacheRect to decide their validity with the provided arguments */
-	static bool IsCachedRectValid(bool bUseEventThreadWindow, const float RequestedContentScaleFactor, ANativeWindow* Window);
+	static bool IsCachedRectValid(bool bUseEventThreadWindow, const float RequestedContentScaleFactor, const int32 RequestedMobileResX, const int32 RequestedMobileResY, ANativeWindow* Window);
 	/** caches some values used to compute the size of the window by GetScreenRect function */
-	static void CacheRect(bool bUseEventThreadWindow, const int32 Width, const int32 Height, const float RequestedContentScaleFactor, ANativeWindow* Window);
+	static void CacheRect(bool bUseEventThreadWindow, const int32 Width, const int32 Height, const float RequestedContentScaleFactor, const int32 RequestedMobileResX, const int32 RequestedMobileResY, ANativeWindow* Window);
 
 	/**
 	 * Protect the constructor; only TSharedRefs of this class can be made.
diff --git a/Engine/Source/Runtime/ApplicationCore/Public/Android/AndroidWindowUtils.h b/Engine/Source/Runtime/ApplicationCore/Public/Android/AndroidWindowUtils.h
--- a/Engine/Source/Runtime/ApplicationCore/Public/Android/AndroidWindowUtils.h
+++ b/Engine/Source/Runtime/ApplicationCore/Public/Android/AndroidWindowUtils.h
@@ -20,12 +20,14 @@ namespace AndroidWindowUtils
 {
 	static void ApplyContentScaleFactor(int32& InOutScreenWidth, int32& InOutScreenHeight)
 	{
-		const float AspectRatio = (float)InOutScreenWidth / (float)InOutScreenHeight;
-
 		// CSF is a multiplier to 1280x720
-		static IConsoleVariable* CVar = IConsoleManager::Get().FindConsoleVariable(TEXT("r.MobileContentScaleFactor"));		
+		static IConsoleVariable* CVarScale = IConsoleManager::Get().FindConsoleVariable(TEXT("r.MobileContentScaleFactor"));		
+		float RequestedContentScaleFactor = CVarScale->GetFloat();
 
-		float RequestedContentScaleFactor = CVar->GetFloat();
+		static IConsoleVariable* CVarResX = IConsoleManager::Get().FindConsoleVariable(TEXT("r.Mobile.DesiredResX"));
+		static IConsoleVariable* CVarResY = IConsoleManager::Get().FindConsoleVariable(TEXT("r.Mobile.DesiredResY"));
+		int32 RequestedResX = CVarResX->GetInt();
+		int32 RequestedResY = CVarResY->GetInt();
 
 		FString CmdLineCSF;
 		if (FParse::Value(FCommandLine::Get(), TEXT("mcsf="), CmdLineCSF, false))
@@ -33,8 +35,18 @@ namespace AndroidWindowUtils
 			RequestedContentScaleFactor = FCString::Atof(*CmdLineCSF);
 		}
 
+		FString CmdLineMDRes;
+		if (FParse::Value(FCommandLine::Get(), TEXT("mobileresx="), CmdLineMDRes, false))
+		{
+			RequestedResX = FCString::Atoi(*CmdLineMDRes);
+		}
+		if (FParse::Value(FCommandLine::Get(), TEXT("mobileresy="), CmdLineMDRes, false))
+		{
+			RequestedResY = FCString::Atoi(*CmdLineMDRes);
+		}
+
 		// 0 means to use native size
-		if (RequestedContentScaleFactor == 0.0f)
+		if (RequestedContentScaleFactor == 0.0f && RequestedResX <= 0 && RequestedResY <= 0)
 		{
 			UE_LOG(LogAndroidWindowUtils, Log, TEXT("Setting Width=%d and Height=%d (requested scale = 0 = auto)"), InOutScreenWidth, InOutScreenHeight);
 		}
@@ -43,17 +55,52 @@ namespace AndroidWindowUtils
 			int32 Width = InOutScreenWidth;
 			int32 Height = InOutScreenHeight;
 
-			if (InOutScreenHeight > InOutScreenWidth)
+			if (RequestedResX > 0)
+			{
+				// set long side for orientation to requested X (height > width == GAndroidIsPortrait)
+				if (InOutScreenHeight > InOutScreenWidth)
+				{
+					Height = RequestedResX;
+					Width = (Height * ((float)InOutScreenWidth / (float)InOutScreenHeight) + 0.5f);
+				}
+				else
+				{
+					Width = RequestedResX;
+					Height = (Width * ((float)InOutScreenHeight / (float)InOutScreenWidth) + 0.5f);
+				}
+			}
+			else if (RequestedResY > 0)
 			{
-				Height = 1280 * RequestedContentScaleFactor;
+				// set short side for orientation to requested Y (height > width == GAndroidIsPortrait)
+				if (InOutScreenHeight > InOutScreenWidth)
+				{
+					Width = RequestedResY;
+					Height = (Width * ((float)InOutScreenHeight / (float)InOutScreenWidth) + 0.5f);
+				}
+				else
+				{
+					Height = RequestedResY;
+					Width = (Height * ((float)InOutScreenWidth / (float)InOutScreenHeight) + 0.5f);
+				}
 			}
 			else
 			{
-				Height = 720 * RequestedContentScaleFactor;
+				const float AspectRatio = (float)InOutScreenWidth / (float)InOutScreenHeight;
+
+				// set resolution relative to 1280x720 ((height > width == GAndroidIsPortrait)
+				if (InOutScreenHeight > InOutScreenWidth)
+				{
+					Height = 1280 * RequestedContentScaleFactor;
+				}
+				else
+				{
+					Height = 720 * RequestedContentScaleFactor;
+				}
+
+				// apply the aspect ration to get the width
+				Width = (Height * AspectRatio + 0.5f);
 			}
 
-			// apply the aspect ration to get the width
-			Width = (Height * AspectRatio + 0.5f);
 			// ensure Width and Height is multiple of 8
 			Width = (Width / 8) * 8;
 			Height = (Height / 8) * 8;
diff --git a/Engine/Source/Runtime/Core/Private/HAL/ConsoleManager.cpp b/Engine/Source/Runtime/Core/Private/HAL/ConsoleManager.cpp
--- a/Engine/Source/Runtime/Core/Private/HAL/ConsoleManager.cpp
+++ b/Engine/Source/Runtime/Core/Private/HAL/ConsoleManager.cpp
@@ -2964,6 +2964,18 @@ static TAutoConsoleVariable<float> CVarMobileContentScaleFactor(
 	TEXT("Content scale multiplier (equates to iOS's contentScaleFactor to support Retina displays"),
 	ECVF_Default);
 
+static TAutoConsoleVariable<int32> CVarMobileDesiredResX(
+	TEXT("r.Mobile.DesiredResX"),
+	0,
+	TEXT("Desired mobile X resolution (longest axis) (non-zero == use for X, calculate Y to retain aspect ratio)"),
+	ECVF_Default);
+
+static TAutoConsoleVariable<int32> CVarMobileDesiredResY(
+	TEXT("r.Mobile.DesiredResY"),
+	0,
+	TEXT("Desired mobile Y resolution (shortest axis) (non-zero == use for Y, calculate X to retain aspect ratio)"),
+	ECVF_Default);
+
 static TAutoConsoleVariable<int32> CVarMobileTonemapperUpscale(
 	TEXT("r.MobileTonemapperUpscale"),
 	0,
diff --git a/Engine/Source/Runtime/Launch/Private/Android/AndroidEventManager.cpp b/Engine/Source/Runtime/Launch/Private/Android/AndroidEventManager.cpp
--- a/Engine/Source/Runtime/Launch/Private/Android/AndroidEventManager.cpp
+++ b/Engine/Source/Runtime/Launch/Private/Android/AndroidEventManager.cpp
@@ -199,9 +199,17 @@ FAppEventManager::FAppEventManager():
 {
 	pthread_mutex_init(&QueueMutex, NULL);
 
-	IConsoleVariable* CVar = IConsoleManager::Get().FindConsoleVariable(TEXT("r.MobileContentScaleFactor"));
-	check(CVar);
-	CVar->SetOnChangedCallback(FConsoleVariableDelegate::CreateStatic(&FAppEventManager::OnScaleFactorChanged));
+	IConsoleVariable* CVarScale = IConsoleManager::Get().FindConsoleVariable(TEXT("r.MobileContentScaleFactor"));
+	check(CVarScale);
+	CVarScale->SetOnChangedCallback(FConsoleVariableDelegate::CreateStatic(&FAppEventManager::OnScaleFactorChanged));
+
+	IConsoleVariable* CVarResX = IConsoleManager::Get().FindConsoleVariable(TEXT("r.Mobile.DesiredResX"));
+	check(CVarResX);
+	CVarResX->SetOnChangedCallback(FConsoleVariableDelegate::CreateStatic(&FAppEventManager::OnScaleFactorChanged));
+
+	IConsoleVariable* CVarResY = IConsoleManager::Get().FindConsoleVariable(TEXT("r.Mobile.DesiredResY"));
+	check(CVarResY);
+	CVarResY->SetOnChangedCallback(FConsoleVariableDelegate::CreateStatic(&FAppEventManager::OnScaleFactorChanged));
 }
 
 void FAppEventManager::OnScaleFactorChanged(IConsoleVariable* CVar)
