Mac: Replace uses of the old "int32 ControllerID" in favor of the new platform input device mapper. The new input device mapper will still call the old CoreDelegates until they are deprecated/removed, so there is no behavioral change here. Using the platform input device mapper will allow us to pass specific Input Device ID's around to slate and the game, so that we can properly route input events to and from specific devices

diff --git a/Engine/Source/Runtime/ApplicationCore/Private/Mac/HIDInputInterface.cpp b/Engine/Source/Runtime/ApplicationCore/Private/Mac/HIDInputInterface.cpp
--- a/Engine/Source/Runtime/ApplicationCore/Private/Mac/HIDInputInterface.cpp
+++ b/Engine/Source/Runtime/ApplicationCore/Private/Mac/HIDInputInterface.cpp
@@ -5,6 +5,7 @@
 #include "Misc/CallbackDevice.h"
 #include "Misc/ConfigCacheIni.h"
 #include "Templates/SharedPointer.h"
+#include "GenericPlatform/GenericPlatformInputDeviceMapper.h"
 
 static int32 GetDevicePropertyAsInt32(IOHIDDeviceRef DeviceRef, CFStringRef Property)
 {
@@ -420,6 +421,7 @@ void HIDInputInterface::OnNewHIDController(IOReturn Result, IOHIDDeviceRef Devic
 			DeviceInfo.SetupMappings();
 
 			CFIndex ElementsCount = CFArrayGetCount(ElementsArray);
+			IPlatformInputDeviceMapper& DeviceMapper = IPlatformInputDeviceMapper::Get();
 			for (CFIndex ElementIndex = 0; ElementIndex < ElementsCount; ++ElementIndex)
 			{
 				FHIDElementInfo Element;
@@ -439,7 +441,11 @@ void HIDInputInterface::OnNewHIDController(IOReturn Result, IOHIDDeviceRef Devic
 
 			CFRelease(ElementsArray);
 
-			FCoreDelegates::OnControllerConnectionChange.Broadcast(true, PLATFORMUSERID_NONE, ControllerIndex);
+			FPlatformUserId PlatformUserId = PLATFORMUSERID_NONE;
+			FInputDeviceId DeviceId = INPUTDEVICEID_NONE;
+
+			DeviceMapper.RemapControllerIdToPlatformUserAndDevice(ControllerIndex, OUT PlatformUserId, OUT DeviceId);
+			DeviceMapper.Internal_MapInputDeviceToUser(DeviceId, PlatformUserId, EInputDeviceConnectionState::Connected);
 
 			bIsGamepadAttached = true;
 		}
@@ -622,12 +628,18 @@ void HIDInputInterface::HIDDeviceMatchingCallback(void* Context, IOReturn Result
 void HIDInputInterface::HIDDeviceRemovalCallback(void* Context, IOReturn Result, void* Sender, IOHIDDeviceRef DeviceRef)
 {
 	HIDInputInterface* HIDInput = (HIDInputInterface*)Context;
+	IPlatformInputDeviceMapper& DeviceMapper = IPlatformInputDeviceMapper::Get();
 
 	for (int32 Index = 0; Index < MAX_NUM_HIDINPUT_CONTROLLERS; ++Index)
 	{
 		if(HIDInput->ControllerStates[Index].Device.DeviceRef == DeviceRef)
 		{
-			FCoreDelegates::OnControllerConnectionChange.Broadcast(false, PLATFORMUSERID_NONE, Index);
+			FPlatformUserId PlatformUserId = PLATFORMUSERID_NONE;
+			FInputDeviceId DeviceId = INPUTDEVICEID_NONE;
+
+			DeviceMapper.RemapControllerIdToPlatformUserAndDevice(Index, OUT PlatformUserId, OUT DeviceId);
+			DeviceMapper.Internal_MapInputDeviceToUser(DeviceId, PlatformUserId, EInputDeviceConnectionState::Disconnected);
+
 			HIDInput->ControllerStates[Index].Device.DeviceRef = NULL;
 			break;
 		}
