Linux Application: Replace uses of the old "int32 ControllerID" in favor of the new platform input device mapper. The new input device mapper will still call the old CoreDelegates until they are deprecated/removed, so there is no behavioral change here. Using the platform input device mapper will allow us to pass specific Input Device ID's around to slate and the game, so that we can properly route input events to and from specific devices

diff --git a/Engine/Source/Runtime/ApplicationCore/Private/Linux/LinuxApplication.cpp b/Engine/Source/Runtime/ApplicationCore/Private/Linux/LinuxApplication.cpp
--- a/Engine/Source/Runtime/ApplicationCore/Private/Linux/LinuxApplication.cpp
+++ b/Engine/Source/Runtime/ApplicationCore/Private/Linux/LinuxApplication.cpp
@@ -10,6 +10,7 @@
 #include "Linux/LinuxPlatformApplicationMisc.h"
 #include "IInputDeviceModule.h"
 #include "IHapticDevice.h"
+#include "GenericPlatform/GenericPlatformInputDeviceMapper.h"
 
 //
 // GameController thresholds
@@ -2090,7 +2091,13 @@ void FLinuxApplication::AddGameController(int Index)
 		}
 	}
 
-	FCoreDelegates::OnControllerConnectionChange.Broadcast(true, PLATFORMUSERID_NONE, Id);
+	FPlatformUserId PlatformUserId = PLATFORMUSERID_NONE;
+	FInputDeviceId DeviceId = INPUTDEVICEID_NONE;
+
+	IPlatformInputDeviceMapper& DeviceMapper = IPlatformInputDeviceMapper::Get();
+	DeviceMapper.RemapControllerIdToPlatformUserAndDevice(Id, OUT PlatformUserId, OUT DeviceId);
+
+	DeviceMapper.Internal_MapInputDeviceToUser(DeviceId, PlatformUserId, EInputDeviceConnectionState::Connected);
 }
 
 void FLinuxApplication::RemoveGameController(SDL_JoystickID Id)
@@ -2111,5 +2118,11 @@ void FLinuxApplication::RemoveGameController(SDL_JoystickID Id)
 	SDL_GameControllerClose(ControllerState.Controller);
 	ControllerStates.Remove(Id);
 
-	FCoreDelegates::OnControllerConnectionChange.Broadcast(false, PLATFORMUSERID_NONE, Id);
+	FPlatformUserId PlatformUserId = PLATFORMUSERID_NONE;
+	FInputDeviceId DeviceId = INPUTDEVICEID_NONE;
+
+	IPlatformInputDeviceMapper& DeviceMapper = IPlatformInputDeviceMapper::Get();
+	DeviceMapper.RemapControllerIdToPlatformUserAndDevice(Id, OUT PlatformUserId, OUT DeviceId);
+
+	DeviceMapper.Internal_MapInputDeviceToUser(DeviceId, PlatformUserId, EInputDeviceConnectionState::Disconnected);
 }
