XInputInterface: Replace uses of the old "int32 ControllerID" in favor of the new platform input device mapper. The new input device mapper will still call the old CoreDelegates until they are deprecated/removed, so there is no behavioral change here. Using the platform input device mapper will allow us to pass specific Input Device ID's around to slate and the game, so that we can properly route input events to and from specific devices

diff --git a/Engine/Source/Runtime/ApplicationCore/Private/Windows/XInputInterface.cpp b/Engine/Source/Runtime/ApplicationCore/Private/Windows/XInputInterface.cpp
--- a/Engine/Source/Runtime/ApplicationCore/Private/Windows/XInputInterface.cpp
+++ b/Engine/Source/Runtime/ApplicationCore/Private/Windows/XInputInterface.cpp
@@ -5,6 +5,7 @@
 #include "Misc/CoreDelegates.h"
 #include "Windows/WindowsApplication.h"
 #include "GenericPlatform/GenericApplication.h"
+#include "GenericPlatform/GenericPlatformInputDeviceMapper.h"
 #include "Misc/ConfigCacheIni.h"
 
 #pragma pack (push,8)
@@ -149,14 +150,19 @@ void XInputInterface::SendControllerEvents()
 		{
 			const XINPUT_STATE& XInputState = XInputStates[ControllerIndex];
 
+			IPlatformInputDeviceMapper& DeviceMapper = IPlatformInputDeviceMapper::Get();
+			FPlatformUserId PlatformUser = PLATFORMUSERID_NONE;
+			FInputDeviceId InputDevice = INPUTDEVICEID_NONE;
+			DeviceMapper.RemapControllerIdToPlatformUserAndDevice(ControllerState.ControllerId, OUT PlatformUser, OUT InputDevice);
+
 			// If the controller is connected now but was not before, refresh the information
 			if (!bWasConnected && ControllerState.bIsConnected)
 			{
-				FCoreDelegates::OnControllerConnectionChange.Broadcast(true, PLATFORMUSERID_NONE, ControllerState.ControllerId);
+				DeviceMapper.Internal_MapInputDeviceToUser(InputDevice, PlatformUser, EInputDeviceConnectionState::Connected);
 			}
 			else if (bWasConnected && !ControllerState.bIsConnected)
 			{
-				FCoreDelegates::OnControllerConnectionChange.Broadcast(false, PLATFORMUSERID_NONE, ControllerState.ControllerId);
+				DeviceMapper.Internal_MapInputDeviceToUser(InputDevice, PlatformUser, EInputDeviceConnectionState::Disconnected);
 			}
 			
 			bool CurrentStates[MAX_NUM_CONTROLLER_BUTTONS] = {0};
