Update the Android input interface to utilize the new Platform Input Device mapper instead of the old CoreDelegates for when input device connection changes happen.

Replace uses of the old "int32 ControllerID" in favor of the new platform input device mapper. The new input device mapper will still call the old CoreDelegates until they are deprecated/removed, so there is no behavioral change here. Using the platform input device mapper will allow us to pass specific Input Device ID's around to slate and the game, so that we can properly route input events to and from specific devices.

diff --git a/Engine/Source/Runtime/ApplicationCore/Private/Android/AndroidInputInterface.cpp b/Engine/Source/Runtime/ApplicationCore/Private/Android/AndroidInputInterface.cpp
--- a/Engine/Source/Runtime/ApplicationCore/Private/Android/AndroidInputInterface.cpp
+++ b/Engine/Source/Runtime/ApplicationCore/Private/Android/AndroidInputInterface.cpp
@@ -11,6 +11,7 @@
 #include "HAL/PlatformTime.h"
 #include "HAL/IConsoleManager.h"
 #include "IHapticDevice.h"
+#include "GenericPlatform/GenericPlatformInputDeviceMapper.h"
 
 #ifndef ANDROID_GAMEPAD_TRIGGER_THRESHOLD
 	#define ANDROID_GAMEPAD_TRIGGER_THRESHOLD	0.30f
@@ -130,11 +131,17 @@ FAndroidInputInterface::FAndroidInputInterface(const TSharedRef< FGenericApplica
 
 void FAndroidInputInterface::ResetGamepadAssignments()
 {
+	IPlatformInputDeviceMapper& DeviceMapper = IPlatformInputDeviceMapper::Get();
+	
 	for (int32 DeviceIndex = 0; DeviceIndex < MAX_NUM_CONTROLLERS; DeviceIndex++)
 	{
 		if (DeviceMapping[DeviceIndex].DeviceState == MappingState::Valid)
 		{
-			FCoreDelegates::OnControllerConnectionChange.Broadcast(false, PLATFORMUSERID_NONE, DeviceIndex);
+			FPlatformUserId PlatformUserId = PLATFORMUSERID_NONE;
+			FInputDeviceId DeviceId = INPUTDEVICEID_NONE;
+			
+			DeviceMapper.RemapControllerIdToPlatformUserAndDevice(DeviceIndex, OUT PlatformUserId, OUT DeviceId);
+			DeviceMapper.Internal_MapInputDeviceToUser(DeviceId, PlatformUserId, EInputDeviceConnectionState::Disconnected);
 		}
 
 		DeviceMapping[DeviceIndex].DeviceInfo.DeviceId = 0;
@@ -149,7 +156,13 @@ void FAndroidInputInterface::ResetGamepadAssignmentToController(int32 Controller
 
 	if (DeviceMapping[ControllerId].DeviceState == MappingState::Valid)
 	{
-		FCoreDelegates::OnControllerConnectionChange.Broadcast(false, PLATFORMUSERID_NONE, ControllerId);
+		IPlatformInputDeviceMapper& DeviceMapper = IPlatformInputDeviceMapper::Get();
+
+		FPlatformUserId PlatformUserId = PLATFORMUSERID_NONE;
+		FInputDeviceId DeviceId = INPUTDEVICEID_NONE;
+			
+		DeviceMapper.RemapControllerIdToPlatformUserAndDevice(ControllerId, OUT PlatformUserId, OUT DeviceId);
+		DeviceMapper.Internal_MapInputDeviceToUser(DeviceId, PlatformUserId, EInputDeviceConnectionState::Disconnected);
 	}
 
 	DeviceMapping[ControllerId].DeviceInfo.DeviceId = 0;
@@ -983,7 +996,12 @@ void FAndroidInputInterface::SendControllerEvents()
 							CurrentDevice.bSupportsHat = true;
 						}
 
-						FCoreDelegates::OnControllerConnectionChange.Broadcast(true, PLATFORMUSERID_NONE, DeviceIndex);
+						IPlatformInputDeviceMapper& DeviceMapper = IPlatformInputDeviceMapper::Get();
+						FPlatformUserId PlatformUserId = PLATFORMUSERID_NONE;
+						FInputDeviceId DeviceId = INPUTDEVICEID_NONE;
+			
+						DeviceMapper.RemapControllerIdToPlatformUser	AndDevice(DeviceIndex, OUT PlatformUserId, OUT DeviceId);
+						DeviceMapper.Internal_MapInputDeviceToUser(DeviceId, PlatformUserId, EInputDeviceConnectionState::Connected);
 
 						FPlatformMisc::LowLevelOutputDebugStringf(TEXT("Assigned new gamepad controller %d: DeviceId=%d, ControllerId=%d, DeviceName=%s, Descriptor=%s"),
 							DeviceIndex, CurrentDevice.DeviceInfo.DeviceId, CurrentDevice.DeviceInfo.ControllerId, *CurrentDevice.DeviceInfo.Name, *CurrentDevice.DeviceInfo.Descriptor);
@@ -1001,8 +1019,14 @@ void FAndroidInputInterface::SendControllerEvents()
 						NewControllerData[FoundMatch].DeviceId = FoundMatch;
 						OldControllerData[FoundMatch].DeviceId = FoundMatch;
 
-						//@TODO: uncomment this line in the future when disconnects are detected
-						//FCoreDelegates::OnControllerConnectionChange.Broadcast(true, PLATFORMUSERID_NONE, FoundMatch);
+						//@TODO: uncomment these line in the future when disconnects are detected
+						// IPlatformInputDeviceMapper& DeviceMapper = IPlatformInputDeviceMapper::Get();
+						//
+						// FPlatformUserId PlatformUserId = PLATFORMUSERID_NONE;
+						// FInputDeviceId DeviceId = INPUTDEVICEID_NONE;
+						//
+						// DeviceMapper.RemapControllerIdToPlatformUserAndDevice(FoundMatch, OUT PlatformUserId, OUT DeviceId);
+						// DeviceMapper.Internal_MapInputDeviceToUser(DeviceId, PlatformUserId, EInputDeviceConnectionState::Connected);
 
 						FPlatformMisc::LowLevelOutputDebugStringf(TEXT("Reconnected gamepad controller %d: DeviceId=%d, ControllerId=%d, DeviceName=%s, Descriptor=%s"),
 							FoundMatch, DeviceMapping[FoundMatch].DeviceInfo.DeviceId, CurrentDevice.DeviceInfo.ControllerId, *CurrentDevice.DeviceInfo.Name, *CurrentDevice.DeviceInfo.Descriptor);
