Adding FirstMove events to windows, enabled with "WindowsApplication.EnableFirstTouchEvent"

diff --git a/Engine/Source/Runtime/ApplicationCore/Private/Windows/WindowsApplication.cpp b/Engine/Source/Runtime/ApplicationCore/Private/Windows/WindowsApplication.cpp
--- a/Engine/Source/Runtime/ApplicationCore/Private/Windows/WindowsApplication.cpp
+++ b/Engine/Source/Runtime/ApplicationCore/Private/Windows/WindowsApplication.cpp
@@ -83,6 +83,16 @@ static int32 ForceRawInputSimulation = false;
 static int32 EnableRawInputSimulationOverRDP = false;
 #endif
 
+/* 
+* Enabling first touch event will prevent small pop on some touch input devices.
+*/
+static bool bEnableFirstTouchEvent = false;
+FAutoConsoleVariableRef CVarEnableFirstTouchEvent(
+	TEXT("WindowsApplication.EnableFirstTouchEvent"),
+	bEnableFirstTouchEvent,
+	TEXT("Enable FirstTouch Event which prevents small pop on some touch input devices")
+);
+
 const FIntPoint FWindowsApplication::MinimizedWindowPosition(-32000,-32000);
 
 FWindowsApplication* WindowsApplication = nullptr;
@@ -2308,7 +2318,10 @@ int32 FWindowsApplication::ProcessDeferredMessage( const FDeferredWindowsMessage
 									TouchIndex = GetFirstFreeTouchIndex();
 									check(TouchIndex >= 0);
 									
-									TouchIDs[TouchIndex] = TOptional<int32>(Input.dwID);
+									TouchInfoArray[TouchIndex].TouchID = TOptional<int32>(Input.dwID);
+									TouchInfoArray[TouchIndex].HasMoved = false;
+									TouchInfoArray[TouchIndex].PreviousLocation = Location;
+
 									UE_LOG(LogWindowsDesktop, Verbose, TEXT("OnTouchStarted at (%f, %f), finger %d (system touch id %d)"), Location.X, Location.Y, TouchIndex, Input.dwID);
 									MessageHandler->OnTouchStarted(CurrentNativeEventWindowPtr, Location, 1.0f, TouchIndex, 0);
 								}
@@ -2323,6 +2336,21 @@ int32 FWindowsApplication::ProcessDeferredMessage( const FDeferredWindowsMessage
 								if ( TouchIndex >= 0 )
 								{
 									UE_LOG(LogWindowsDesktop, Verbose, TEXT("OnTouchMoved at (%f, %f), finger %d (system touch id %d)"), Location.X, Location.Y, TouchIndex, Input.dwID);
+
+									if (bEnableFirstTouchEvent)
+									{
+										// track first move event, for helping with "pop" on the filtered small movements
+										if (!TouchInfoArray[TouchIndex].HasMoved)
+										{
+											if (TouchInfoArray[TouchIndex].PreviousLocation != Location)
+											{
+												TouchInfoArray[TouchIndex].HasMoved = true;
+												MessageHandler->OnTouchFirstMove(Location, 1.0f, TouchIndex, 0);
+											}
+										}
+									}
+
+									TouchInfoArray[TouchIndex].PreviousLocation = Location;
 									MessageHandler->OnTouchMoved(Location, 1.0f, TouchIndex, 0);
 								}
 							}
@@ -2331,7 +2359,7 @@ int32 FWindowsApplication::ProcessDeferredMessage( const FDeferredWindowsMessage
 								int32 TouchIndex = GetTouchIndexForID( Input.dwID );
 								if ( TouchIndex >= 0 )
 								{
-									TouchIDs[TouchIndex] = TOptional<int32>();
+									TouchInfoArray[TouchIndex].TouchID = TOptional<int32>();
 									UE_LOG(LogWindowsDesktop, Verbose, TEXT("OnTouchEnded at (%f, %f), finger %d (system touch id %d)"), Location.X, Location.Y, TouchIndex, Input.dwID);
 									MessageHandler->OnTouchEnded(Location, TouchIndex, 0);
 								}
@@ -3032,9 +3060,9 @@ void FWindowsApplication::QueryConnectedMice()
 
 uint32 FWindowsApplication::GetTouchIndexForID( int32 TouchID )
 {
-	for (int i = 0; i < TouchIDs.Num(); i++)
+	for (int i = 0; i < TouchInfoArray.Num(); i++)
 	{
-		if ( TouchIDs[i].IsSet() && TouchIDs[i].GetValue() == TouchID )
+		if (TouchInfoArray[i].TouchID.IsSet() && TouchInfoArray[i].TouchID.GetValue() == TouchID )
 		{
 			return i;
 		}
@@ -3044,15 +3072,15 @@ uint32 FWindowsApplication::GetTouchIndexForID( int32 TouchID )
 
 uint32 FWindowsApplication::GetFirstFreeTouchIndex()
 {
-	for ( int i = 0; i < TouchIDs.Num(); i++ )
+	for ( int i = 0; i < TouchInfoArray.Num(); i++ )
 	{
-		if ( TouchIDs[i].IsSet() == false )
+		if (!TouchInfoArray[i].TouchID.IsSet())
 		{
 			return i;
 		}
 	}
 
-	return TouchIDs.Add(TOptional<int32>());
+	return TouchInfoArray.Add(TouchInfo());
 }
 
 void FTaskbarList::Initialize()
diff --git a/Engine/Source/Runtime/ApplicationCore/Public/Windows/WindowsApplication.h b/Engine/Source/Runtime/ApplicationCore/Public/Windows/WindowsApplication.h
--- a/Engine/Source/Runtime/ApplicationCore/Public/Windows/WindowsApplication.h
+++ b/Engine/Source/Runtime/ApplicationCore/Public/Windows/WindowsApplication.h
@@ -254,12 +254,10 @@ public:
 	POINTL CursorPosition;
 };
 
-
 //disable warnings from overriding the deprecated force feedback.  
 //calls to the deprecated function will still generate warnings.
 PRAGMA_DISABLE_DEPRECATION_WARNINGS
 
-
 /**
  * Interface for classes that handle Windows events.
  */
@@ -531,8 +529,19 @@ private:
 	TOGGLEKEYS							StartupToggleKeys;
 	FILTERKEYS							StartupFilterKeys;
 
-	/** Maps touch indexes to windows touch IDs. */
-	TArray<TOptional<int32>> TouchIDs;
+	struct TouchInfo
+	{
+		bool HasMoved;
+		FVector2D PreviousLocation;
+		TOptional<int32> TouchID;
+
+		TouchInfo()
+			: HasMoved(false)
+			, PreviousLocation(0.f, 0.f)
+		{ }
+	};
+	/** Maps touch information such as TouchID PreviousLocation and HasMoved to windows touch IDs. */
+	TArray<TouchInfo> TouchInfoArray;
 
 	bool bSimulatingHighPrecisionMouseInputForRDP;
 
