Deprecate the core delegates for OnControllerConnectionChange and OnControllerPairingChange in favor of the new IPlatformInputDeviceMapper versions that will allow you to get info about what gamepad specifically has sent input events.

diff --git a/Engine/Source/Runtime/ApplicationCore/Private/GenericPlatform/GenericPlatformInputDeviceMapper.cpp b/Engine/Source/Runtime/ApplicationCore/Private/GenericPlatform/GenericPlatformInputDeviceMapper.cpp
--- a/Engine/Source/Runtime/ApplicationCore/Private/GenericPlatform/GenericPlatformInputDeviceMapper.cpp
+++ b/Engine/Source/Runtime/ApplicationCore/Private/GenericPlatform/GenericPlatformInputDeviceMapper.cpp
@@ -185,7 +185,9 @@ bool IPlatformInputDeviceMapper::Internal_MapInputDeviceToUser(FInputDeviceId De
 	if (ShouldBroadcastLegacyDelegates())
 	{
 		const bool bIsConnected = (ConnectionState == EInputDeviceConnectionState::Connected);
+		PRAGMA_DISABLE_DEPRECATION_WARNINGS
 		FCoreDelegates::OnControllerConnectionChange.Broadcast(bIsConnected, UserId, DeviceId.GetId());
+		PRAGMA_ENABLE_DEPRECATION_WARNINGS
 	}
 	
 	return true;
@@ -222,8 +224,9 @@ bool IPlatformInputDeviceMapper::Internal_ChangeInputDeviceUserMapping(FInputDev
 		// Remap the DeviceId to the older int32 "ControllerId" format for the legacy delegates
 		int32 LegacyControllerId = INDEX_NONE;
 		RemapUserAndDeviceToControllerId(NewUserId, LegacyControllerId, DeviceId);
-		
+		PRAGMA_DISABLE_DEPRECATION_WARNINGS
 		FCoreDelegates::OnControllerPairingChange.Broadcast(LegacyControllerId, NewUserId, OldUserId);
+		PRAGMA_ENABLE_DEPRECATION_WARNINGS
 	}
 	return true;
 }
diff --git a/Engine/Source/Runtime/Core/Private/Misc/CoreDelegates.cpp b/Engine/Source/Runtime/Core/Private/Misc/CoreDelegates.cpp
--- a/Engine/Source/Runtime/Core/Private/Misc/CoreDelegates.cpp
+++ b/Engine/Source/Runtime/Core/Private/Misc/CoreDelegates.cpp
@@ -45,8 +45,10 @@ FCoreDelegates::FNoPakFilesMountedDelegate FCoreDelegates::NoPakFilesMountedDele
 FCoreDelegates::FOnFileOpenedForReadFromPakFile FCoreDelegates::OnFileOpenedForReadFromPakFile;
 
 FCoreDelegates::FOnUserLoginChangedEvent FCoreDelegates::OnUserLoginChangedEvent;
+PRAGMA_DISABLE_DEPRECATION_WARNINGS
 FCoreDelegates::FOnUserControllerConnectionChange FCoreDelegates::OnControllerConnectionChange;
 FCoreDelegates::FOnUserControllerPairingChange FCoreDelegates::OnControllerPairingChange;
+PRAGMA_ENABLE_DEPRECATION_WARNINGS
 FCoreDelegates::FOnSafeFrameChangedEvent FCoreDelegates::OnSafeFrameChangedEvent;
 FCoreDelegates::FOnHandleSystemEnsure FCoreDelegates::OnHandleSystemEnsure;
 FCoreDelegates::FOnHandleSystemError FCoreDelegates::OnHandleSystemError;
diff --git a/Engine/Source/Runtime/Core/Public/Misc/CoreDelegates.h b/Engine/Source/Runtime/Core/Public/Misc/CoreDelegates.h
--- a/Engine/Source/Runtime/Core/Public/Misc/CoreDelegates.h
+++ b/Engine/Source/Runtime/Core/Public/Misc/CoreDelegates.h
@@ -156,9 +156,11 @@ public:
 	static FOnUserLoginChangedEvent OnUserLoginChangedEvent;
 
 	// Callback when controllers disconnected / reconnected
+	UE_DEPRECATED(5.1, "OnControllerConnectionChange, use IPlatformInputDeviceMapper::GetOnInputDeviceConnectionChange() instead")
 	static FOnUserControllerConnectionChange OnControllerConnectionChange;
 
 	// Callback when a single controller pairing changes
+	UE_DEPRECATED(5.1, "OnControllerPairingChange, use IPlatformInputDeviceMapper::GetOnInputDevicePairingChange() instead")
 	static FOnUserControllerPairingChange OnControllerPairingChange;
 
 	// Callback when a user changes the safe frame size
