Add a "Primary Platform User" virtual function to allow platforms to define a primary user for things like the keyboard and mouse which there will only be one of. Map the Default Input device to the Primary Platform user upon construction of the device mapper. This gaurantees that there will be a platform user registered for the SlateApplication upon start up.

Add a callback to OnUserLoginChangedEvent to the device mapper. This will provide platforms a hook for when a platform user is logged in or out and allow you to remap any previously paired input devices away from the now logged out user.

diff --git a/Engine/Source/Runtime/ApplicationCore/Private/GenericPlatform/GenericPlatformInputDeviceMapper.cpp b/Engine/Source/Runtime/ApplicationCore/Private/GenericPlatform/GenericPlatformInputDeviceMapper.cpp
--- a/Engine/Source/Runtime/ApplicationCore/Private/GenericPlatform/GenericPlatformInputDeviceMapper.cpp
+++ b/Engine/Source/Runtime/ApplicationCore/Private/GenericPlatform/GenericPlatformInputDeviceMapper.cpp
@@ -24,6 +24,16 @@
 	return *StaticManager;
 }
 
+IPlatformInputDeviceMapper::IPlatformInputDeviceMapper()
+{
+	BindCoreDelegates();
+}
+
+IPlatformInputDeviceMapper::~IPlatformInputDeviceMapper()
+{
+	UnbindCoreDelegates();
+}
+
 int32 IPlatformInputDeviceMapper::GetAllInputDevicesForUser(const FPlatformUserId UserId, TArray<FInputDeviceId>& OutInputDevices) const
 {
 	OutInputDevices.Reset();
@@ -231,17 +241,33 @@
 	return true;
 }
 
+void IPlatformInputDeviceMapper::BindCoreDelegates()
+{
+	FCoreDelegates::OnUserLoginChangedEvent.AddRaw(this, &IPlatformInputDeviceMapper::OnUserLoginChangedEvent);
+}
+
+void IPlatformInputDeviceMapper::UnbindCoreDelegates()
+{
+	FCoreDelegates::OnUserLoginChangedEvent.RemoveAll(this);
+}
+
 ///////////////////////////////////////////////////////////////////////////
 // FGenericPlatformInputDeviceMapper
 
 FGenericPlatformInputDeviceMapper::FGenericPlatformInputDeviceMapper(const bool InbUsingControllerIdAsUserId, const bool InbShouldBroadcastLegacyDelegates)
-	: bUsingControllerIdAsUserId(InbUsingControllerIdAsUserId)
+	: IPlatformInputDeviceMapper()
+	, bUsingControllerIdAsUserId(InbUsingControllerIdAsUserId)
 	, bShouldBroadcastLegacyDelegates(InbShouldBroadcastLegacyDelegates)
 {
 	// Set the last input device id to be the default of 0, that way any new devices will have
 	// an index of 1 or higher and we can use the Default Input Device as a fallback for any
 	// unpaired input devices without an owning PlatformUserId
 	LastInputDeviceId = GetDefaultInputDevice();
+	LastPlatformUserId = GetPrimaryPlatformUser();
+
+	// By default map the Default Input device to the Primary platform user in a connected state. This ensures that the SlateApplication has a 
+	// "Default" user to deal with representing the keyboard and mouse
+	Internal_MapInputDeviceToUser(GetDefaultInputDevice(), GetPrimaryPlatformUser(), EInputDeviceConnectionState::Connected);
 }
 
 FPlatformUserId FGenericPlatformInputDeviceMapper::GetUserForUnpairedInputDevices() const
@@ -252,6 +278,13 @@
 	return PLATFORMUSERID_NONE;
 }
 
+FPlatformUserId FGenericPlatformInputDeviceMapper::GetPrimaryPlatformUser() const
+{
+	// Most platforms will want the primary user to be 0
+	static const FPlatformUserId PrimaryPlatformUser = FPlatformUserId::CreateFromInternalId(0);
+	return PrimaryPlatformUser;
+}
+
 FInputDeviceId FGenericPlatformInputDeviceMapper::GetDefaultInputDevice() const
 {
 	static const FInputDeviceId DefaultInputDeviceId = FInputDeviceId::CreateFromInternalId(0);
@@ -307,6 +340,36 @@
 	return bShouldBroadcastLegacyDelegates;
 }
 
+void FGenericPlatformInputDeviceMapper::OnUserLoginChangedEvent(bool bLoggedIn, int32 RawPlatformUserId, int32 UserIndex)
+{
+	// Attain the FPlatformUserId from the raw int32 that will be given by the Platform code
+	FPlatformUserId LoggedOutPlatformUserId = PLATFORMUSERID_NONE;
+	FInputDeviceId DummyInputDevice = INPUTDEVICEID_NONE;
+	RemapControllerIdToPlatformUserAndDevice(RawPlatformUserId, LoggedOutPlatformUserId, DummyInputDevice);
+	
+	// A new user has logged in
+	if (bLoggedIn)
+	{
+		// TODO: As of right now there is no logic that needs to run when a new platform user logs in, but there may be in the future
+	}
+	// A platform user has logged out
+	else
+	{
+		// Remap any input devices that the logged out platform user had to the "Unpaired" user
+		FPlatformUserId UnkownUserId = GetUserForUnpairedInputDevices();
+		if (LoggedOutPlatformUserId != UnkownUserId)
+		{
+			TArray<FInputDeviceId> InputDevices;
+			GetAllInputDevicesForUser(LoggedOutPlatformUserId, InputDevices);
+
+			for (FInputDeviceId DeviceId : InputDevices)
+			{
+				Internal_ChangeInputDeviceUserMapping(DeviceId, UnkownUserId, LoggedOutPlatformUserId);
+			}
+		}
+	}
+}
+
 FPlatformUserId FGenericPlatformInputDeviceMapper::AllocateNewUserId()
 {
 	// Create a new platform user ID that is 1 higher than the last one
diff --git a/Engine/Source/Runtime/ApplicationCore/Public/GenericPlatform/GenericPlatformInputDeviceMapper.h b/Engine/Source/Runtime/ApplicationCore/Public/GenericPlatform/GenericPlatformInputDeviceMapper.h
--- a/Engine/Source/Runtime/ApplicationCore/Public/GenericPlatform/GenericPlatformInputDeviceMapper.h
+++ b/Engine/Source/Runtime/ApplicationCore/Public/GenericPlatform/GenericPlatformInputDeviceMapper.h
@@ -29,10 +29,10 @@
 	/** Get the platform input device mapper */
 	static IPlatformInputDeviceMapper& Get();
 	
-	IPlatformInputDeviceMapper() = default;
+	IPlatformInputDeviceMapper();
 	
 	/** Virtual destructor */
-	virtual ~IPlatformInputDeviceMapper() = default;
+	virtual ~IPlatformInputDeviceMapper();
 	
 	/**
 	 * Populates the OutInputDevices array with any InputDeviceID's that are mapped to the given platform user
@@ -78,6 +78,14 @@
 
 	/** Returns true if the given input device is mapped to the unpaired platform user id. */
 	virtual bool IsInputDeviceMappedToUnpairedUser(const FInputDeviceId InputDevice) const;
+
+	/**
+	 * Returns the 'Primary' Platform user for this platform.
+	 * This typcially has an internal ID of '0' and is used as the default platform user to
+	 * map devices such as the keyboard and mouse that don't get assigned unique ID's from their
+	 * owning platform code.
+	 */
+	virtual FPlatformUserId GetPrimaryPlatformUser() const = 0;
 	
 	/** Returns the default device id used for things like keyboard/mouse input */
 	virtual FInputDeviceId GetDefaultInputDevice() const = 0;
@@ -179,6 +187,21 @@
 protected:
 
 	/**
+	* Binds to any core delegates that the platform may broadcast. By default it is only the OnUserLoginChangedEvent delegate. 
+	* Called during the constructor of this device mapper.
+	*/
+	virtual void BindCoreDelegates();
+	
+	/** 
+	* Unbind from any delegates that have been hooked into.
+	* Called during the destructor of this device mapper.
+	*/
+	virtual void UnbindCoreDelegates();
+
+	/** Callback for when FCoreDelegates::FOnUserLoginChangedEvent is broadcasted */
+	virtual void OnUserLoginChangedEvent(bool bLoggedIn, int32 UserId, int32 UserIndex) = 0;
+	
+	/**
 	 * If true, this device mapper is operating in a backward compatible mode where there is a
 	 * 1:1 mapping between controller id and user id
 	 */
@@ -226,6 +249,7 @@
 	
 	/** This is unsupported by default and will return PLATFORMUSERID_NONE on the generic platform */  
 	virtual FPlatformUserId GetUserForUnpairedInputDevices() const override;
+	virtual FPlatformUserId GetPrimaryPlatformUser() const override;
 	virtual FInputDeviceId GetDefaultInputDevice() const override;
 
 	virtual bool RemapControllerIdToPlatformUserAndDevice(int32 ControllerId, FPlatformUserId& InOutUserId, FInputDeviceId& OutInputDeviceId) override;
@@ -234,6 +258,12 @@
 	virtual bool ShouldBroadcastLegacyDelegates() const override;
 	
 protected:
+
+	/** 
+	* Callback for when FCoreDelegates::FOnUserLoginChangedEvent is broadcasted.
+	* If the user has logged out, then remap any input devices that the user had to the "Unpaired" user on this platform.
+	*/
+	virtual void OnUserLoginChangedEvent(bool bLoggedIn, int32 RawPlatformUserId, int32 UserIndex) override;
 	
 	/** Allocates a new user id when a user becomes active, will return none if no more can be created */
 	virtual FPlatformUserId AllocateNewUserId() override;
