Remove dependence on UGameInstance for update check.
Fall back to using the first logged in account on beacon connection when no other player is found.
Fix resource leak in windows splash screen.

diff --git a/Engine/Plugins/Online/OnlineFramework/Source/Hotfix/Private/UpdateManager.cpp b/Engine/Plugins/Online/OnlineFramework/Source/Hotfix/Private/UpdateManager.cpp
--- a/Engine/Plugins/Online/OnlineFramework/Source/Hotfix/Private/UpdateManager.cpp
+++ b/Engine/Plugins/Online/OnlineFramework/Source/Hotfix/Private/UpdateManager.cpp
@@ -198,7 +198,7 @@
 	UE_LOG(LogHotfixManager, Display, TEXT("CheckComplete %s"), UpdateCompletionEnum ? *UpdateCompletionEnum->GetNameStringByValue((int64)Result) : TEXT("Invalid"));
 
 	UGameInstance* GameInstance = GetGameInstance();
-	bool bIsServer = GameInstance->IsDedicatedServerInstance();
+	bool bIsServer = GameInstance ? GameInstance->IsDedicatedServerInstance() : false;
 
 #if !UE_BUILD_SHIPPING
 	int32 DbgVal = CVarDebugUpdateManager.GetValueOnGameThread();
@@ -263,11 +263,9 @@
 {
 	ensure(ChecksEnabled());
 
-	UGameInstance* GameInstance = GetGameInstance();
-	check(GameInstance);
-
 	SetUpdateState(EUpdateState::CheckingForPatch);
-	if (GameInstance->IsDedicatedServerInstance() || !bCurrentUpdatePatchCheckEnabled)
+	UGameInstance* GameInstance = GetGameInstance();
+	if ((GameInstance && GameInstance->IsDedicatedServerInstance()) || !bCurrentUpdatePatchCheckEnabled)
 	{
 		PatchCheckComplete(EPatchCheckResult::NoPatchRequired);
 		return;
diff --git a/Engine/Plugins/Online/OnlineSubsystemUtils/Source/OnlineSubsystemUtils/Private/OnlineBeaconClient.cpp b/Engine/Plugins/Online/OnlineSubsystemUtils/Source/OnlineSubsystemUtils/Private/OnlineBeaconClient.cpp
--- a/Engine/Plugins/Online/OnlineSubsystemUtils/Source/OnlineSubsystemUtils/Private/OnlineBeaconClient.cpp
+++ b/Engine/Plugins/Online/OnlineSubsystemUtils/Source/OnlineSubsystemUtils/Private/OnlineBeaconClient.cpp
@@ -127,6 +127,18 @@
 							// Send the player unique Id at login
 							BeaconConnection->PlayerId = LocalPlayer->GetPreferredUniqueNetId();
 						}
+						else
+						{
+							// Fall back to querying the identity interface.
+							if (IOnlineIdentityPtr IdentityPtr = Online::GetIdentityInterface(World))
+							{
+								TArray<TSharedPtr<FUserOnlineAccount> > UserAccounts = IdentityPtr->GetAllUserAccounts();
+								if (!UserAccounts.IsEmpty())
+								{
+									BeaconConnection->PlayerId = UserAccounts[0]->GetUserId();
+								}
+							}
+						}
 					}
 
 #if NETCONNECTION_HAS_SETENCRYPTIONKEY
diff --git a/Engine/Source/Runtime/ApplicationCore/Private/Windows/WindowsPlatformSplash.cpp b/Engine/Source/Runtime/ApplicationCore/Private/Windows/WindowsPlatformSplash.cpp
--- a/Engine/Source/Runtime/ApplicationCore/Private/Windows/WindowsPlatformSplash.cpp
+++ b/Engine/Source/Runtime/ApplicationCore/Private/Windows/WindowsPlatformSplash.cpp
@@ -50,6 +50,7 @@
 static HFONT GSplashScreenNormalTextFontHandle = NULL;
 static HFONT GSplashScreenTitleTextFontHandle = NULL;
 static FCriticalSection GSplashScreenSynchronizationObject;
+static TAtomic<bool> GSplashExitRequested = false;
 
 
 
@@ -543,6 +544,7 @@
 			int32 CurrentOpacityByte = 0;
 
 			MSG message;
+			bool bIsWindowDestructionStarted = false;
 			bool bIsSplashFinished = false;
 			while( !bIsSplashFinished )
 			{
@@ -579,6 +581,12 @@
 					// Give up some time
 					FPlatformProcess::Sleep( 1.0f / 60.0f );
 				}
+
+				if (GSplashExitRequested && !bIsWindowDestructionStarted)
+				{
+					DestroyWindow(GSplashScreenWnd);
+					bIsWindowDestructionStarted = true;
+				}
 			}
 		}
 
@@ -679,6 +687,7 @@
 			}
 			
 			GSplashScreenFileName = SplashPath;
+			GSplashExitRequested = false;
 			DWORD ThreadID = 0;
 			GSplashScreenThread = CreateThread(NULL, 128 * 1024, (LPTHREAD_START_ROUTINE)StartSplashScreenThread, (LPVOID)NULL, STACK_SIZE_PARAM_IS_A_RESERVATION, &ThreadID);
 #if	STATS
@@ -695,8 +704,8 @@
 		TRACE_CPUPROFILER_EVENT_SCOPE(FWindowsPlatformSplash::Hide);
 		if(GSplashScreenWnd)
 		{
-			// Send message to splash screen window to destroy itself
-			PostMessageW(GSplashScreenWnd, WM_DESTROY, 0, 0);
+			// The window must be destroyed from the creating thread to prevent leaking resources.
+			GSplashExitRequested = true;
 		}
 
 		// Wait for splash screen thread to finish
