When mapping the legacy int32 controllerID checkif there is already an allocated platform user id for the input device ID. If there is then we can return that instead of allocating a new one. Add a virtual function that will let you specify that we should create a new platform user id if it is a new input device, mapping each input device to a single platform user. This will be the case on some platforms that don't allow multiple input devices per user.

Update SonyControllers to use the known user index as the platform user, this way we will allocate a new platform user for different input devices.

#rnx

diff --git a/Engine/Source/Runtime/ApplicationCore/Private/GenericPlatform/GenericPlatformInputDeviceMapper.cpp b/Engine/Source/Runtime/ApplicationCore/Private/GenericPlatform/GenericPlatformInputDeviceMapper.cpp
--- a/Engine/Source/Runtime/ApplicationCore/Private/GenericPlatform/GenericPlatformInputDeviceMapper.cpp
+++ b/Engine/Source/Runtime/ApplicationCore/Private/GenericPlatform/GenericPlatformInputDeviceMapper.cpp
@@ -251,6 +251,12 @@
 	FCoreDelegates::OnUserLoginChangedEvent.RemoveAll(this);
 }
 
+bool IPlatformInputDeviceMapper::ShouldCreateUniqueUserForEachDevice() const
+{
+	// TODO: Get this from some settings you can toggle per platform in the editor
+	return true;
+}
+
 ///////////////////////////////////////////////////////////////////////////
 // FGenericPlatformInputDeviceMapper
 
@@ -304,8 +310,36 @@
 		else if (ControllerId >= 0)
 		{
 			// Just use controller id, and copy over device
-			InOutUserId = FPlatformUserId::CreateFromInternalId(ControllerId);
 			OutInputDeviceId = FInputDeviceId::CreateFromInternalId(ControllerId);
+			// If we were already given a valid platform user then we can just stop now.
+			// This will be the case on platforms with an existing concept of "User Logins" from platform itself
+			if (InOutUserId.IsValid())
+			{
+				return true;
+			}
+
+			// If it wasn't valid, then check for a valid known existing user to use
+			FPlatformUserId ExistingUser = GetUserForInputDevice(OutInputDeviceId);
+			if (ExistingUser.IsValid())
+			{
+				InOutUserId = ExistingUser;
+				return true;
+			}
+
+			// Otherwise this is a fresh input device, and we need to create a new PlatformUserId for it.
+			// Some platforms do not have the concept of "User ID"s (i.e. platforms that don't allow having multiple users logged in at once)
+			// For those platforms, they may want to create a new platform user ID for each additional input device that is connected. This would
+			// allow them to create the facade that there is a separation between the connected input devices and their platform users, allowing
+			// gameplay code to differentiate between platform users in a consistent manner.
+			if (ShouldCreateUniqueUserForEachDevice())
+			{
+				InOutUserId = AllocateNewUserId();
+			}
+			else
+			{
+				// Otherwise just have a 1:1 mapping of input device to user id's
+				InOutUserId = FPlatformUserId::CreateFromInternalId(ControllerId);
+			}
 			return true;
 		}
 		else if (InOutUserId.GetInternalId() >= 0)
diff --git a/Engine/Source/Runtime/ApplicationCore/Public/GenericPlatform/GenericPlatformInputDeviceMapper.h b/Engine/Source/Runtime/ApplicationCore/Public/GenericPlatform/GenericPlatformInputDeviceMapper.h
--- a/Engine/Source/Runtime/ApplicationCore/Public/GenericPlatform/GenericPlatformInputDeviceMapper.h
+++ b/Engine/Source/Runtime/ApplicationCore/Public/GenericPlatform/GenericPlatformInputDeviceMapper.h
@@ -208,6 +208,13 @@
 	virtual bool IsUsingControllerIdAsUserId() const = 0;
 
 	/**
+	 * If true, then when we remap the platform's raw int32 ControllerId to its FInputDeviceId
+	 * and owning FPlatformUserId we should create a new platform user id if this controller id is
+	 * not already mapped. 
+	 */
+	virtual bool ShouldCreateUniqueUserForEachDevice() const;
+
+	/**
 	 * If true, than this device mapper will broadcast the older 
 	 * CoreDelegates as well as the new delegates. Set this to
 	 * true if your platform needs calls from OnControllerConnectionChange or
