- Fixed more MacOSX*.sdk -> MacOSX.sdk

diff --git a/Engine/Source/Runtime/ApplicationCore/ApplicationCore.Build.cs b/Engine/Source/Runtime/ApplicationCore/ApplicationCore.Build.cs
--- a/Engine/Source/Runtime/ApplicationCore/ApplicationCore.Build.cs
+++ b/Engine/Source/Runtime/ApplicationCore/ApplicationCore.Build.cs
@@ -1,4 +1,4 @@
-// Copyright Epic Games, Inc. All Rights Reserved.
+﻿// Copyright Epic Games, Inc. All Rights Reserved.
 
 using UnrealBuildTool;
 
@@ -51,8 +51,8 @@ public class ApplicationCore : ModuleRules
 			);
 			if (Target.bBuildEditor == true)
 			{
-				string SDKROOT = Utils.RunLocalProcessAndReturnStdOut("/usr/bin/xcrun", "--sdk macosx --show-sdk-path");
-				PublicAdditionalLibraries.Add(SDKROOT + "/System/Library/PrivateFrameworks/MultitouchSupport.framework/Versions/Current/MultitouchSupport.tbd");
+				string XcodeRoot = Utils.RunLocalProcessAndReturnStdOut("/usr/bin/xcode-select", "--print-path");
+				PublicAdditionalLibraries.Add(XcodeRoot + "/Platforms/MacOSX.platform/Developer/SDKs/MacOSX.sdk/System/Library/PrivateFrameworks/MultitouchSupport.framework/Versions/Current/MultitouchSupport.tbd");
 			}
 			
 			PublicFrameworks.Add("GameController");
diff --git a/Engine/Source/Runtime/Core/Core.Build.cs b/Engine/Source/Runtime/Core/Core.Build.cs
--- a/Engine/Source/Runtime/Core/Core.Build.cs
+++ b/Engine/Source/Runtime/Core/Core.Build.cs
@@ -1,4 +1,4 @@
-// Copyright Epic Games, Inc. All Rights Reserved.
+﻿// Copyright Epic Games, Inc. All Rights Reserved.
 
 using UnrealBuildTool;
 using System;
@@ -107,8 +107,8 @@ public class Core : ModuleRules
 
 			if (Target.bBuildEditor == true)
 			{
-				string SDKROOT = Utils.RunLocalProcessAndReturnStdOut("/usr/bin/xcrun", "--sdk macosx --show-sdk-path");
-				PublicAdditionalLibraries.Add(SDKROOT + "/System/Library/PrivateFrameworks/MultitouchSupport.framework/Versions/Current/MultitouchSupport.tbd");
+				string XcodeRoot = Utils.RunLocalProcessAndReturnStdOut("/usr/bin/xcode-select", "--print-path");
+				PublicAdditionalLibraries.Add(XcodeRoot + "/Platforms/MacOSX.platform/Developer/SDKs/MacOSX.sdk/System/Library/PrivateFrameworks/MultitouchSupport.framework/Versions/Current/MultitouchSupport.tbd");
 			}
 		}
 		else if (Target.Platform == UnrealTargetPlatform.IOS || Target.Platform == UnrealTargetPlatform.TVOS)
