Fix a bug where the app would recursively call itself because there was an implicit cast to int32 from Platform user

#rnx

diff --git a/Engine/Source/Runtime/ApplicationCore/Public/GenericPlatform/GenericApplicationMessageHandler.h b/Engine/Source/Runtime/ApplicationCore/Public/GenericPlatform/GenericApplicationMessageHandler.h
--- a/Engine/Source/Runtime/ApplicationCore/Public/GenericPlatform/GenericApplicationMessageHandler.h
+++ b/Engine/Source/Runtime/ApplicationCore/Public/GenericPlatform/GenericApplicationMessageHandler.h
@@ -527,7 +527,7 @@ public:
 	{
 		if (ShouldUsePlatformUserId())
 		{
-			return OnControllerButtonPressed(KeyName, FPlatformUserId(ControllerId), IsRepeat);
+			return OnControllerButtonPressed(KeyName, FPlatformUserId(ControllerId), INPUTDEVICEID_NONE, IsRepeat);
 		}
 		return false;
 	}
@@ -536,7 +536,7 @@ public:
 	{
 		if (ShouldUsePlatformUserId())
 		{
-			return OnControllerButtonReleased(KeyName, FPlatformUserId(ControllerId), IsRepeat);
+			return OnControllerButtonReleased(KeyName, FPlatformUserId(ControllerId), INPUTDEVICEID_NONE, IsRepeat);
 		}
 		return false;
 	}
