Replace calls to the old OnControllerX functions in slate with the new version that has a PlatformUser/InputDevice

#rnx

diff --git a/Engine/Source/Runtime/ApplicationCore/Private/Windows/XInputInterface.cpp b/Engine/Source/Runtime/ApplicationCore/Private/Windows/XInputInterface.cpp
--- a/Engine/Source/Runtime/ApplicationCore/Private/Windows/XInputInterface.cpp
+++ b/Engine/Source/Runtime/ApplicationCore/Private/Windows/XInputInterface.cpp
@@ -194,10 +194,10 @@ void XInputInterface::SendControllerEvents()
 			CurrentStates[X360ToXboxControllerMapping[23]] = !!(XInputState.Gamepad.sThumbRX > XINPUT_GAMEPAD_RIGHT_THUMB_DEADZONE);
 
 			// Send new analog data if it's different or outside the platform deadzone.
-			auto OnControllerAnalog = [this, &ControllerState](const FName& GamePadKey, const auto NewAxisValue, const float NewAxisValueNormalized, auto& OldAxisValue, const auto DeadZone) {
+			auto OnControllerAnalog = [this, &PlatformUser, &InputDevice](const FName& GamePadKey, const auto NewAxisValue, const float NewAxisValueNormalized, auto& OldAxisValue, const auto DeadZone) {
 				if (OldAxisValue != NewAxisValue || FMath::Abs((int32)NewAxisValue) > DeadZone)
 				{
-					MessageHandler->OnControllerAnalog(GamePadKey, ControllerState.ControllerId, NewAxisValueNormalized);
+					MessageHandler->OnControllerAnalog(GamePadKey, PlatformUser, InputDevice, NewAxisValueNormalized);
 				}
 				OldAxisValue = NewAxisValue;
 			};
@@ -222,11 +222,11 @@ void XInputInterface::SendControllerEvents()
 				{
 					if( CurrentStates[ButtonIndex] )
 					{
-						MessageHandler->OnControllerButtonPressed( Buttons[ButtonIndex], ControllerState.ControllerId, false );
+						MessageHandler->OnControllerButtonPressed( Buttons[ButtonIndex], PlatformUser, InputDevice, false );
 					}
 					else
 					{
-						MessageHandler->OnControllerButtonReleased( Buttons[ButtonIndex], ControllerState.ControllerId, false );
+						MessageHandler->OnControllerButtonReleased( Buttons[ButtonIndex], PlatformUser, InputDevice, false );
 					}
 
 					if ( CurrentStates[ButtonIndex] != 0 )
@@ -237,7 +237,7 @@ void XInputInterface::SendControllerEvents()
 				}
 				else if ( CurrentStates[ButtonIndex] != 0 && ControllerState.NextRepeatTime[ButtonIndex] <= CurrentTime )
 				{
-					MessageHandler->OnControllerButtonPressed( Buttons[ButtonIndex], ControllerState.ControllerId, true );
+					MessageHandler->OnControllerButtonPressed( Buttons[ButtonIndex], PlatformUser, InputDevice, true );
 
 					// set the button's NextRepeatTime to the ButtonRepeatDelay
 					ControllerState.NextRepeatTime[ButtonIndex] = CurrentTime + ButtonRepeatDelay;
diff --git a/Engine/Source/Runtime/Slate/Private/Widgets/Input/SVirtualJoystick.cpp b/Engine/Source/Runtime/Slate/Private/Widgets/Input/SVirtualJoystick.cpp
--- a/Engine/Source/Runtime/Slate/Private/Widgets/Input/SVirtualJoystick.cpp
+++ b/Engine/Source/Runtime/Slate/Private/Widgets/Input/SVirtualJoystick.cpp
@@ -4,7 +4,7 @@
 #include "Rendering/DrawElements.h"
 #include "Misc/ConfigCacheIni.h"
 #include "Framework/Application/SlateApplication.h"
-
+#include "GenericPlatform/GenericPlatformInputDeviceMapper.h"
 
 const float OPACITY_LERP_RATE = 3.f;
 
@@ -415,8 +415,10 @@ void SVirtualJoystick::Tick(const FGeometry& AllottedGeometry, const double InCu
 			const FGamepadKeyNames::Type YAxis = (Control.Info.AltInputKey.IsValid() ? Control.Info.AltInputKey.GetFName() : (ControlIndex == 0 ? FGamepadKeyNames::LeftAnalogY : FGamepadKeyNames::RightAnalogY));
 
 			FSlateApplication::Get().SetAllUserFocusToGameViewport();
-			FSlateApplication::Get().OnControllerAnalog(XAxis, 0, NormalizedOffset.X);
-			FSlateApplication::Get().OnControllerAnalog(YAxis, 0, -NormalizedOffset.Y);
+			
+			FInputDeviceId PrimaryInputDevice = IPlatformInputDeviceMapper::Get().GetPrimaryInputDeviceForUser(FSlateApplicationBase::SlateAppPrimaryPlatformUser);
+			FSlateApplication::Get().OnControllerAnalog(XAxis, FSlateApplicationBase::SlateAppPrimaryPlatformUser, PrimaryInputDevice, NormalizedOffset.X);
+			FSlateApplication::Get().OnControllerAnalog(YAxis, FSlateApplicationBase::SlateAppPrimaryPlatformUser, PrimaryInputDevice, -NormalizedOffset.Y);
 		}
 
 		// is this active?
