Revert changes to WindowsPlatformSplash due to a race condition.

diff --git a/Engine/Source/Runtime/ApplicationCore/Private/Windows/WindowsPlatformSplash.cpp b/Engine/Source/Runtime/ApplicationCore/Private/Windows/WindowsPlatformSplash.cpp
--- a/Engine/Source/Runtime/ApplicationCore/Private/Windows/WindowsPlatformSplash.cpp
+++ b/Engine/Source/Runtime/ApplicationCore/Private/Windows/WindowsPlatformSplash.cpp
@@ -50,7 +50,6 @@
 static HFONT GSplashScreenNormalTextFontHandle = NULL;
 static HFONT GSplashScreenTitleTextFontHandle = NULL;
 static FCriticalSection GSplashScreenSynchronizationObject;
-static TAtomic<bool> GSplashExitRequested = false;
 
 
 
@@ -544,7 +543,6 @@
 			int32 CurrentOpacityByte = 0;
 
 			MSG message;
-			bool bIsWindowDestructionStarted = false;
 			bool bIsSplashFinished = false;
 			while( !bIsSplashFinished )
 			{
@@ -581,12 +579,6 @@
 					// Give up some time
 					FPlatformProcess::Sleep( 1.0f / 60.0f );
 				}
-
-				if (GSplashExitRequested && !bIsWindowDestructionStarted)
-				{
-					DestroyWindow(GSplashScreenWnd);
-					bIsWindowDestructionStarted = true;
-				}
 			}
 		}
 
@@ -687,7 +679,6 @@
 			}
 			
 			GSplashScreenFileName = SplashPath;
-			GSplashExitRequested = false;
 			DWORD ThreadID = 0;
 			GSplashScreenThread = CreateThread(NULL, 128 * 1024, (LPTHREAD_START_ROUTINE)StartSplashScreenThread, (LPVOID)NULL, STACK_SIZE_PARAM_IS_A_RESERVATION, &ThreadID);
 #if	STATS
@@ -704,8 +695,8 @@
 		TRACE_CPUPROFILER_EVENT_SCOPE(FWindowsPlatformSplash::Hide);
 		if(GSplashScreenWnd)
 		{
-			// The window must be destroyed from the creating thread to prevent leaking resources.
-			GSplashExitRequested = true;
+			// Send message to splash screen window to destroy itself
+			PostMessageW(GSplashScreenWnd, WM_DESTROY, 0, 0);
 		}
 
 		// Wait for splash screen thread to finish
