Add r.Mobile.DesiredResX/Y and -mobileresx/y commandline overrides to iOS.  Also clears out some old iOS OpenGL references.

diff --git a/Engine/Source/Runtime/ApplicationCore/Private/IOS/IOSAppDelegate.cpp b/Engine/Source/Runtime/ApplicationCore/Private/IOS/IOSAppDelegate.cpp
--- a/Engine/Source/Runtime/ApplicationCore/Private/IOS/IOSAppDelegate.cpp
+++ b/Engine/Source/Runtime/ApplicationCore/Private/IOS/IOSAppDelegate.cpp
@@ -825,11 +825,16 @@ static IOSAppDelegate* CachedDelegate = nil;
 #endif
 }
 
-- (void)LoadMobileContentScaleFactor
+- (void)LoadScreenResolutionModifiers
 {
+	// cache these UI thread sensitive vars for later use
+	self.ScreenScale = [[UIScreen mainScreen] scale];
+	self.NativeScale = [[UIScreen mainScreen] nativeScale];
+
 	// need to cache the MobileContentScaleFactor for framebuffer creation.
-	static IConsoleVariable* CVar = IConsoleManager::Get().FindConsoleVariable(TEXT("r.MobileContentScaleFactor"));
-	self.MobileContentScaleFactor = CVar ? CVar->GetFloat() : 0;
+	static IConsoleVariable* CVarScale = IConsoleManager::Get().FindConsoleVariable(TEXT("r.MobileContentScaleFactor"));
+	check(CVarScale);
+	self.MobileContentScaleFactor = CVarScale ? CVarScale->GetFloat() : 0;
 
 	// Can also be overridden from the commandline using "mcsf="
 	FString CmdLineCSF;
@@ -837,11 +842,34 @@ static IOSAppDelegate* CachedDelegate = nil;
 	{
 		self.MobileContentScaleFactor = FCString::Atof(*CmdLineCSF);
 	}
-}
 
-- (float)GetMobileContentScaleFactor
-{
-	return self.MobileContentScaleFactor;
+	static IConsoleVariable* CVarResX = IConsoleManager::Get().FindConsoleVariable(TEXT("r.Mobile.DesiredResX"));
+	static IConsoleVariable* CVarResY = IConsoleManager::Get().FindConsoleVariable(TEXT("r.Mobile.DesiredResY"));
+	check(CVarResX);
+	check(CVarResY);
+		
+	self.RequestedResX = CVarResX ? CVarResX->GetInt() : 0;
+	self.RequestedResY = CVarResY ? CVarResY->GetInt() : 0;
+		
+	static bool bOnFirstUse = true;
+	if (bOnFirstUse)
+	{
+		FString CmdLineMDRes;
+		if (FParse::Value(FCommandLine::Get(), TEXT("mobileresx="), CmdLineMDRes, false))
+		{
+				self.RequestedResX = FCString::Atoi(*CmdLineMDRes);
+		}
+		if (FParse::Value(FCommandLine::Get(), TEXT("mobileresy="), CmdLineMDRes, false))
+		{
+				self.RequestedResY = FCString::Atoi(*CmdLineMDRes);
+		}
+
+		CVarScale->SetOnChangedCallback(FConsoleVariableDelegate::CreateStatic(&FIOSWindow::OnScaleFactorChanged));
+		CVarResX->SetOnChangedCallback(FConsoleVariableDelegate::CreateStatic(&FIOSWindow::OnConsoleResolutionChanged));
+		CVarResY->SetOnChangedCallback(FConsoleVariableDelegate::CreateStatic(&FIOSWindow::OnConsoleResolutionChanged));
+			
+		bOnFirstUse = false;
+	}
 }
 
 - (void)CheckForZoomAccessibility
diff --git a/Engine/Source/Runtime/ApplicationCore/Private/IOS/IOSApplication.cpp b/Engine/Source/Runtime/ApplicationCore/Private/IOS/IOSApplication.cpp
--- a/Engine/Source/Runtime/ApplicationCore/Private/IOS/IOSApplication.cpp
+++ b/Engine/Source/Runtime/ApplicationCore/Private/IOS/IOSApplication.cpp
@@ -3,6 +3,7 @@
 #include "IOS/IOSApplication.h"
 #include "IOS/IOSInputInterface.h"
 #include "IOS/IOSCursor.h"
+#include "IOS/IOSView.h"
 #include "IOSWindow.h"
 #include "Misc/CoreDelegates.h"
 #include "IOS/IOSAppDelegate.h"
@@ -141,6 +142,10 @@ UIEdgeInsets CachedInsets;
 
 void FDisplayMetrics::RebuildDisplayMetrics(FDisplayMetrics& OutDisplayMetrics)
 {
+	const FPlatformRect& Rect = FIOSWindow::GetUIWindowRect();
+	const FIOSView *View = [[IOSAppDelegate GetDelegate] IOSView];
+	[View CalculateContentScaleFactor:Rect.Right ScreenHeight:Rect.Bottom];
+	
 	// Get screen rect
 	OutDisplayMetrics.PrimaryDisplayWorkAreaRect = FIOSWindow::GetScreenRect();
 	OutDisplayMetrics.VirtualDisplayRect = OutDisplayMetrics.PrimaryDisplayWorkAreaRect;
@@ -149,11 +154,8 @@ void FDisplayMetrics::RebuildDisplayMetrics(FDisplayMetrics& OutDisplayMetrics)
     OutDisplayMetrics.PrimaryDisplayWidth = OutDisplayMetrics.PrimaryDisplayWorkAreaRect.Right - OutDisplayMetrics.PrimaryDisplayWorkAreaRect.Left;
     OutDisplayMetrics.PrimaryDisplayHeight = OutDisplayMetrics.PrimaryDisplayWorkAreaRect.Bottom - OutDisplayMetrics.PrimaryDisplayWorkAreaRect.Top;
     
-    // Get ui window rect
-    OutDisplayMetrics.IosUiWindowAreaRect = FIOSWindow::GetUIWindowRect();
-    
 #if !PLATFORM_TVOS
-    const float RequestedContentScaleFactor = [[IOSAppDelegate GetDelegate].IOSView contentScaleFactor];
+    const float RequestedContentScaleFactor = View.contentScaleFactor;
     
     //we need to set these according to the orientation
     TAutoConsoleVariable<float>* CVar_Left = nullptr;
diff --git a/Engine/Source/Runtime/ApplicationCore/Private/IOS/IOSView.cpp b/Engine/Source/Runtime/ApplicationCore/Private/IOS/IOSView.cpp
--- a/Engine/Source/Runtime/ApplicationCore/Private/IOS/IOSView.cpp
+++ b/Engine/Source/Runtime/ApplicationCore/Private/IOS/IOSView.cpp
@@ -3,16 +3,13 @@
 #include "IOS/IOSView.h"
 #include "IOS/IOSAppDelegate.h"
 #include "IOS/IOSApplication.h"
-#include "IOSWindow.h"
 #include "Misc/ConfigCacheIni.h"
 #include "HAL/IConsoleManager.h"
 #include "Misc/CommandLine.h"
 #include "IOS/IOSPlatformProcess.h"
 
-#include <OpenGLES/ES2/gl.h>
 #import "IOS/IOSAsyncTask.h"
 #import <QuartzCore/QuartzCore.h>
-#import <OpenGLES/EAGLDrawable.h>
 #import <UIKit/UIGeometry.h>
 
 #include "IOS/IOSCommandLineHelper.h"
@@ -22,9 +19,7 @@
 #include "IOS/Accessibility/IOSAccessibilityElement.h"
 #endif
 
-#if HAS_METAL
 id<MTLDevice> GMetalDevice = nil;
-#endif
 
 
 @interface IndexedPosition : UITextPosition {
@@ -87,9 +82,7 @@ id<MTLDevice> GMetalDevice = nil;
 @synthesize autocorrectionType = AutocorrectionType;
 @synthesize autocapitalizationType = AutocapitalizationType;
 @synthesize secureTextEntry = bSecureTextEntry;
-@synthesize SwapCount, OnScreenColorRenderBuffer, OnScreenColorRenderBufferMSAA, markedTextStyle;
-
-
+@synthesize SwapCount, markedTextStyle;
 
 
 #if BUILD_EMBEDDED_APP
@@ -114,12 +107,6 @@ id<MTLDevice> GMetalDevice = nil;
 
 #endif
 
-
-
-#if !HAS_METAL
-#error HAS_METAL must be defined
-#endif
-
 /**
  * @return The Layer Class for the window
  */
@@ -131,7 +118,6 @@ id<MTLDevice> GMetalDevice = nil;
 	return [CAMetalLayer class];
 #endif
 	
-#if HAS_METAL
 	// make sure the project setting has enabled Metal support (per-project user settings in the editor)
 	bool bSupportsMetal = false;
 	bool bSupportsMetalMRT = false;
@@ -167,7 +153,6 @@ id<MTLDevice> GMetalDevice = nil;
 		return [CAMetalLayer class];
 	}
 	else
-#endif
 	{
 		return nil;
 	}
@@ -179,14 +164,11 @@ id<MTLDevice> GMetalDevice = nil;
 
 	CachedMarkedText = nil;
 
-	// figure out if we should start up GL or Metal
-#if HAS_METAL
+	check(GMetalDevice);
 	// if the device is valid, we know Metal is usable (see +layerClass)
 	MetalDevice = GMetalDevice;
 	if (MetalDevice != nil)
 	{
-		bIsUsingMetal = true;
-		
 		// grab the MetalLayer and typecast it to match what's in layerClass
 		CAMetalLayer* MetalLayer = (CAMetalLayer*)self.layer;
 		MetalLayer.presentsWithTransaction = NO;
@@ -201,25 +183,19 @@ id<MTLDevice> GMetalDevice = nil;
 		MetalLayer.pixelFormat = MTLPixelFormatBGRA8Unorm;
 		MetalLayer.framebufferOnly = NO;
 		
+		NSLog(@"::: Created a UIView that will support Metal :::");
 	}
-#endif
 	
-	NSLog(@"::: Created a UIView that will support %@ :::", bIsUsingMetal ? @"Metal" : @"@GLES");
 	
-	// Initialize some variables
-	SwapCount = 0;
-
-//	self.userInteractionEnabled = YES;
-//	self.clearsContextBeforeDrawing = NO;
 #if !PLATFORM_TVOS
 	self.multipleTouchEnabled = YES;
 #endif
 
+	SwapCount = 0;
 	FMemory::Memzero(AllTouches, sizeof(AllTouches));
 	[self setAutoresizingMask: UIViewAutoresizingFlexibleWidth | UIViewAutoresizingFlexibleHeight];
 	bIsInitialized = false;
 	
-	
 	[self InitKeyboard];
 
 #if BUILD_EMBEDDED_APP
@@ -234,7 +210,7 @@ id<MTLDevice> GMetalDevice = nil;
 	AppDelegate.IOSView = self;
 	
 	// initialize the backbuffer of the view (so the RHI can use it)
-	[self CreateFramebuffer:YES];
+	[self CreateFramebuffer];
 	
 #endif
 	
@@ -267,39 +243,71 @@ id<MTLDevice> GMetalDevice = nil;
 	[super dealloc];
 }
 
-- (bool)CreateFramebuffer:(bool)bIsForOnDevice
+- (void)CalculateContentScaleFactor:(int32)ScreenWidth ScreenHeight:(int32)ScreenHeight
 {
-	if (!bIsInitialized)
-	{
-		// look up what the device can support
-		const float Scale = [[UIScreen mainScreen] scale];
-		const float NativeScale = self.window.screen.nativeScale;
-		const float RequestedContentScaleFactor = [[IOSAppDelegate GetDelegate] GetMobileContentScaleFactor];
+	const IOSAppDelegate* AppDelegate = [IOSAppDelegate GetDelegate];
+	const float mNativeScale = AppDelegate.NativeScale;
+	const float RequestedContentScaleFactor = AppDelegate.MobileContentScaleFactor;
+	const int32 RequestedResX = AppDelegate.RequestedResX;
+	const int32 RequestedResY = AppDelegate.RequestedResY;
 		
-		UE_LOG(LogIOS, Log, TEXT("RequestedContentScaleFactor %f to nativeScale which is = (s:%f, ns:%f, csf:%f"), RequestedContentScaleFactor, Scale, NativeScale, self.contentScaleFactor);
-		// 0 means to leave the scale alone, use native
-		if (RequestedContentScaleFactor == 0.0f)
+//	UE_LOG(LogIOS, Log, TEXT("RequestedContentScaleFactor %f to nativeScale which is = (s:%f, ns:%f, csf:%f"), RequestedContentScaleFactor, AppDelegate.ScreenScale, mNativeScale, self.contentScaleFactor);
+	
+	int32 Width = ScreenWidth;
+	int32 Height = ScreenHeight;
+	
+	// 0 means to use native size
+	if (RequestedContentScaleFactor == 0.0f && RequestedResX <= 0 && RequestedResY <= 0)
+	{
+		self.contentScaleFactor = mNativeScale;
+	}
+	else
+	{
+		float AspectRatio = (float)ScreenHeight / (float)ScreenWidth;
+		self.contentScaleFactor = 1.0f;
+		if (RequestedResX > 0)
 		{
-			self.contentScaleFactor = NativeScale;
-			UE_LOG(LogIOS, Log, TEXT("Setting contentScaleFactor to nativeScale which is = %f"), self.contentScaleFactor);
+			// set long side for orientation to requested X
+			if (ScreenHeight > ScreenWidth)
+			{
+				Height = RequestedResX;
+				Width = (Height * AspectRatio + 0.5f);
+			}
+			else
+			{
+				Width = RequestedResX;
+				Height = (Width * AspectRatio + 0.5f);
+			}
 		}
-		else
+		else if (RequestedResY > 0)
 		{
-			// for TV screens, always use scale factor of 1
-			self.contentScaleFactor = bIsForOnDevice ? RequestedContentScaleFactor : 1.0f;
-			UE_LOG(LogIOS, Log, TEXT("Setting contentScaleFactor to %0.4f (optimal = %0.4f)"), self.contentScaleFactor, NativeScale);
+			// set short side for orientation to requested Y
+			if (ScreenHeight > ScreenWidth)
+			{
+				Width = RequestedResY;
+				Height = (Width * AspectRatio + 0.5f);
+			}
+			else
+			{
+				Height = RequestedResY;
+				Width = (Height * AspectRatio + 0.5f);
+			}
 		}
-
-		// TODO: bIsUsingMetal is always true. Remove it.
-		if (bIsUsingMetal)
+		else
 		{
-			CAMetalLayer* MetalLayer = (CAMetalLayer*)self.layer;
-			CGSize DrawableSize = self.bounds.size;
-			DrawableSize.width *= self.contentScaleFactor;
-			DrawableSize.height *= self.contentScaleFactor;
-			MetalLayer.drawableSize = DrawableSize;
+			self.contentScaleFactor = RequestedContentScaleFactor;
 		}
+	}
 
+	_ViewSize.width = Width;
+	_ViewSize.height = Height;
+}
+
+- (bool)CreateFramebuffer
+{
+	if (!bIsInitialized)
+	{
+		[self CalculateContentScaleFactor:self.frame.size.width ScreenHeight:self.frame.size.height];
 		bIsInitialized = true;
 	}
 	return true;
@@ -318,27 +326,18 @@ id<MTLDevice> GMetalDevice = nil;
 
 -(void)UpdateRenderWidth:(uint32)Width andHeight:(uint32)Height
 {
-#if HAS_METAL
-	if (bIsUsingMetal)
+	if (MetalDevice != nil)
 	{
-		if (MetalDevice != nil)
-		{
-			// grab the MetalLayer and typecast it to match what's in layerClass, then set the new size
-			CAMetalLayer* MetalLayer = (CAMetalLayer*)self.layer;
-			MetalLayer.drawableSize = CGSizeMake(Width, Height);;
-		}
-		return;
+		// grab the MetalLayer and typecast it to match what's in layerClass, then set the new size
+		CAMetalLayer* MetalLayer = (CAMetalLayer*)self.layer;
+		MetalLayer.drawableSize = CGSizeMake(Width, Height);
 	}
-#endif
-
 }
 
-#if HAS_METAL
 - (id<CAMetalDrawable>)MakeDrawable
 {
     return [(CAMetalLayer*)self.layer nextDrawable];
 }
-#endif
 
 - (void)DestroyFramebuffer
 {
@@ -349,17 +348,8 @@ id<MTLDevice> GMetalDevice = nil;
 	}
 }
 
-- (void)MakeCurrent
-{
-}
-
-- (void)UnmakeCurrent
-{
-}
-
 - (void)SwapBuffers
 {
-	// increment our swap counter
 	SwapCount++;
 }
 
@@ -436,8 +426,8 @@ self.accessibilityElements = @[Window.accessibilityContainer];
 	TouchInput TouchMessage;
 	TouchMessage.Handle = TouchIndex;
 	TouchMessage.Type = Type;
-	TouchMessage.Position = FVector2D(FMath::Min<float>(self.frame.size.width - 1, Loc.x), FMath::Min<float>(self.frame.size.height - 1, Loc.y)) * Scale;
-	TouchMessage.LastPosition = FVector2D(FMath::Min<float>(self.frame.size.width - 1, PrevLoc.x), FMath::Min<float>(self.frame.size.height - 1, PrevLoc.y)) * Scale;
+	TouchMessage.Position = FVector2D(FMath::Min<float>(_ViewSize.width - 1, Loc.x), FMath::Min<float>(_ViewSize.height - 1, Loc.y)) * Scale;
+	TouchMessage.LastPosition = FVector2D(FMath::Min<float>(_ViewSize.width - 1, PrevLoc.x), FMath::Min<float>(_ViewSize.height - 1, PrevLoc.y)) * Scale;
 	TouchMessage.Force = Type != TouchEnded ? Force : 0.0f;
 	
 	// skip moves that didn't actually move - this will help input handling to skip over the first
@@ -530,7 +520,14 @@ self.accessibilityElements = @[Window.accessibilityContainer];
 		CGPoint Loc = [Touch locationInView:self];
 		CGPoint PrevLoc = [Touch previousLocationInView:self];
 		
-		// convert TOuch pointer to a unique 0 based index
+		// View may have been modified via Cvars ("r.mobile.DesiredResX/Y" or CommandLine "mcfs, mobileresx/y"
+		CGPoint ViewSizeModifier = CGPointMake(_ViewSize.width/self.frame.size.width, _ViewSize.height/self.frame.size.height);
+		Loc.x *= ViewSizeModifier.x;
+		Loc.y *= ViewSizeModifier.y;
+		PrevLoc.x *= ViewSizeModifier.x;
+		PrevLoc.y *= ViewSizeModifier.y;
+			
+		// convert Touch pointer to a unique 0 based index
 		int32 TouchIndex = [self GetTouchIndex:Touch];
 		if (TouchIndex < 0)
 		{
@@ -1010,7 +1007,7 @@ self.accessibilityElements = @[Window.accessibilityContainer];
 @implementation IOSViewController
 
 /**
- * The ViewController was created, so now we need to create our view to be controlled (an EAGLView)
+ * The ViewController was created, so now we need to create our view to be controlled
  */
 - (void) loadView
 {
diff --git a/Engine/Source/Runtime/ApplicationCore/Private/IOS/IOSWindow.cpp b/Engine/Source/Runtime/ApplicationCore/Private/IOS/IOSWindow.cpp
--- a/Engine/Source/Runtime/ApplicationCore/Private/IOS/IOSWindow.cpp
+++ b/Engine/Source/Runtime/ApplicationCore/Private/IOS/IOSWindow.cpp
@@ -24,13 +24,12 @@ void FIOSWindow::Initialize( class FIOSApplication* const Application, const TSh
 	OwningApplication = Application;
 	Definition = InDefinition;
     
-    UIWindow        *foundKeyWindow = nil;
     NSArray         *allWindows = [[UIApplication sharedApplication]windows];
     for (UIWindow   *currentWindow in allWindows)
     {
         if (currentWindow.isKeyWindow)
         {
-            Window = foundKeyWindow;
+			Window = currentWindow;
             break;
         }
     }
@@ -61,19 +60,76 @@ void FIOSWindow::Initialize( class FIOSApplication* const Application, const TSh
 #endif
 }
 
-FPlatformRect FIOSWindow::GetScreenRect()
+void FIOSWindow::OnScaleFactorChanged(IConsoleVariable* CVar)
 {
-	// get the main view's frame
 	IOSAppDelegate* AppDelegate = [IOSAppDelegate GetDelegate];
-	UIView* View = AppDelegate.IOSView;
+	FIOSView* View = AppDelegate.IOSView;
+		
+	// If r.MobileContentScaleFactor is set by a console command, clear out the r.mobile.DesiredResX/Y CVars
+	if ((CVar->GetFlags() & ECVF_SetByMask) == ECVF_SetByConsole)
+	{
+		IConsoleVariable* CVarResX = IConsoleManager::Get().FindConsoleVariable(TEXT("r.Mobile.DesiredResX"));
+		IConsoleVariable* CVarResY = IConsoleManager::Get().FindConsoleVariable(TEXT("r.Mobile.DesiredResY"));
+		
+		// If CVarResX/Y needs to be reset, let that CVar callback handle the layout change
+		bool OtherCVarChanged = false;
+		if (CVarResX && CVarResX->GetInt() != 0)
+		{
+			CVarResX->Set(0, ECVF_SetByConsole);
+			OtherCVarChanged = true;
+		}
+		if (CVarResY && CVarResY->GetInt() != 0)
+		{
+			CVarResY->Set(0, ECVF_SetByConsole);
+			OtherCVarChanged = true;
+		}
+		
+		if (OtherCVarChanged)
+		{
+			return;
+		}
+	}
+		
+	// Load the latest Cvars that might affect screen size
+	[AppDelegate LoadScreenResolutionModifiers];
+		
+	// Force a re-layout of our views as the size has probably changed
 	CGRect Frame = [View frame];
-	CGFloat Scale = View.contentScaleFactor;
+	[View CalculateContentScaleFactor:Frame.size.width ScreenHeight:Frame.size.height];
+	[View layoutSubviews];
+}
 
+void FIOSWindow::OnConsoleResolutionChanged(IConsoleVariable* CVar)
+{
+	IOSAppDelegate* AppDelegate = [IOSAppDelegate GetDelegate];
+	FIOSView* View = AppDelegate.IOSView;
+		
+	// Load the latest Cvars that might affect screen size
+	[AppDelegate LoadScreenResolutionModifiers];
+		
+	// Force a re-layout of our views as the size has probably changed
+	CGRect Frame = [View frame];
+	[View CalculateContentScaleFactor:Frame.size.width ScreenHeight:Frame.size.height];
+	[View layoutSubviews];
+}
+
+FPlatformRect FIOSWindow::GetScreenRect()
+{
+	// get the main view's frame
+	IOSAppDelegate* AppDelegate = [IOSAppDelegate GetDelegate];
+	FIOSView* View = AppDelegate.IOSView;
+	
 	FPlatformRect ScreenRect;
-	ScreenRect.Top = Frame.origin.y * Scale;
-	ScreenRect.Bottom = (Frame.origin.y + Frame.size.height) * Scale;
-	ScreenRect.Left = Frame.origin.x * Scale;
-	ScreenRect.Right = (Frame.origin.x + Frame.size.width) * Scale;
+	if (View != nil)
+	{
+		CGRect Frame = [View frame];
+		CGFloat Scale = View.contentScaleFactor;
+		
+		ScreenRect.Top = Frame.origin.y * Scale;
+		ScreenRect.Bottom = (Frame.origin.y + View.ViewSize.height) * Scale;
+		ScreenRect.Left = Frame.origin.x * Scale;
+		ScreenRect.Right = (Frame.origin.x + View.ViewSize.width) * Scale;
+	}
 
 	return ScreenRect;
 }
@@ -94,7 +150,6 @@ FPlatformRect FIOSWindow::GetUIWindowRect()
 	return WindowRect;
 }
 
-
 bool FIOSWindow::GetFullScreenInfo( int32& X, int32& Y, int32& Width, int32& Height ) const
 {
 	FPlatformRect ScreenRect = GetScreenRect();
diff --git a/Engine/Source/Runtime/ApplicationCore/Private/IOS/IOSWindow.h b/Engine/Source/Runtime/ApplicationCore/Private/IOS/IOSWindow.h
--- a/Engine/Source/Runtime/ApplicationCore/Private/IOS/IOSWindow.h
+++ b/Engine/Source/Runtime/ApplicationCore/Private/IOS/IOSWindow.h
@@ -36,6 +36,10 @@ public:
 	/** Returns the rectangle of the screen the window is associated with */
 	virtual bool GetFullScreenInfo( int32& X, int32& Y, int32& Width, int32& Height ) const override;
 
+	/** Callbacks for Cvar changes "r.MobileContentScaleFactor" and "r.Mobile.DesiredResX/Y" respectively */
+	static void OnScaleFactorChanged(IConsoleVariable* CVar);
+	static void OnConsoleResolutionChanged(IConsoleVariable* CVar);
+		
 protected:
 	/** @return true if the native window is currently in fullscreen mode, false otherwise */
 	virtual EWindowMode::Type GetWindowMode() const override { return EWindowMode::Fullscreen; }
diff --git a/Engine/Source/Runtime/ApplicationCore/Public/GenericPlatform/GenericApplication.h b/Engine/Source/Runtime/ApplicationCore/Public/GenericPlatform/GenericApplication.h
--- a/Engine/Source/Runtime/ApplicationCore/Public/GenericPlatform/GenericApplication.h
+++ b/Engine/Source/Runtime/ApplicationCore/Public/GenericPlatform/GenericApplication.h
@@ -376,11 +376,6 @@ struct FDisplayMetrics
 	/** Virtual display coordinate range (includes all active displays) */
 	FPlatformRect VirtualDisplayRect;
 
-#if PLATFORM_IOS
-	/** Area of the UIWindow on IOS devices */
-	FPlatformRect IosUiWindowAreaRect;
-#endif
-
 	/**
 	 * The safe area for all content on TVs (see http://en.wikipedia.org/wiki/Safe_area_%28television%29) - content will be inset 
 	 * Left - X
diff --git a/Engine/Source/Runtime/ApplicationCore/Public/IOS/IOSAppDelegate.h b/Engine/Source/Runtime/ApplicationCore/Public/IOS/IOSAppDelegate.h
--- a/Engine/Source/Runtime/ApplicationCore/Public/IOS/IOSAppDelegate.h
+++ b/Engine/Source/Runtime/ApplicationCore/Public/IOS/IOSAppDelegate.h
@@ -180,7 +180,11 @@ APPLICATIONCORE_API
 @property (assign) bool bBatteryState;
 @property (assign) int BatteryLevel;
 
+@property (assign) float ScreenScale;	// UIScreen.scale
+@property (assign) float NativeScale;	// UIWindow.screen.nativeScale
 @property (assign) float MobileContentScaleFactor;
+@property (assign) int RequestedResX;
+@property (assign) int RequestedResY;
 @property (assign) bool bUpdateAvailable;
 
 /**
@@ -195,8 +199,7 @@ APPLICATIONCORE_API
 +(bool)WaitAndRunOnGameThread:(TUniqueFunction<void()>)Function;
 -(void)NoUrlCommandLine;
 
--(void)LoadMobileContentScaleFactor;
--(float)GetMobileContentScaleFactor;
+-(void)LoadScreenResolutionModifiers;
 
 -(int)GetAudioVolume;
 -(bool)AreHeadphonesPluggedIn;
diff --git a/Engine/Source/Runtime/ApplicationCore/Public/IOS/IOSView.h b/Engine/Source/Runtime/ApplicationCore/Public/IOS/IOSView.h
--- a/Engine/Source/Runtime/ApplicationCore/Public/IOS/IOSView.h
+++ b/Engine/Source/Runtime/ApplicationCore/Public/IOS/IOSView.h
@@ -7,12 +7,9 @@
 #include "IOS/IOSInputInterface.h"
 
 #import <UIKit/UIKit.h>
-#import <OpenGLES/EAGL.h>
 
-#if HAS_METAL
 #import <Metal/Metal.h>
 #import <QuartzCore/CAMetalLayer.h>
-#endif
 
 #if WITH_ACCESSIBILITY
 #include "GenericPlatform/Accessibility/GenericAccessibleInterfaces.h"
@@ -50,17 +47,10 @@ APPLICATIONCORE_API
 	float PreviousForces[10];
 	bool HasMoved[10];
 
-
-	//// METAL MEMBERS
-#if HAS_METAL
 	// global metal device
 	id<MTLDevice> MetalDevice;
 	id<CAMetalDrawable> PanicDrawable;
-#endif
 
-	// are we using the Metal API?
-	bool bIsUsingMetal;
-	
 	//// KEYBOARD MEMBERS
 	
 	// whether or not to use the new style virtual keyboard that sends events to the engine instead of using an alert
@@ -76,7 +66,6 @@ APPLICATIONCORE_API
 	BOOL bSecureTextEntry;
 	
 	volatile int32 KeyboardShowCount;
-
 }
 
 #if WITH_ACCESSIBILITY
@@ -85,25 +74,19 @@ APPLICATIONCORE_API
 #endif
 
 //// SHARED FUNCTIONALITY
-@property (nonatomic) GLuint SwapCount;
+@property (nonatomic) uint SwapCount;
+@property (assign, nonatomic) CGSize ViewSize;
 
--(bool)CreateFramebuffer:(bool)bIsForOnDevice;
+-(bool)CreateFramebuffer;
 -(void)DestroyFramebuffer;
 -(void)UpdateRenderWidth:(uint32)Width andHeight:(uint32)Height;
+-(void)CalculateContentScaleFactor:(int32)ScreenWidth ScreenHeight:(int32)ScreenHeight;
 
-//// GL FUNCTIONALITY
-@property (nonatomic) GLuint OnScreenColorRenderBuffer;
-@property (nonatomic) GLuint OnScreenColorRenderBufferMSAA;
-
-- (void)MakeCurrent;
-- (void)UnmakeCurrent;
 - (void)SwapBuffers;
 
 //// METAL FUNCTIONALITY
-#if HAS_METAL
 // Return a drawable object (ie a back buffer texture) for the RHI to render to
 - (id<CAMetalDrawable>)MakeDrawable;
-#endif
 
 //// KEYBOARD FUNCTIONALITY
 -(void)InitKeyboard;
diff --git a/Engine/Source/Runtime/Core/Private/IOS/IOSPlatformMisc.cpp b/Engine/Source/Runtime/Core/Private/IOS/IOSPlatformMisc.cpp
--- a/Engine/Source/Runtime/Core/Private/IOS/IOSPlatformMisc.cpp
+++ b/Engine/Source/Runtime/Core/Private/IOS/IOSPlatformMisc.cpp
@@ -74,8 +74,9 @@ static int32 GetFreeMemoryMB()
 void FIOSPlatformMisc::PlatformInit()
 {
 	// PlatformInit() starts the UI thread which creates the framebuffer and it requires
-	// "r.MobileContentScaleFactor" to be available before it's creation, so need to cache that value now.
-	[[IOSAppDelegate GetDelegate] LoadMobileContentScaleFactor];
+	// "r.MobileContentScaleFactor" and "r.Mobile.DesiredResX/Y" to be available before 
+	// it's creation, so need to cache those value now.
+	[[IOSAppDelegate GetDelegate] LoadScreenResolutionModifiers];
 		
 	FAppEntry::PlatformInit();
 
@@ -321,7 +322,7 @@ bool FIOSPlatformMisc::HasPlatformFeature(const TCHAR* FeatureName)
 {
 	if (FCString::Stricmp(FeatureName, TEXT("Metal")) == 0)
 	{
-		return [IOSAppDelegate GetDelegate].IOSView->bIsUsingMetal;
+		return true;
 	}
 
 	return FGenericPlatformMisc::HasPlatformFeature(FeatureName);
diff --git a/Engine/Source/Runtime/Launch/Private/IOS/LaunchIOS.cpp b/Engine/Source/Runtime/Launch/Private/IOS/LaunchIOS.cpp
--- a/Engine/Source/Runtime/Launch/Private/IOS/LaunchIOS.cpp
+++ b/Engine/Source/Runtime/Launch/Private/IOS/LaunchIOS.cpp
@@ -355,7 +355,7 @@ static void MainThreadInit()
 	[AppDelegate.RootView addSubview:AppDelegate.IOSView];
 
 	// initialize the backbuffer of the view (so the RHI can use it)
-	[AppDelegate.IOSView CreateFramebuffer:YES];
+	[AppDelegate.IOSView CreateFramebuffer];
 #endif
 }
 
@@ -386,9 +386,6 @@ void FAppEntry::PlatformInit()
 		FPlatformProcess::Sleep(0.005f);
 	}
 
-	// set the GL context to this thread
-	[AppDelegate.IOSView MakeCurrent];
-
 	// Set GSystemResolution now that we have the size.
 	FDisplayMetrics DisplayMetrics;
 	FDisplayMetrics::RebuildDisplayMetrics(DisplayMetrics);
