Allow platforms to override the platform input device mapper.

#rnx

diff --git a/Engine/Source/Runtime/ApplicationCore/Private/GenericPlatform/GenericPlatformInputDeviceMapper.cpp b/Engine/Source/Runtime/ApplicationCore/Private/GenericPlatform/GenericPlatformInputDeviceMapper.cpp
--- a/Engine/Source/Runtime/ApplicationCore/Private/GenericPlatform/GenericPlatformInputDeviceMapper.cpp
+++ b/Engine/Source/Runtime/ApplicationCore/Private/GenericPlatform/GenericPlatformInputDeviceMapper.cpp
@@ -1,7 +1,7 @@
 // Copyright Epic Games, Inc. All Rights Reserved.
 
 #include "GenericPlatform/GenericPlatformInputDeviceMapper.h"
-#include "GenericPlatform/GenericPlatformApplicationMisc.h"
+#include "HAL/PlatformApplicationMisc.h"
 #include "Misc/CoreDelegates.h"
 
 ///////////////////////////////////////////////////////////////////////////
@@ -17,7 +17,7 @@ IPlatformInputDeviceMapper& IPlatformInputDeviceMapper::Get()
 	static IPlatformInputDeviceMapper* StaticManager = nullptr;
 	if (!StaticManager)
 	{
-		StaticManager = FGenericPlatformApplicationMisc::CreatePlatformInputDeviceManager();
+		StaticManager = FPlatformApplicationMisc::CreatePlatformInputDeviceManager();
 		check(StaticManager);
 	}
 
