Remove an over-aggressive ensure when mapping input devices. This ensure was there to check if there was a default platform user that was unpaired, which is possible on some platforms.

#rnx

diff --git a/Engine/Source/Runtime/ApplicationCore/Private/GenericPlatform/GenericPlatformInputDeviceMapper.cpp b/Engine/Source/Runtime/ApplicationCore/Private/GenericPlatform/GenericPlatformInputDeviceMapper.cpp
--- a/Engine/Source/Runtime/ApplicationCore/Private/GenericPlatform/GenericPlatformInputDeviceMapper.cpp
+++ b/Engine/Source/Runtime/ApplicationCore/Private/GenericPlatform/GenericPlatformInputDeviceMapper.cpp
@@ -215,7 +215,7 @@ bool IPlatformInputDeviceMapper::Internal_ChangeInputDeviceUserMapping(FInputDev
 	if (FPlatformInputDeviceState* ExistingDeviceState = MappedInputDevices.Find(DeviceId))
 	{
 		// Only change the platform user of this device if the old user matches up with the one that was given
-		if (ensureMsgf(ExistingDeviceState->OwningPlatformUser == OldUserId, TEXT("Attempting to change the Input Device User Mapping with a mismatched OldUserId!")))
+		if (ExistingDeviceState->OwningPlatformUser == OldUserId)
 		{
 			ExistingDeviceState->OwningPlatformUser = NewUserId;
 		}
