Update the LinuxApplication to use the newer version of OnController functions instead of the old "int32 ControllerId" versions that are going to be deprecated soon.

#rnx

diff --git a/Engine/Source/Runtime/ApplicationCore/Private/Linux/LinuxApplication.cpp b/Engine/Source/Runtime/ApplicationCore/Private/Linux/LinuxApplication.cpp
--- a/Engine/Source/Runtime/ApplicationCore/Private/Linux/LinuxApplication.cpp
+++ b/Engine/Source/Runtime/ApplicationCore/Private/Linux/LinuxApplication.cpp
@@ -32,6 +32,21 @@ float ShortToNormalFloat(short AxisVal)
 	return float(AxisVal) / Norm;
 }
 
+namespace UE::Input
+{
+	/** A helper method for converting an Input Device id to an int32 */
+	static int32 ConvertInputDeviceToInt(FInputDeviceId DeviceId)
+	{
+		int32 UserStateControllerId = INDEX_NONE;
+		IPlatformInputDeviceMapper& Mapper = IPlatformInputDeviceMapper::Get();
+
+		FPlatformUserId UserId = Mapper.GetUserForInputDevice(DeviceId);
+		Mapper.RemapUserAndDeviceToControllerId(UserId, UserStateControllerId, DeviceId);
+
+		return UserStateControllerId;
+	}
+}
+
 FLinuxApplication* LinuxApplication = NULL;
 
 FLinuxApplication* FLinuxApplication::CreateLinuxApplication()
@@ -477,6 +492,7 @@ void FLinuxApplication::ProcessDeferredMessage( SDL_Event Event )
 			}
 
 			SDLControllerState &ControllerState = ControllerStates[caxisEvent.which];
+			FPlatformUserId UserId = IPlatformInputDeviceMapper::Get().GetUserForInputDevice(ControllerState.DeviceId);
 
 			switch (caxisEvent.axis)
 			{
@@ -486,26 +502,26 @@ void FLinuxApplication::ProcessDeferredMessage( SDL_Event Event )
 				{
 					if(!ControllerState.AnalogOverThreshold[0])
 					{
-						MessageHandler->OnControllerButtonPressed(FGamepadKeyNames::LeftStickRight, ControllerState.ControllerIndex, false);
+						MessageHandler->OnControllerButtonPressed(FGamepadKeyNames::LeftStickRight, UserId, ControllerState.DeviceId, false);
 						ControllerState.AnalogOverThreshold[0] = true;
 					}
 				}
 				else if(ControllerState.AnalogOverThreshold[0])
 				{
-					MessageHandler->OnControllerButtonReleased(FGamepadKeyNames::LeftStickRight, ControllerState.ControllerIndex, false);
+					MessageHandler->OnControllerButtonReleased(FGamepadKeyNames::LeftStickRight, UserId, ControllerState.DeviceId, false);
 					ControllerState.AnalogOverThreshold[0] = false;
 				}
 				if(caxisEvent.value < -GAMECONTROLLER_LEFT_THUMB_DEADZONE)
 				{
 					if(!ControllerState.AnalogOverThreshold[1])
 					{
-						MessageHandler->OnControllerButtonPressed(FGamepadKeyNames::LeftStickLeft, ControllerState.ControllerIndex, false);
+						MessageHandler->OnControllerButtonPressed(FGamepadKeyNames::LeftStickLeft, UserId, ControllerState.DeviceId, false);
 						ControllerState.AnalogOverThreshold[1] = true;
 					}
 				}
 				else if(ControllerState.AnalogOverThreshold[1])
 				{
-					MessageHandler->OnControllerButtonReleased(FGamepadKeyNames::LeftStickLeft, ControllerState.ControllerIndex, false);
+					MessageHandler->OnControllerButtonReleased(FGamepadKeyNames::LeftStickLeft, UserId, ControllerState.DeviceId, false);
 					ControllerState.AnalogOverThreshold[1] = false;
 				}
 				break;
@@ -516,26 +532,26 @@ void FLinuxApplication::ProcessDeferredMessage( SDL_Event Event )
 				{
 					if(!ControllerState.AnalogOverThreshold[2])
 					{
-						MessageHandler->OnControllerButtonPressed(FGamepadKeyNames::LeftStickDown, ControllerState.ControllerIndex, false);
+						MessageHandler->OnControllerButtonPressed(FGamepadKeyNames::LeftStickDown, UserId, ControllerState.DeviceId, false);
 						ControllerState.AnalogOverThreshold[2] = true;
 					}
 				}
 				else if(ControllerState.AnalogOverThreshold[2])
 				{
-					MessageHandler->OnControllerButtonReleased(FGamepadKeyNames::LeftStickDown, ControllerState.ControllerIndex, false);
+					MessageHandler->OnControllerButtonReleased(FGamepadKeyNames::LeftStickDown, UserId, ControllerState.DeviceId, false);
 					ControllerState.AnalogOverThreshold[2] = false;
 				}
 				if(caxisEvent.value < -GAMECONTROLLER_LEFT_THUMB_DEADZONE)
 				{
 					if(!ControllerState.AnalogOverThreshold[3])
 					{
-						MessageHandler->OnControllerButtonPressed(FGamepadKeyNames::LeftStickUp, ControllerState.ControllerIndex, false);
+						MessageHandler->OnControllerButtonPressed(FGamepadKeyNames::LeftStickUp, UserId, ControllerState.DeviceId, false);
 						ControllerState.AnalogOverThreshold[3] = true;
 					}
 				}
 				else if(ControllerState.AnalogOverThreshold[3])
 				{
-					MessageHandler->OnControllerButtonReleased(FGamepadKeyNames::LeftStickUp, ControllerState.ControllerIndex, false);
+					MessageHandler->OnControllerButtonReleased(FGamepadKeyNames::LeftStickUp, UserId, ControllerState.DeviceId, false);
 					ControllerState.AnalogOverThreshold[3] = false;
 				}
 				break;
@@ -545,26 +561,26 @@ void FLinuxApplication::ProcessDeferredMessage( SDL_Event Event )
 				{
 					if(!ControllerState.AnalogOverThreshold[4])
 					{
-						MessageHandler->OnControllerButtonPressed(FGamepadKeyNames::RightStickRight, ControllerState.ControllerIndex, false);
+						MessageHandler->OnControllerButtonPressed(FGamepadKeyNames::RightStickRight, UserId, ControllerState.DeviceId, false);
 						ControllerState.AnalogOverThreshold[4] = true;
 					}
 				}
 				else if(ControllerState.AnalogOverThreshold[4])
 				{
-					MessageHandler->OnControllerButtonReleased(FGamepadKeyNames::RightStickRight, ControllerState.ControllerIndex, false);
+					MessageHandler->OnControllerButtonReleased(FGamepadKeyNames::RightStickRight, UserId, ControllerState.DeviceId, false);
 					ControllerState.AnalogOverThreshold[4] = false;
 				}
 				if(caxisEvent.value < -GAMECONTROLLER_RIGHT_THUMB_DEADZONE)
 				{
 					if(!ControllerState.AnalogOverThreshold[5])
 					{
-						MessageHandler->OnControllerButtonPressed(FGamepadKeyNames::RightStickLeft, ControllerState.ControllerIndex, false);
+						MessageHandler->OnControllerButtonPressed(FGamepadKeyNames::RightStickLeft, UserId, ControllerState.DeviceId, false);
 						ControllerState.AnalogOverThreshold[5] = true;
 					}
 				}
 				else if(ControllerState.AnalogOverThreshold[5])
 				{
-					MessageHandler->OnControllerButtonReleased(FGamepadKeyNames::RightStickLeft, ControllerState.ControllerIndex, false);
+					MessageHandler->OnControllerButtonReleased(FGamepadKeyNames::RightStickLeft, UserId, ControllerState.DeviceId, false);
 					ControllerState.AnalogOverThreshold[5] = false;
 				}
 				break;
@@ -575,26 +591,26 @@ void FLinuxApplication::ProcessDeferredMessage( SDL_Event Event )
 				{
 					if(!ControllerState.AnalogOverThreshold[6])
 					{
-						MessageHandler->OnControllerButtonPressed(FGamepadKeyNames::RightStickDown, ControllerState.ControllerIndex, false);
+						MessageHandler->OnControllerButtonPressed(FGamepadKeyNames::RightStickDown, UserId, ControllerState.DeviceId, false);
 						ControllerState.AnalogOverThreshold[6] = true;
 					}
 				}
 				else if(ControllerState.AnalogOverThreshold[6])
 				{
-					MessageHandler->OnControllerButtonReleased(FGamepadKeyNames::RightStickDown, ControllerState.ControllerIndex, false);
+					MessageHandler->OnControllerButtonReleased(FGamepadKeyNames::RightStickDown, UserId, ControllerState.DeviceId, false);
 					ControllerState.AnalogOverThreshold[6] = false;
 				}
 				if(caxisEvent.value < -GAMECONTROLLER_RIGHT_THUMB_DEADZONE)
 				{
 					if(!ControllerState.AnalogOverThreshold[7])
 					{
-						MessageHandler->OnControllerButtonPressed(FGamepadKeyNames::RightStickUp, ControllerState.ControllerIndex, false);
+						MessageHandler->OnControllerButtonPressed(FGamepadKeyNames::RightStickUp, UserId, ControllerState.DeviceId, false);
 						ControllerState.AnalogOverThreshold[7] = true;
 					}
 				}
 				else if(ControllerState.AnalogOverThreshold[7])
 				{
-					MessageHandler->OnControllerButtonReleased(FGamepadKeyNames::RightStickUp, ControllerState.ControllerIndex, false);
+					MessageHandler->OnControllerButtonReleased(FGamepadKeyNames::RightStickUp, UserId, ControllerState.DeviceId, false);
 					ControllerState.AnalogOverThreshold[7] = false;
 				}
 				break;
@@ -604,13 +620,13 @@ void FLinuxApplication::ProcessDeferredMessage( SDL_Event Event )
 				{
 					if(!ControllerState.AnalogOverThreshold[8])
 					{
-						MessageHandler->OnControllerButtonPressed(FGamepadKeyNames::LeftTriggerThreshold, ControllerState.ControllerIndex, false);
+						MessageHandler->OnControllerButtonPressed(FGamepadKeyNames::LeftTriggerThreshold, UserId, ControllerState.DeviceId, false);
 						ControllerState.AnalogOverThreshold[8] = true;
 					}
 				}
 				else if(ControllerState.AnalogOverThreshold[8])
 				{
-					MessageHandler->OnControllerButtonReleased(FGamepadKeyNames::LeftTriggerThreshold, ControllerState.ControllerIndex, false);
+					MessageHandler->OnControllerButtonReleased(FGamepadKeyNames::LeftTriggerThreshold, UserId, ControllerState.DeviceId, false);
 					ControllerState.AnalogOverThreshold[8] = false;
 				}
 				break;
@@ -620,13 +636,13 @@ void FLinuxApplication::ProcessDeferredMessage( SDL_Event Event )
 				{
 					if(!ControllerState.AnalogOverThreshold[9])
 					{
-						MessageHandler->OnControllerButtonPressed(FGamepadKeyNames::RightTriggerThreshold, ControllerState.ControllerIndex, false);
+						MessageHandler->OnControllerButtonPressed(FGamepadKeyNames::RightTriggerThreshold, UserId, ControllerState.DeviceId, false);
 						ControllerState.AnalogOverThreshold[9] = true;
 					}
 				}
 				else if(ControllerState.AnalogOverThreshold[9])
 				{
-					MessageHandler->OnControllerButtonReleased(FGamepadKeyNames::RightTriggerThreshold, ControllerState.ControllerIndex, false);
+					MessageHandler->OnControllerButtonReleased(FGamepadKeyNames::RightTriggerThreshold, UserId, ControllerState.DeviceId, false);
 					ControllerState.AnalogOverThreshold[9] = false;
 				}
 				break;
@@ -702,13 +718,15 @@ void FLinuxApplication::ProcessDeferredMessage( SDL_Event Event )
 
 			if (Button != FGamepadKeyNames::Invalid)
 			{
+				FPlatformUserId UserId = IPlatformInputDeviceMapper::Get().GetUserForInputDevice(ControllerStates[cbuttonEvent.which].DeviceId);
+
 				if(cbuttonEvent.type == SDL_CONTROLLERBUTTONDOWN)
 				{
-					MessageHandler->OnControllerButtonPressed(Button, ControllerStates[cbuttonEvent.which].ControllerIndex, false);
+					MessageHandler->OnControllerButtonPressed(Button, UserId, ControllerStates[cbuttonEvent.which].DeviceId, false);
 				}
 				else
 				{
-					MessageHandler->OnControllerButtonReleased(Button, ControllerStates[cbuttonEvent.which].ControllerIndex, false);
+					MessageHandler->OnControllerButtonReleased(Button, UserId, ControllerStates[cbuttonEvent.which].DeviceId, false);
 				}
 			}
 		}
@@ -1125,11 +1143,13 @@ void FLinuxApplication::ProcessDeferredEvents( const float TimeDelta )
 
 void FLinuxApplication::PollGameDeviceState( const float TimeDelta )
 {
+	IPlatformInputDeviceMapper& Mapper = IPlatformInputDeviceMapper::Get();
 	for(auto ControllerIt = ControllerStates.CreateIterator(); ControllerIt; ++ControllerIt)
 	{
 		for(auto Event = ControllerIt.Value().AxisEvents.CreateConstIterator(); Event; ++Event)
 		{
-			MessageHandler->OnControllerAnalog(Event.Key(), ControllerIt.Value().ControllerIndex, Event.Value());
+			FPlatformUserId UserId = Mapper.GetUserForInputDevice(ControllerIt.Value().DeviceId);
+			MessageHandler->OnControllerAnalog(Event.Key(), UserId, ControllerIt.Value().DeviceId, Event.Value());
 		}
 		ControllerIt.Value().AxisEvents.Empty();
 
@@ -1899,7 +1919,8 @@ void FLinuxApplication::SetForceFeedbackChannelValue(int32 ControllerId, FForceF
 	for (auto ControllerIt = ControllerStates.CreateIterator(); ControllerIt; ++ControllerIt)
 	{
 		auto& ControllerState = ControllerIt.Value();
-		if (ControllerState.ControllerIndex == ControllerId)
+		int32 ControllerIndex = UE::Input::ConvertInputDeviceToInt(ControllerState.DeviceId);
+		if (ControllerIndex == ControllerId)
 		{
 			if (ControllerState.Haptic != nullptr)
 			{
@@ -1939,7 +1960,8 @@ void FLinuxApplication::SetForceFeedbackChannelValues(int32 ControllerId, const
 	for (auto ControllerIt = ControllerStates.CreateIterator(); ControllerIt; ++ControllerIt)
 	{
 		auto& ControllerState = ControllerIt.Value();
-		if (ControllerState.ControllerIndex == ControllerId)
+		int32 ControllerIndex = UE::Input::ConvertInputDeviceToInt(ControllerState.DeviceId);
+		if (ControllerIndex == ControllerId)
 		{
 			if (ControllerState.Haptic != nullptr)
 			{
@@ -2066,10 +2088,13 @@ void FLinuxApplication::AddGameController(int Index)
 		return;
 	}
 
+	IPlatformInputDeviceMapper& DeviceMapper = IPlatformInputDeviceMapper::Get();
+
 	uint32 UsedBits = 0;
 	for (auto ControllerIt = ControllerStates.CreateIterator(); ControllerIt; ++ControllerIt)
 	{
-		UsedBits |= (1 << ControllerIt.Value().ControllerIndex);
+		FPlatformUserId UserId = DeviceMapper.GetUserForInputDevice(ControllerIt.Value().DeviceId);
+		UsedBits |= (1 << UserId.GetInternalId());
 	}
 
 	int32 FirstUnusedIndex = FMath::CountTrailingZeros(~UsedBits);
@@ -2077,7 +2102,10 @@ void FLinuxApplication::AddGameController(int Index)
 	UE_LOG(LogLinux, Verbose, TEXT("Adding controller %i '%s'"), FirstUnusedIndex, UTF8_TO_TCHAR(SDL_GameControllerName(Controller)));
 	auto& ControllerState = ControllerStates.Add(Id);
 	ControllerState.Controller = Controller;
-	ControllerState.ControllerIndex = FirstUnusedIndex;
+
+	FPlatformUserId UserId = FPlatformUserId::CreateFromInternalId(FirstUnusedIndex);
+	DeviceMapper.RemapControllerIdToPlatformUserAndDevice(FirstUnusedIndex, UserId, ControllerState.DeviceId);
+
 	// Check for haptic support.
 	ControllerState.Haptic = SDL_HapticOpenFromJoystick(SDL_GameControllerGetJoystick(Controller));
 
@@ -2085,19 +2113,13 @@ void FLinuxApplication::AddGameController(int Index)
 	{
 		if ((SDL_HapticQuery(ControllerState.Haptic) & (SDL_HAPTIC_SINE | SDL_HAPTIC_LEFTRIGHT)) == 0)
 		{
-			UE_LOG(LogLinux, Warning, TEXT("No supported haptic effects for controller %i"), ControllerState.ControllerIndex);
+			UE_LOG(LogLinux, Warning, TEXT("No supported haptic effects for controller %i"), ControllerState.DeviceId.GetId());
 			SDL_HapticClose(ControllerState.Haptic);
 			ControllerState.Haptic = nullptr;
 		}
 	}
 
-	FPlatformUserId PlatformUserId = PLATFORMUSERID_NONE;
-	FInputDeviceId DeviceId = INPUTDEVICEID_NONE;
-
-	IPlatformInputDeviceMapper& DeviceMapper = IPlatformInputDeviceMapper::Get();
-	DeviceMapper.RemapControllerIdToPlatformUserAndDevice(Id, OUT PlatformUserId, OUT DeviceId);
-
-	DeviceMapper.Internal_MapInputDeviceToUser(DeviceId, PlatformUserId, EInputDeviceConnectionState::Connected);
+	DeviceMapper.Internal_MapInputDeviceToUser(ControllerState.DeviceId, UserId, EInputDeviceConnectionState::Connected);
 }
 
 void FLinuxApplication::RemoveGameController(SDL_JoystickID Id)
@@ -2108,7 +2130,7 @@ void FLinuxApplication::RemoveGameController(SDL_JoystickID Id)
 	}
 
 	SDLControllerState& ControllerState = ControllerStates[Id];
-	UE_LOG(LogLinux, Verbose, TEXT("Removing controller %i '%s'"), ControllerState.ControllerIndex, UTF8_TO_TCHAR(SDL_GameControllerName(ControllerState.Controller)));
+	UE_LOG(LogLinux, Verbose, TEXT("Removing controller %i '%s'"), ControllerState.DeviceId.GetId(), UTF8_TO_TCHAR(SDL_GameControllerName(ControllerState.Controller)));
 
 	if (ControllerState.Haptic != nullptr)
 	{
diff --git a/Engine/Source/Runtime/ApplicationCore/Public/Linux/LinuxApplication.h b/Engine/Source/Runtime/ApplicationCore/Public/Linux/LinuxApplication.h
--- a/Engine/Source/Runtime/ApplicationCore/Public/Linux/LinuxApplication.h
+++ b/Engine/Source/Runtime/ApplicationCore/Public/Linux/LinuxApplication.h
@@ -255,8 +255,8 @@ private:
 		/** Tracks whether the "button" was previously pressed so we don't generate extra events */
 		bool AnalogOverThreshold[10];
 
-		/** The player index of the controller, because the joystick index includes devices that are not controllers. */
-		int32 ControllerIndex;
+		/** The input device Id of the controller that can be used to find the matching ULocalPlayer */
+		FInputDeviceId DeviceId;
 
 		/** Store axis values from events here to be handled once per frame. */
 		TMap<FGamepadKeyNames::Type, float> AxisEvents;
@@ -272,7 +272,7 @@ private:
 
 		SDLControllerState()
 			:	Controller(nullptr)
-			,	ControllerIndex(-1)
+			,	DeviceId(INPUTDEVICEID_NONE)
 			,	Haptic(nullptr)
 			,	EffectId(-1)
 			,	bEffectRunning(false)
