Add ShowInputDeviceSelector function to display a platform-specific input device selection UI, if the platform supports it.

#rnx

diff --git a/Engine/Source/Runtime/ApplicationCore/Public/GenericPlatform/GenericPlatformApplicationMisc.h b/Engine/Source/Runtime/ApplicationCore/Public/GenericPlatform/GenericPlatformApplicationMisc.h
--- a/Engine/Source/Runtime/ApplicationCore/Public/GenericPlatform/GenericPlatformApplicationMisc.h
+++ b/Engine/Source/Runtime/ApplicationCore/Public/GenericPlatform/GenericPlatformApplicationMisc.h
@@ -18,6 +18,17 @@ enum class EScreenPhysicalAccuracy
 	Truth
 };
 
+/**
+ * Callback for when FindInputDeviceForUserWithUI has completed
+ */
+struct FShowInputDeviceSelectorParams
+{
+	FInputDeviceId InputDeviceId;
+	FPlatformUserId PlatformUserId;
+};
+typedef TFunction<void(const FShowInputDeviceSelectorParams&)> FShowInputDeviceSelectorComplete;
+
+
 struct APPLICATIONCORE_API FGenericPlatformApplicationMisc
 {
 	static void PreInit();
@@ -306,6 +317,21 @@ struct APPLICATIONCORE_API FGenericPlatformApplicationMisc
 		return Accuracy;		
 	}
 	
+	/**
+	 * Asyncronously display the platform-specific input device selection UI, if supported.
+	 * 
+	 * @param InitiatingUserId The platform user to find an input device for.
+	 * @param OnShowInputDeviceSelectorComplete Callback for when the input device selection operation has completed
+	 * @return true if the UI will be shown and the callback will be called
+	 * @return false if the platform does not support an input device selection API, or it is not currently available (the callback will not be called)
+	 */
+	static bool ShowInputDeviceSelector( FPlatformUserId InitiatingUserId, FShowInputDeviceSelectorComplete OnShowInputDeviceSelectorComplete )
+	{
+		// no default implementation
+		return false;
+	}
+
+
 protected:
 	static bool CachedPhysicalScreenData;
 	static EScreenPhysicalAccuracy CachedPhysicalScreenAccuracy;
