CIS Fix for deprecation functions in android input

#rnx

diff --git a/Engine/Source/Runtime/ApplicationCore/Private/Android/AndroidInputInterface.cpp b/Engine/Source/Runtime/ApplicationCore/Private/Android/AndroidInputInterface.cpp
--- a/Engine/Source/Runtime/ApplicationCore/Private/Android/AndroidInputInterface.cpp
+++ b/Engine/Source/Runtime/ApplicationCore/Private/Android/AndroidInputInterface.cpp
@@ -1076,34 +1076,39 @@ void FAndroidInputInterface::SendControllerEvents()
 
 			FAndroidControllerData& OldControllerState = OldControllerData[ControllerIndex];
 			FAndroidControllerData& NewControllerState = NewControllerData[ControllerIndex];
+			
+			IPlatformInputDeviceMapper& DeviceMapper = IPlatformInputDeviceMapper::Get();		            			
+			FPlatformUserId UserId = FGenericPlatformMisc::GetPlatformUserForUserIndex(ControllerIndex);
+			FInputDeviceId DeviceId = INPUTDEVICEID_NONE;
+			DeviceMapper.RemapControllerIdToPlatformUserAndDevice(NewControllerState.DeviceId, OUT UserId, OUT DeviceId);
 
 			// Send controller events any time we have a large enough input threshold similarly to PC/Console (see: XInputInterface.cpp)
 			const float RepeatDeadzone = 0.24f;
 
 			if (NewControllerState.LXAnalog != OldControllerState.LXAnalog || FMath::Abs(NewControllerState.LXAnalog) >= RepeatDeadzone)
 			{
-				MessageHandler->OnControllerAnalog(FGamepadKeyNames::LeftAnalogX, NewControllerState.DeviceId, NewControllerState.LXAnalog);
+				MessageHandler->OnControllerAnalog(FGamepadKeyNames::LeftAnalogX, UserId, DeviceId, NewControllerState.LXAnalog);
 				NewControllerState.ButtonStates[MAX_NUM_PHYSICAL_CONTROLLER_BUTTONS + 1] = NewControllerState.LXAnalog >= RepeatDeadzone;
 				NewControllerState.ButtonStates[MAX_NUM_PHYSICAL_CONTROLLER_BUTTONS + 0] = NewControllerState.LXAnalog <= -RepeatDeadzone;
 			}
 			if (NewControllerState.LYAnalog != OldControllerState.LYAnalog || FMath::Abs(NewControllerState.LYAnalog) >= RepeatDeadzone)
 			{
 				//LOGD("    Sending updated LeftAnalogY value of %f", NewControllerState.LYAnalog);
-				MessageHandler->OnControllerAnalog(FGamepadKeyNames::LeftAnalogY, NewControllerState.DeviceId, NewControllerState.LYAnalog);
+				MessageHandler->OnControllerAnalog(FGamepadKeyNames::LeftAnalogY, UserId, DeviceId, NewControllerState.LYAnalog);
 				NewControllerState.ButtonStates[MAX_NUM_PHYSICAL_CONTROLLER_BUTTONS + 2] = NewControllerState.LYAnalog >= RepeatDeadzone;
 				NewControllerState.ButtonStates[MAX_NUM_PHYSICAL_CONTROLLER_BUTTONS + 3] = NewControllerState.LYAnalog <= -RepeatDeadzone;
 			}
 			if (NewControllerState.RXAnalog != OldControllerState.RXAnalog || FMath::Abs(NewControllerState.RXAnalog) >= RepeatDeadzone)
 			{
 				//LOGD("    Sending updated RightAnalogX value of %f", NewControllerState.RXAnalog);
-				MessageHandler->OnControllerAnalog(FGamepadKeyNames::RightAnalogX, NewControllerState.DeviceId, NewControllerState.RXAnalog);
+				MessageHandler->OnControllerAnalog(FGamepadKeyNames::RightAnalogX, UserId, DeviceId, NewControllerState.RXAnalog);
 				NewControllerState.ButtonStates[MAX_NUM_PHYSICAL_CONTROLLER_BUTTONS + 5] = NewControllerState.RXAnalog >= RepeatDeadzone;
 				NewControllerState.ButtonStates[MAX_NUM_PHYSICAL_CONTROLLER_BUTTONS + 4] = NewControllerState.RXAnalog <= -RepeatDeadzone;
 			}
 			if (NewControllerState.RYAnalog != OldControllerState.RYAnalog || FMath::Abs(NewControllerState.RYAnalog) >= RepeatDeadzone)
 			{
 				//LOGD("    Sending updated RightAnalogY value of %f", NewControllerState.RYAnalog);
-				MessageHandler->OnControllerAnalog(FGamepadKeyNames::RightAnalogY, NewControllerState.DeviceId, NewControllerState.RYAnalog);
+				MessageHandler->OnControllerAnalog(FGamepadKeyNames::RightAnalogY, UserId, DeviceId, NewControllerState.RYAnalog);
 				NewControllerState.ButtonStates[MAX_NUM_PHYSICAL_CONTROLLER_BUTTONS + 6] = NewControllerState.RYAnalog >= RepeatDeadzone;
 				NewControllerState.ButtonStates[MAX_NUM_PHYSICAL_CONTROLLER_BUTTONS + 7] = NewControllerState.RYAnalog <= -RepeatDeadzone;
 			}
@@ -1112,7 +1117,7 @@ void FAndroidInputInterface::SendControllerEvents()
 			if (NewControllerState.LTAnalog != OldControllerState.LTAnalog)
 			{
 				//LOGD("    Sending updated LeftTriggerAnalog value of %f", NewControllerState.LTAnalog);
-				MessageHandler->OnControllerAnalog(FGamepadKeyNames::LeftTriggerAnalog, NewControllerState.DeviceId, NewControllerState.LTAnalog);
+				MessageHandler->OnControllerAnalog(FGamepadKeyNames::LeftTriggerAnalog, UserId, DeviceId, NewControllerState.LTAnalog);
 
 				if (bUseTriggerThresholdForClick)
 				{
@@ -1124,7 +1129,7 @@ void FAndroidInputInterface::SendControllerEvents()
 			if (NewControllerState.RTAnalog != OldControllerState.RTAnalog)
 			{
 				//LOGD("    Sending updated RightTriggerAnalog value of %f", NewControllerState.RTAnalog);
-				MessageHandler->OnControllerAnalog(FGamepadKeyNames::RightTriggerAnalog, NewControllerState.DeviceId, NewControllerState.RTAnalog);
+				MessageHandler->OnControllerAnalog(FGamepadKeyNames::RightTriggerAnalog, UserId, DeviceId, NewControllerState.RTAnalog);
 
 				if (bUseTriggerThresholdForClick)
 				{
@@ -1144,12 +1149,12 @@ void FAndroidInputInterface::SendControllerEvents()
 					if (NewControllerState.ButtonStates[ButtonIndex])
 					{
 						//LOGD("    Sending joystick button down %d (first)", ButtonMapping[ButtonIndex]);
-						MessageHandler->OnControllerButtonPressed(ButtonMapping[ButtonIndex], NewControllerState.DeviceId, false);
+						MessageHandler->OnControllerButtonPressed(ButtonMapping[ButtonIndex], UserId, DeviceId, false);
 					}
 					else
 					{
 						//LOGD("    Sending joystick button up %d", ButtonMapping[ButtonIndex]);
-						MessageHandler->OnControllerButtonReleased(ButtonMapping[ButtonIndex], NewControllerState.DeviceId, false);
+						MessageHandler->OnControllerButtonReleased(ButtonMapping[ButtonIndex], UserId, DeviceId, false);
 					}
 
 					if (NewControllerState.ButtonStates[ButtonIndex])
@@ -1161,7 +1166,7 @@ void FAndroidInputInterface::SendControllerEvents()
 				else if (NewControllerState.ButtonStates[ButtonIndex] && NewControllerState.NextRepeatTime[ButtonIndex] <= CurrentTime)
 				{
 					// Send button repeat events
-					MessageHandler->OnControllerButtonPressed(ButtonMapping[ButtonIndex], NewControllerState.DeviceId, true);
+					MessageHandler->OnControllerButtonPressed(ButtonMapping[ButtonIndex], UserId, DeviceId, true);
 
 					// Set the button's NextRepeatTime to the ButtonRepeatDelay
 					NewControllerState.NextRepeatTime[ButtonIndex] = CurrentTime + ButtonRepeatDelay;
