Fix Mac CIS to use the newer version of OnController functions

#rnx

diff --git a/Engine/Source/Runtime/ApplicationCore/Private/Apple/AppleControllerInterface.cpp b/Engine/Source/Runtime/ApplicationCore/Private/Apple/AppleControllerInterface.cpp
--- a/Engine/Source/Runtime/ApplicationCore/Private/Apple/AppleControllerInterface.cpp
+++ b/Engine/Source/Runtime/ApplicationCore/Private/Apple/AppleControllerInterface.cpp
@@ -3,6 +3,7 @@
 #include "AppleControllerInterface.h"
 #include "HAL/PlatformTime.h"
 #include "Misc/ScopeLock.h"
+#include "GenericPlatform/GenericPlatformInputDeviceMapper.h"
 
 DEFINE_LOG_CATEGORY(LogAppleController);
 
@@ -238,7 +239,12 @@ void FAppleControllerInterface::SendControllerEvents()
 		}
 		
 		GCController* ControllerImpl = Controller.Controller;
-
+		
+		IPlatformInputDeviceMapper& DeviceMapper = IPlatformInputDeviceMapper::Get();
+		FPlatformUserId UserId = FGenericPlatformMisc::GetPlatformUserForUserIndex(Controller.PlayerIndex);
+		FInputDeviceId DeviceId = INPUTDEVICEID_NONE;
+		DeviceMapper.RemapControllerIdToPlatformUserAndDevice(Controller.PlayerIndex, OUT UserId, OUT DeviceId);
+		
 		// Assumes iOS13, tvOS 13 & macOS 10.15
         GCExtendedGamepad* ExtendedGamepad = [ControllerImpl capture].extendedGamepad;
 		GCMotion* Motion = ControllerImpl.motion;
@@ -246,8 +252,8 @@ void FAppleControllerInterface::SendControllerEvents()
 		// Workaround for unreliable buttonMenu behavior in iOS/tvOS 14
 		if (Controller.bPauseWasPressed)
         {
-            MessageHandler->OnControllerButtonPressed(FGamepadKeyNames::SpecialRight, Controller.PlayerIndex, false);
-            MessageHandler->OnControllerButtonReleased(FGamepadKeyNames::SpecialRight, Controller.PlayerIndex, false);
+            MessageHandler->OnControllerButtonPressed(FGamepadKeyNames::SpecialRight, UserId, DeviceId, false);
+            MessageHandler->OnControllerButtonReleased(FGamepadKeyNames::SpecialRight, UserId, DeviceId, false);
 
             Controller.bPauseWasPressed = false;
         }
@@ -348,13 +354,18 @@ void FAppleControllerInterface::HandleInputInternal(const FGamepadKeyNames::Type
     const float InitialRepeatDelay = 0.2f;
     const float RepeatDelay = 0.1;
     GCController* Cont = Controllers[ControllerIndex].Controller;
+    
+	IPlatformInputDeviceMapper& DeviceMapper = IPlatformInputDeviceMapper::Get();
+	FPlatformUserId UserId = FGenericPlatformMisc::GetPlatformUserForUserIndex(Controllers[ControllerIndex].PlayerIndex);
+	FInputDeviceId DeviceId = INPUTDEVICEID_NONE;
+	DeviceMapper.RemapControllerIdToPlatformUserAndDevice(Controllers[ControllerIndex].PlayerIndex, OUT UserId, OUT DeviceId);
 
     if (bWasPressed != bIsPressed)
     {
 #if APPLE_CONTROLLER_DEBUG
         NSLog(@"%@ button %s on controller %d", bIsPressed ? @"Pressed" : @"Released", TCHAR_TO_ANSI(*UEButton.ToString()), Controllers[ControllerIndex].PlayerIndex);
 #endif
-        bIsPressed ? MessageHandler->OnControllerButtonPressed(UEButton, Controllers[ControllerIndex].PlayerIndex, false) : MessageHandler->OnControllerButtonReleased(UEButton, Controllers[ControllerIndex].PlayerIndex, false);
+        bIsPressed ? MessageHandler->OnControllerButtonPressed(UEButton, UserId, DeviceId, false) : MessageHandler->OnControllerButtonReleased(UEButton,UserId, DeviceId, false);
         NextKeyRepeatTime.FindOrAdd(UEButton) = CurrentTime + InitialRepeatDelay;
     }
     else if(bIsPressed)
@@ -362,7 +373,7 @@ void FAppleControllerInterface::HandleInputInternal(const FGamepadKeyNames::Type
         double* NextRepeatTime = NextKeyRepeatTime.Find(UEButton);
         if(NextRepeatTime && *NextRepeatTime <= CurrentTime)
         {
-            MessageHandler->OnControllerButtonPressed(UEButton, Controllers[ControllerIndex].PlayerIndex, true);
+            MessageHandler->OnControllerButtonPressed(UEButton, UserId, DeviceId, true);
             *NextRepeatTime = CurrentTime + RepeatDelay;
         }
     }
@@ -495,6 +506,11 @@ void FAppleControllerInterface::HandleAnalogGamepad(const FGamepadKeyNames::Type
 {
     GCController* Cont = Controllers[ControllerIndex].Controller;
     
+    IPlatformInputDeviceMapper& DeviceMapper = IPlatformInputDeviceMapper::Get();
+	FPlatformUserId UserId = FGenericPlatformMisc::GetPlatformUserForUserIndex(Controllers[ControllerIndex].PlayerIndex);
+	FInputDeviceId DeviceId = INPUTDEVICEID_NONE;
+	DeviceMapper.RemapControllerIdToPlatformUserAndDevice(Controllers[ControllerIndex].PlayerIndex, OUT UserId, OUT DeviceId);
+    
     // Send controller events any time we are passed the given input threshold similarly to PC/Console (see: XInputInterface.cpp)
     const float RepeatDeadzone = 0.24;
     bool bWasPositivePressed = false;
@@ -541,5 +557,5 @@ void FAppleControllerInterface::HandleAnalogGamepad(const FGamepadKeyNames::Type
 #if APPLE_CONTROLLER_DEBUG
     NSLog(@"Axis %s is %f", TCHAR_TO_ANSI(*UEAxis.ToString()), axisValue);
 #endif
-    MessageHandler->OnControllerAnalog(UEAxis, Controllers[ControllerIndex].PlayerIndex, axisValue);
+    MessageHandler->OnControllerAnalog(UEAxis, UserId, DeviceId, axisValue);
 }
diff --git a/Engine/Source/Runtime/ApplicationCore/Private/Mac/HIDInputInterface.cpp b/Engine/Source/Runtime/ApplicationCore/Private/Mac/HIDInputInterface.cpp
--- a/Engine/Source/Runtime/ApplicationCore/Private/Mac/HIDInputInterface.cpp
+++ b/Engine/Source/Runtime/ApplicationCore/Private/Mac/HIDInputInterface.cpp
@@ -462,6 +462,11 @@ void HIDInputInterface::SendControllerEvents()
 	{
 		FControllerState& ControllerState = ControllerStates[ControllerIndex];
 
+		IPlatformInputDeviceMapper& DeviceMapper = IPlatformInputDeviceMapper::Get();
+		FPlatformUserId UserId = FGenericPlatformMisc::GetPlatformUserForUserIndex(ControllerState.ControllerId);
+		FInputDeviceId DeviceId = INPUTDEVICEID_NONE;
+		DeviceMapper.RemapControllerIdToPlatformUserAndDevice(ControllerState.ControllerId, OUT UserId, OUT DeviceId);
+
 		if( ControllerState.Device.DeviceRef )
 		{
 			bool CurrentButtonStates[MAX_NUM_CONTROLLER_BUTTONS] = {0};
@@ -491,28 +496,28 @@ void HIDInputInterface::SendControllerEvents()
 
 						if (Element.Usage == ControllerState.Device.LeftAnalogXMapping && ControllerState.LeftAnalogX != NewValue)
 						{
-							MessageHandler->OnControllerAnalog(FGamepadKeyNames::LeftAnalogX, ControllerState.ControllerId, FloatValue);
+							MessageHandler->OnControllerAnalog(FGamepadKeyNames::LeftAnalogX, UserId, DeviceId, FloatValue);
 							CurrentButtonStates[18] = FloatValue < -0.2f;	// FGamepadKeyNames::LeftStickLeft
 							CurrentButtonStates[19] = FloatValue > 0.2f;	// FGamepadKeyNames::LeftStickRight
 							ControllerState.LeftAnalogX = NewValue;
 						}
 						else if (Element.Usage == ControllerState.Device.LeftAnalogYMapping && ControllerState.LeftAnalogY != NewValue)
 						{
-							MessageHandler->OnControllerAnalog(FGamepadKeyNames::LeftAnalogY, ControllerState.ControllerId, -FloatValue);
+							MessageHandler->OnControllerAnalog(FGamepadKeyNames::LeftAnalogY, UserId, DeviceId, -FloatValue);
 							CurrentButtonStates[16] = FloatValue < -0.2f;	// FGamepadKeyNames::LeftStickUp
 							CurrentButtonStates[17] = FloatValue > 0.2f;	// FGamepadKeyNames::LeftStickDown
 							ControllerState.LeftAnalogY = NewValue;
 						}
 						else if (Element.Usage == ControllerState.Device.RightAnalogXMapping && ControllerState.RightAnalogX != NewValue)
 						{
-							MessageHandler->OnControllerAnalog(FGamepadKeyNames::RightAnalogX, ControllerState.ControllerId, FloatValue);
+							MessageHandler->OnControllerAnalog(FGamepadKeyNames::RightAnalogX, UserId, DeviceId, FloatValue);
 							CurrentButtonStates[22] = FloatValue < -0.2f;	// FGamepadKeyNames::RightStickLeft
 							CurrentButtonStates[23] = FloatValue > 0.2f;	// FGamepadKeyNames::RightStickDown
 							ControllerState.RightAnalogX = NewValue;
 						}
 						else if (Element.Usage == ControllerState.Device.RightAnalogYMapping && ControllerState.RightAnalogY != NewValue)
 						{
-							MessageHandler->OnControllerAnalog(FGamepadKeyNames::RightAnalogY, ControllerState.ControllerId, -FloatValue);
+							MessageHandler->OnControllerAnalog(FGamepadKeyNames::RightAnalogY, UserId, DeviceId, -FloatValue);
 							CurrentButtonStates[20] = FloatValue < -0.2f;	// FGamepadKeyNames::RightStickUp
 							CurrentButtonStates[21] = FloatValue > 0.2f;	// FGamepadKeyNames::RightStickRight
 							ControllerState.RightAnalogY = NewValue;
@@ -521,7 +526,7 @@ void HIDInputInterface::SendControllerEvents()
 						{
 							if (ControllerState.LeftTriggerAnalog != NewValue)
 							{
-								MessageHandler->OnControllerAnalog(FGamepadKeyNames::LeftTriggerAnalog, ControllerState.ControllerId, FloatValue);
+								MessageHandler->OnControllerAnalog(FGamepadKeyNames::LeftTriggerAnalog, UserId, DeviceId, FloatValue);
 								ControllerState.LeftTriggerAnalog = NewValue;
 							}
 							CurrentButtonStates[10] = FloatValue > 0.01f;
@@ -530,7 +535,7 @@ void HIDInputInterface::SendControllerEvents()
 						{
 							if (ControllerState.RightTriggerAnalog != NewValue)
 							{
-								MessageHandler->OnControllerAnalog(FGamepadKeyNames::RightTriggerAnalog, ControllerState.ControllerId, FloatValue);
+								MessageHandler->OnControllerAnalog(FGamepadKeyNames::RightTriggerAnalog, UserId, DeviceId, FloatValue);
 								ControllerState.RightTriggerAnalog = NewValue;
 							}
 							CurrentButtonStates[11] = FloatValue > 0.01f;
@@ -563,11 +568,11 @@ void HIDInputInterface::SendControllerEvents()
 				{
 					if (CurrentButtonStates[ButtonIndex])
 					{
-						MessageHandler->OnControllerButtonPressed(Buttons[ButtonIndex], ControllerState.ControllerId, false);
+						MessageHandler->OnControllerButtonPressed(Buttons[ButtonIndex], UserId, DeviceId, false);
 					}
 					else
 					{
-						MessageHandler->OnControllerButtonReleased(Buttons[ButtonIndex], ControllerState.ControllerId, false);
+						MessageHandler->OnControllerButtonReleased(Buttons[ButtonIndex], UserId, DeviceId, false);
 					}
 
 					if (CurrentButtonStates[ButtonIndex] != 0)
@@ -578,7 +583,7 @@ void HIDInputInterface::SendControllerEvents()
 				}
 				else if (CurrentButtonStates[ButtonIndex] != 0 && ControllerState.NextRepeatTime[ButtonIndex] <= CurrentTime)
 				{
-					MessageHandler->OnControllerButtonPressed(Buttons[ButtonIndex], ControllerState.ControllerId, true);
+					MessageHandler->OnControllerButtonPressed(Buttons[ButtonIndex], UserId, DeviceId, true);
 
 					// set the button's NextRepeatTime to the ButtonRepeatDelay
 					ControllerState.NextRepeatTime[ButtonIndex] = CurrentTime + ButtonRepeatDelay;
