CIS Fix about checking keywords. Remove the enable/disable deprecation warnings from the GenericApplicationMessageHandler header file because they were already wrapped with them.

#rnx

diff --git a/Engine/Source/Runtime/ApplicationCore/Public/GenericPlatform/GenericApplicationMessageHandler.h b/Engine/Source/Runtime/ApplicationCore/Public/GenericPlatform/GenericApplicationMessageHandler.h
--- a/Engine/Source/Runtime/ApplicationCore/Public/GenericPlatform/GenericApplicationMessageHandler.h
+++ b/Engine/Source/Runtime/ApplicationCore/Public/GenericPlatform/GenericApplicationMessageHandler.h
@@ -306,34 +306,28 @@ public:
 
 	virtual bool OnControllerAnalog(FGamepadKeyNames::Type KeyName, FPlatformUserId PlatformUserId, FInputDeviceId InputDeviceId, float AnalogValue)
 	{
-		PRAGMA_DISABLE_DEPRECATION_WARNINGS
 		if (!ShouldUsePlatformUserId())
 		{
 			return OnControllerAnalog(KeyName, PlatformUserId.GetInternalId(), AnalogValue);
 		}
-		PRAGMA_ENABLE_DEPRECATION_WARNINGS
 		return false;
 	}
 
 	virtual bool OnControllerButtonPressed(FGamepadKeyNames::Type KeyName, FPlatformUserId PlatformUserId, FInputDeviceId InputDeviceId, bool IsRepeat)
 	{
-		PRAGMA_DISABLE_DEPRECATION_WARNINGS
 		if (!ShouldUsePlatformUserId())
 		{
 			return OnControllerButtonPressed(KeyName, PlatformUserId.GetInternalId(), IsRepeat);
 		}
-		PRAGMA_ENABLE_DEPRECATION_WARNINGS
 		return false;
 	}
 
 	virtual bool OnControllerButtonReleased(FGamepadKeyNames::Type KeyName, FPlatformUserId PlatformUserId, FInputDeviceId InputDeviceId, bool IsRepeat)
 	{
-		PRAGMA_DISABLE_DEPRECATION_WARNINGS
 		if (!ShouldUsePlatformUserId())
 		{
 			return OnControllerButtonReleased(KeyName, PlatformUserId.GetInternalId(), IsRepeat);
 		}
-		PRAGMA_ENABLE_DEPRECATION_WARNINGS
 		return false;
 	}
 
