Add ShowPlatformUserSelector function as a platform-side replacement of ShowLoginUI.

 - Will only be implemented on relevant platforms.

diff --git a/Engine/Source/Runtime/ApplicationCore/Public/GenericPlatform/GenericPlatformApplicationMisc.h b/Engine/Source/Runtime/ApplicationCore/Public/GenericPlatform/GenericPlatformApplicationMisc.h
--- a/Engine/Source/Runtime/ApplicationCore/Public/GenericPlatform/GenericPlatformApplicationMisc.h
+++ b/Engine/Source/Runtime/ApplicationCore/Public/GenericPlatform/GenericPlatformApplicationMisc.h
@@ -11,6 +11,8 @@
 #include "Math/Vector2D.h"
 #include "Templates/IsFloatingPoint.h"
 #include "Templates/UnrealTypeTraits.h"
+#include "Templates/Function.h"
+#include "Misc/EnumClassFlags.h"
 
 class FAutoConsoleVariableRef;
 class FFeedbackContext;
@@ -38,6 +40,30 @@ struct FShowInputDeviceSelectorParams
 };
 typedef TFunction<void(const FShowInputDeviceSelectorParams&)> FShowInputDeviceSelectorComplete;
 
+/**
+ * Callback for when ShowPlatformUserSelector has completed.
+ */
+struct FPlatformUserSelectionCompleteParams
+{
+	FPlatformUserId SelectedUserId;
+	bool bSuccess;
+};
+typedef TFunction<void(const FPlatformUserSelectionCompleteParams& Params)> FPlatformUserSelectionComplete;
+
+/*
+ * Options for ShowPlatformUserSelector. Not all platforms will support all flags
+ */
+enum class EPlatformUserSelectorFlags : uint8
+{
+	None                         = 0,
+	RequiresOnlineEnabledProfile = (1 << 1), // whether to only show profiles that have an online account
+	ShowSkipButton               = (1 << 2), // display a 'skip' button
+	AllowGuests                  = (1 << 4), // temporary guest accounts, throwaways etc.
+
+	Default                      = ShowSkipButton, // most commonly used options
+};
+ENUM_CLASS_FLAGS(EPlatformUserSelectorFlags);
+
 
 struct APPLICATIONCORE_API FGenericPlatformApplicationMisc
 {
@@ -341,6 +367,21 @@ struct APPLICATIONCORE_API FGenericPlatformApplicationMisc
 		return false;
 	}
 
+	/**
+	 * Asyncronously display the platform-specific user selection UI, if supported.
+	 * 
+	 * @param InitiatingInputDeviceId The input device that prompted showing the UI, if applicable. The input device may be re-paired with the newly-selected user, depending on the platform
+	 * @param Flags customization options for the user selection UI.
+	 * @param OnUserSelectionComplete callback for when the user selection operation has completed
+	 * @return true if the UI will be shown and the callback will be called
+	 * @return false if the platform does not support a user-selection API, or it is not currently available (the callback will not be called)
+	 */
+	static bool ShowPlatformUserSelector(FInputDeviceId InitiatingInputDeviceId, EPlatformUserSelectorFlags Flags, FPlatformUserSelectionComplete OnUserSelectionComplete)
+	{
+		// no default implementation
+		return false;
+	}
+
 
 protected:
 	static bool CachedPhysicalScreenData;
