[Runtime ApplicationCore]
* New Console for windows. Instead of using build in console window a win32 window is used with a bunch of convenience features
  * Activity fields. Useful to visualize what is going on in the process
  * Log list view that can be navigated while new entries are being added.
  * Filters that can include/exclude log entries based on keywords
  * Console command box. Execute your console commands directly in the console window. Note that they won't run until Tick loop has started.
  * Customizable color support. Make your dedicated server console dark gray and your client console light blue.
  * Remembers size and position separately for client/server/other. This means that UnrealEngine.exe can be used for all and console can be separately placed

Is currently only enabled if "-NewConsole" exists on command line. Once new console is well tested we will probably change so built-in console can be enabled with -DefaultConsole or something. There is one advantage with the default console and that is that it is still interactive even though process is not running. Sometimes that can be useful

diff --git a/Engine/Source/Runtime/ApplicationCore/Private/Windows/WindowsConsoleOutputDevice2.cpp b/Engine/Source/Runtime/ApplicationCore/Private/Windows/WindowsConsoleOutputDevice2.cpp
new file mode 100644
--- /dev/null
+++ b/Engine/Source/Runtime/ApplicationCore/Private/Windows/WindowsConsoleOutputDevice2.cpp
@@ -0,0 +1,2247 @@
+// Copyright Epic Games, Inc. All Rights Reserved.
+
+#include "Windows/WindowsConsoleOutputDevice2.h"
+#include "Async/Async.h"
+#include "Containers/RingBuffer.h"
+#include "CoreGlobals.h"
+#include "Features/IModularFeatures.h"
+#include "GenericPlatform/GenericApplication.h"
+#include "HAL/ConsoleManager.h"
+#include "HAL/PlatformApplicationMisc.h"
+#include "HAL/PlatformTime.h"
+#include "HAL/UnrealMemory.h"
+#include "Logging/LogMacros.h"
+#include "Math/UnrealMathUtility.h"
+#include "Misc/App.h"
+#include "Misc/CommandLine.h"
+#include "Misc/ConfigCacheIni.h"
+#include "Misc/CoreDelegates.h"
+#include "Misc/CString.h"
+#include "Misc/OutputDevice.h"
+#include "Misc/OutputDeviceHelper.h"
+#include "Misc/Parse.h"
+#include "Misc/ScopeRWLock.h"
+#include "Misc/TrackedActivity.h"
+#include "String/Find.h"
+#include "Templates/UnrealTemplate.h"
+#include "Windows/WindowsHWrapper.h"
+#include "Windows/WindowsPlatformApplicationMisc.h"
+
+#include "Windows/AllowWindowsPlatformTypes.h"
+#include <CommCtrl.h>
+#include <uxtheme.h>
+#include "Windows/HideWindowsPlatformTypes.h"
+
+
+class FWindowsConsoleOutputDevice2::FConsoleWindow
+{
+public:
+	enum
+	{
+		WM_NEWLOGENTRIES = WM_USER + 1,
+		WM_ACTIVITIESDIRTY,
+	};
+	enum ENotificationId : int
+	{
+		ID_LOG = 0x8801,
+		ID_LOG_INCLUDE_FILTER,
+		ID_LOG_EXCLUDE_FILTER,
+		ID_COMMAND,
+		ID_COMMANDBUTTON,
+		ID_CLEARLOGBUTTON,
+		ID_ADDCHECKPOINTBUTTON,
+	};
+	enum : int
+	{
+		HTEXPAND = 1000, // Just some value that is not any of the existing HT-ones
+	};
+
+	FConsoleWindow(FWindowsConsoleOutputDevice2* InOwner)
+	:	Owner(InOwner)
+	{
+		// Set up status light brushes
+		COLORREF Colors[] = {RGB(237, 28, 36), RGB(255, 255, 0), RGB(128, 255, 128)};
+		for (int32 I = 0; I != 3; ++I)
+			StatusLightBrush[I] = CreateSolidBrush(Colors[I]);
+
+		// Can't remove this in dtor since it happens after OnExit delegate is destroyed`
+		FCoreDelegates::OnExit.AddLambda([this]
+			{
+				bool bSaveIni = bIsVisible;
+				bIsVisible = false;
+				if (bSaveIni)
+					Owner->SaveToINI();
+			});
+	}
+
+	~FConsoleWindow()
+	{
+		PostMessageW(MainHwnd, WM_CLOSE, 0, 0);
+		WaitForSingleObject(Thread, 1000);
+		CloseHandle(Thread);
+		Thread = nullptr;
+
+		for (int32 I=0; I!=3; ++I)
+			DeleteObject(StatusLightBrush[I]);
+	}
+
+	void AddLogEntry(const FStringView& Text, uint16 TextAttribute)
+	{
+		bool bWasEmpty;
+		{
+			FScopeLock _(&NewLogEntriesCs);
+			bWasEmpty = NewLogEntries.Num() == 0;
+			NewLogEntries.Add({FString(Text), TextAttribute});
+		}
+		if (bWasEmpty)
+			PostMessageW(MainHwnd, WM_NEWLOGENTRIES, 0, 0);
+	}
+
+	// Light 0 = none, 1 = Red, 2 = Yellow, 3 = Green
+	void SetActivity(const TCHAR* Id, const TCHAR* Name, const TCHAR* Status, int32 Light = 0, bool bAlignLeft = true)
+	{
+		if (!*Name)
+			return;
+
+		FScopeLock _(&ActivitiesCs);
+
+		bool bDirty = false;
+		if (int* Index = IdToActivityIndex.Find(Id))
+		{
+			Activity& A = Activities[*Index];
+			if (A.Status != Status)
+			{
+				A.Status = Status;
+				A.bStatusDirty = true;
+				bDirty = true;
+			}
+			if (A.Light != Light)
+			{
+				A.Light = Light;
+				bActivitiesLightDirty = true;
+				bDirty = true;
+			}
+		}
+		else
+		{
+			IdToActivityIndex.Add(Id, Activities.Num());
+			Activity& A = Activities.AddDefaulted_GetRef();;
+			A.bAlignLeft = bAlignLeft;
+			A.Name = Name;
+			A.Status = Status;
+			A.Light = Light;
+			if (Light != 0)
+				bActivitiesLightDirty = true;
+			bDirty = true;
+		}
+
+		if (!bDirty || bActivitiesDirty)
+			return;
+		bActivitiesDirty = true;
+		PostMessageW(MainHwnd, WM_ACTIVITIESDIRTY, 0, 0);
+	}
+
+	void RemoveStatus(const TCHAR* Id)
+	{
+		FScopeLock _(&ActivitiesCs);
+
+		int I;
+		if (!IdToActivityIndex.RemoveAndCopyValue(Id, I))
+			return;
+		Activities[I].Name.Reset();
+		bActivitiesDirty = true;
+		PostMessageW(MainHwnd, WM_ACTIVITIESDIRTY, 0, 0);
+	}
+
+	void Start()
+	{
+		Thread = CreateThread(NULL, 0, StaticThreadProc, this, 0, NULL);
+	}
+
+	static DWORD WINAPI StaticThreadProc(LPVOID lpParameter)
+	{
+		return ((FConsoleWindow*)lpParameter)->ThreadProc();
+	}
+
+	uint32 ThreadProc()
+	{
+		HINSTANCE HInstance = (HINSTANCE)GetModuleHandle(NULL);
+
+		CreateColors();
+
+		Icon = LoadIcon(HInstance, MAKEINTRESOURCE(FWindowsPlatformApplicationMisc::GetAppIcon()));
+
+		WNDCLASSEX WndClassEx;
+		ZeroMemory(&WndClassEx, sizeof(WndClassEx));
+		WndClassEx.cbSize = sizeof(WndClassEx);
+		WndClassEx.style = CS_HREDRAW | CS_VREDRAW;
+		WndClassEx.lpfnWndProc = &StaticMainWinProc;
+		WndClassEx.hIcon = Icon;
+		WndClassEx.hCursor = LoadCursor(NULL, IDC_ARROW);
+		WndClassEx.hInstance = HInstance;
+		WndClassEx.hbrBackground = NULL;
+		WndClassEx.lpszClassName = TEXT("FConsoleWindow");
+		ATOM WndClassAtom = RegisterClassEx(&WndClassEx);
+
+		WndClassEx.lpfnWndProc = &StaticScrollBarWinProc;
+		WndClassEx.lpszClassName = TEXT("FScrollBar");
+		ATOM ScrollBarClassAtom = RegisterClassEx(&WndClassEx);
+
+		NONCLIENTMETRICS NonClientMetrics;
+		NonClientMetrics.cbSize = sizeof(NonClientMetrics);
+		SystemParametersInfo(SPI_GETNONCLIENTMETRICS, sizeof(NonClientMetrics), &NonClientMetrics, 0);
+		Font = (HFONT)CreateFontIndirect(&NonClientMetrics.lfMessageFont);
+
+		int32 FontHeight = -MulDiv(8, GetDeviceCaps(GetDC(NULL), LOGPIXELSY), 72);
+		HFONT hLogFont = (HFONT)CreateFontW(FontHeight, 0, 0, 0, FW_NORMAL, false, false, false, ANSI_CHARSET, OUT_DEFAULT_PRECIS, CLIP_DEFAULT_PRECIS, ANTIALIASED_QUALITY, FIXED_PITCH | FF_MODERN, TEXT("Courier New"));
+
+		DWORD WindowStyle = WS_CAPTION | WS_SYSMENU | WS_THICKFRAME | WS_MAXIMIZEBOX | WS_MINIMIZEBOX | WS_CLIPCHILDREN;
+		const TCHAR* WindowClassName = MAKEINTATOM(WndClassAtom);
+		MainHwnd = CreateWindow(WindowClassName, *ConsoleTitle, WindowStyle, ConsolePosX, ConsolePosY, ConsoleWidth, ConsoleHeight, NULL, NULL, HInstance, this);
+		SetWindowLongPtr(MainHwnd, GWLP_USERDATA, (LONG_PTR)this);
+
+		int32 X = 20;
+		int32 Y = 10;
+		int32 Height = 24;
+		FilterHwnd[0] = CreateTextHwnd(TEXT("Log Filter"), Font, X, Y, 70, 20);
+		X += 100;
+		FilterHwnd[1] = CreateTextHwnd(TEXT("Include:"), Font, X, Y, 47, 20);
+		X += 50;
+		FilterHwnd[2] = CreateWindow2(WC_EDIT, NULL, ES_AUTOHSCROLL, Font, X, Y - 1, 200, 20, ID_LOG_INCLUDE_FILTER);
+		X += 240;
+		FilterHwnd[3] = CreateTextHwnd(TEXT("Exclude:"), Font, X, Y, 47, 20);
+		X += 50;
+		FilterHwnd[4] = CreateWindow2(WC_EDIT, NULL, ES_AUTOHSCROLL, Font, X, Y - 1, 200, 20, ID_LOG_EXCLUDE_FILTER);
+		X += 250;
+		FilterHwnd[5] = CreateWindow2(WC_BUTTON, TEXT("Clear Log"), 0, Font, X, Y - 2, 80, 22, ID_CLEARLOGBUTTON);
+		X += 90;
+		FilterHwnd[6] = CreateWindow2(WC_BUTTON, TEXT("Log CHECKPOINT0"), 0, Font, X, Y - 2, 140, 22, ID_ADDCHECKPOINTBUTTON);
+
+		DWORD LogHwndStyle = LBS_NOINTEGRALHEIGHT | LBS_EXTENDEDSEL | LBS_WANTKEYBOARDINPUT | LBS_NOTIFY | LBS_HASSTRINGS | LBS_OWNERDRAWFIXED;
+		LogHwnd = CreateWindow2(WC_LISTBOX, NULL, LogHwndStyle, hLogFont, 0, 0, 1, 1, ID_LOG);
+		SetWindowLongPtr(LogHwnd, GWLP_USERDATA, (LONG_PTR)this);
+		LogHwndWndProc = (WNDPROC)SetWindowLongPtr(LogHwnd, GWLP_WNDPROC, (LONG_PTR)StaticLogHwndWndProc);
+
+		const TCHAR* ScrollBarClassName = MAKEINTATOM(ScrollBarClassAtom);
+		LogScrollHwnd = CreateWindow2(ScrollBarClassName, NULL, SBS_VERT, 0);
+		SetWindowLongPtr(LogScrollHwnd, GWLP_USERDATA, (LONG_PTR)this);
+
+		CommandHwnd = CreateWindow2(WC_EDIT, NULL, ES_AUTOHSCROLL, Font, 0, 0, 1, 1, ID_COMMAND);
+		RunCommandHwnd = CreateWindow2(WC_BUTTON, TEXT("RunCommand"), 0, Font, 0, 0, 1, 1, ID_COMMANDBUTTON);
+
+		SetFocus(CommandHwnd);
+
+		SIZE FontSize;
+		GetTextExtentPoint32(GetDC(LogHwnd), TEXT("A"), 1, &FontSize);
+		LogFontWidth = FontSize.cx;
+		LogFontHeight = FontSize.cy - 2; // We want more compact listbox
+
+		RECT Rect;
+		GetClientRect(MainHwnd, &Rect);
+		UpdateSize(Rect.right, Rect.bottom, false);
+
+		UpdateTime(0);
+
+		FTrackedActivity::RegisterEventListener([&](FTrackedActivity::EEvent Event, const FTrackedActivity::FInfo& Info)
+			{
+				TCHAR Id[256];
+				const TCHAR* Str = Info.Name;
+				FCString::Sprintf(Id, TEXT("%u%s"), Info.Id, Str);
+
+				if (Event != FTrackedActivity::EEvent::Removed)
+					SetActivity(Id, Str, Info.Status, (int32)Info.Light, Info.Type == FTrackedActivity::EType::Activity);
+				else
+					RemoveStatus(Id);
+			});
+
+		FTrackedActivity::TraverseActivities([&](const FTrackedActivity::FInfo& Info)
+			{
+				TCHAR Id[256];
+				const TCHAR* Str = Info.Name;
+				FCString::Sprintf(Id, TEXT("%u%s"), Info.Id, Str);
+				SetActivity(Id, Str, Info.Status, (int32)Info.Light);
+			});
+
+		UpdateWindow(MainHwnd);
+
+		SetForegroundWindow(MainHwnd);
+
+		PostMessageW(MainHwnd, WM_NEWLOGENTRIES, 0, 0);
+
+		FDateTime startTime(FDateTime::Now());
+		uint64 lastSeconds = 0;
+		bool bLoop = true;
+		bool bWasVisible = false;
+		while (bLoop)
+		{
+			if (bWasVisible != bIsVisible)
+			{
+				if (bIsVisible)
+					MoveWindow(MainHwnd, ConsolePosX, ConsolePosY, ConsoleWidth, ConsoleHeight, true);
+				ShowWindow(MainHwnd, bIsVisible ? SW_SHOW : SW_HIDE);
+				bWasVisible = bIsVisible;
+			}
+
+			uint64 seconds = uint64((FDateTime::Now() - startTime).GetTotalSeconds());
+			if (lastSeconds != seconds)
+			{
+				lastSeconds = seconds;
+				UpdateTime(seconds);
+			}
+
+			DWORD timeout = 200;
+			if (MsgWaitForMultipleObjects(0, NULL, false, timeout, QS_ALLINPUT) == WAIT_OBJECT_0)
+			{
+				MSG Msg;
+				while (PeekMessage(&Msg, NULL, 0, 0, PM_REMOVE))
+				{
+					if (!IsDialogMessage(MainHwnd, &Msg))
+					{
+						TranslateMessage(&Msg);
+						DispatchMessage(&Msg);
+					}
+					if (Msg.message == WM_QUIT)
+					{
+						DestroyWindow(MainHwnd);
+						DeleteObject(hLogFont);
+						DeleteObject(Font);
+						UnregisterClass(WindowClassName, HInstance);
+						MainHwnd = 0;
+						bLoop = false;
+						break;
+					}
+				}
+			}
+		}
+
+		DestroyColors();
+
+		if (!IsEngineExitRequested())
+		{
+			Owner->SaveToINI();
+
+			FWindowsPlatformMisc::CallGracefulTerminationHandler();
+		}
+
+		return 0;
+	}
+
+	void CreateColors()
+	{
+		COLORREF BorderColor = RGB(0,0,0);
+		COLORREF LogBackgroundColor = RGB(0, 0, 0);
+		switch (ConsoleColor)
+		{
+		default:
+		case EConsoleColor::DarkGray:
+			BackgroundColor = RGB(32, 32, 32);
+			TextColor = RGB(220, 220, 220);
+			EditBackgroundColor = RGB(64, 64, 64);
+			StatusBackgroundColor[0] = RGB(42, 42, 42);
+			StatusBackgroundColor[1] = RGB(49, 49, 49);
+			ButtonColor = RGB(48, 48, 48);
+			ButtonHighlightColor = RGB(100, 100, 100);
+			ButtonPressedColor = RGB(130, 130, 100);
+			ThumbColor = RGB(77, 77, 77);
+			ScrollBackgroundColor = RGB(48, 48, 48);
+			BorderColor = RGB(90, 90, 90);
+			break;
+
+		case EConsoleColor::Gray:
+			BackgroundColor = GetSysColor(COLOR_3DFACE);
+			TextColor = RGB(0, 0, 0);
+			EditBackgroundColor = BackgroundColor;
+			StatusBackgroundColor[0] = GetSysColor(COLOR_3DLIGHT);
+			StatusBackgroundColor[1] = GetSysColor(COLOR_BTNHIGHLIGHT);
+			ButtonColor = GetSysColor(COLOR_3DLIGHT);
+			ButtonHighlightColor = GetSysColor(COLOR_GRADIENTINACTIVECAPTION);
+			ButtonPressedColor = GetSysColor(COLOR_GRADIENTACTIVECAPTION);
+			ThumbColor = GetSysColor(COLOR_BTNSHADOW);
+			ScrollBackgroundColor = ButtonColor;
+			BorderColor = GetSysColor(COLOR_WINDOWFRAME);
+			break;
+
+		case EConsoleColor::Red:
+			BackgroundColor = RGB(200, 50, 50);
+			TextColor = RGB(230, 230, 230);
+			EditBackgroundColor = RGB(160, 50, 50);
+			StatusBackgroundColor[0] = RGB(190, 40, 40);
+			StatusBackgroundColor[1] = RGB(210, 60, 60);
+			ButtonColor = RGB(160, 50, 50);
+			ButtonHighlightColor = RGB(190, 50, 50);
+			ButtonPressedColor = RGB(180, 50, 50);
+			ThumbColor = RGB(120, 50, 50);
+			ScrollBackgroundColor = EditBackgroundColor;
+			BorderColor = RGB(100, 40, 40);
+			break;
+
+		case EConsoleColor::LightBlue:
+			BackgroundColor = RGB(115, 211, 244);
+			LogBackgroundColor = RGB(0, 15, 20);
+			TextColor = RGB(0, 0, 0);
+			EditBackgroundColor = RGB(100, 190, 220);
+			StatusBackgroundColor[0] = RGB(105, 201, 234);
+			StatusBackgroundColor[1] = RGB(125, 221, 254);
+			ButtonColor = RGB(100, 190, 220);
+			ButtonHighlightColor = RGB(115, 211, 244);
+			ButtonPressedColor = RGB(125, 221, 255);
+			ThumbColor = RGB(60, 150, 180);
+			ScrollBackgroundColor = RGB(100, 190, 220);
+			BorderColor = RGB(30, 30, 220);
+			break;
+
+		case EConsoleColor::DarkBlue:
+			BackgroundColor = RGB(10, 10, 150);
+			TextColor = RGB(220, 220, 220);
+			EditBackgroundColor = RGB(10, 10, 100);
+			StatusBackgroundColor[0] = RGB(20, 20, 190);
+			StatusBackgroundColor[1] = RGB(20, 20, 220);
+			ButtonColor = RGB(10, 10, 100);
+			ButtonHighlightColor = RGB(10, 10, 190);
+			ButtonPressedColor = RGB(10, 10, 150);
+			ThumbColor = RGB(10, 10, 220);
+			ScrollBackgroundColor = RGB(10, 10, 120);
+			BorderColor = RGB(30, 30, 220);
+			break;
+
+		case EConsoleColor::Load:
+		{
+			FString Filename = TEXT("DebugConsoleColors.ini");
+			const TCHAR* Selection = TEXT("Colors");
+			FConfigCacheIni Config(EConfigCacheType::Temporary);
+			Config.LoadFile(Filename);
+			FColor Color;
+			if (Config.GetColor(Selection, TEXT("Background"), Color, Filename))
+				BackgroundColor = Color.ToPackedARGB();
+			if (Config.GetColor(Selection, TEXT("Text"), Color, Filename))
+				TextColor = Color.ToPackedARGB();
+			if (Config.GetColor(Selection, TEXT("LogBackground"), Color, Filename))
+				LogBackgroundColor = Color.ToPackedARGB();
+			if (Config.GetColor(Selection, TEXT("EditBackground"), Color, Filename))
+				EditBackgroundColor = Color.ToPackedARGB();
+			if (Config.GetColor(Selection, TEXT("StatusBackground1"), Color, Filename))
+				StatusBackgroundColor[0] = Color.ToPackedARGB();
+			if (Config.GetColor(Selection, TEXT("StatusBackground2"), Color, Filename))
+				StatusBackgroundColor[1] = Color.ToPackedARGB();
+			if (Config.GetColor(Selection, TEXT("EditBackground"), Color, Filename))
+				EditBackgroundColor = Color.ToPackedARGB();
+			if (Config.GetColor(Selection, TEXT("Button"), Color, Filename))
+				ButtonColor = Color.ToPackedARGB();
+			if (Config.GetColor(Selection, TEXT("ButtonHighlight"), Color, Filename))
+				ButtonHighlightColor = Color.ToPackedARGB();
+			if (Config.GetColor(Selection, TEXT("ButtonPressed"), Color, Filename))
+				ButtonPressedColor = Color.ToPackedARGB();
+			if (Config.GetColor(Selection, TEXT("Thumb"), Color, Filename))
+				ThumbColor = Color.ToPackedARGB();
+			if (Config.GetColor(Selection, TEXT("ScrollBackground"), Color, Filename))
+				ScrollBackgroundColor = Color.ToPackedARGB();
+			if (Config.GetColor(Selection, TEXT("Border"), Color, Filename))
+				BorderColor = Color.ToPackedARGB();
+			break;
+		}
+		}
+
+		TextBrush = CreateSolidBrush(TextColor);
+		BackgroundBrush = CreateSolidBrush(BackgroundColor);
+		LogBackgroundBrush = CreateSolidBrush(LogBackgroundColor);
+		EditBackgroundBrush = CreateSolidBrush(EditBackgroundColor);
+		StatusBackgroundBrush[0] = CreateSolidBrush(StatusBackgroundColor[0]);
+		StatusBackgroundBrush[1] = CreateSolidBrush(StatusBackgroundColor[1]);
+		ButtonBrush = CreateSolidBrush(ButtonColor);
+		ButtonHighlightBrush = CreateSolidBrush(ButtonHighlightColor);
+		ButtonPressedBrush = CreateSolidBrush(ButtonPressedColor);
+		ThumbBrush = CreateSolidBrush(ThumbColor);
+		ScrollBackgroundBrush = CreateSolidBrush(ScrollBackgroundColor);
+
+		NoPen = CreatePen(PS_SOLID, 0, TextColor);
+		BorderPen = CreatePen(PS_SOLID, 1, BorderColor);
+		WindowBorderPen = CreatePen(PS_INSIDEFRAME, 1, BorderColor);
+		TextPen = CreatePen(PS_SOLID, 1, TextColor);
+	}
+
+	void DestroyColors()
+	{
+		DeleteObject(NoPen);
+		DeleteObject(BorderPen);
+		DeleteObject(WindowBorderPen);
+		DeleteObject(TextPen);
+
+		DeleteObject(TextBrush);
+		DeleteObject(BackgroundBrush);
+		DeleteObject(LogBackgroundBrush);
+		DeleteObject(EditBackgroundBrush);
+		DeleteObject(StatusBackgroundBrush[0]);
+		DeleteObject(StatusBackgroundBrush[1]);
+		DeleteObject(ButtonBrush);
+		DeleteObject(ButtonHighlightBrush);
+		DeleteObject(ThumbBrush);
+	}
+
+	HWND CreateWindow2(const TCHAR* ClassName, const TCHAR* Str, DWORD Style, HFONT hFont, int32 X = 0, int32 Y = 0, int32 Width = 1, int32 Height = 1, ENotificationId InNid = (ENotificationId)0, DWORD exStyle = 0)
+	{
+		HWND h = CreateWindowEx(exStyle, ClassName, Str, Style | WS_VISIBLE | WS_CHILD, X, Y, Width, Height, MainHwnd, (HMENU)InNid, NULL, NULL);
+		if (hFont)
+			SendMessageW(h, WM_SETFONT, (WPARAM)hFont, 0);
+		return h;
+	}
+
+	HWND CreateTextHwnd(const TCHAR* Str, HFONT hFont, int32 X = 0, int32 Y = 0, int32 Width = 1, int32 Height = 1, DWORD exStyle = 0)
+	{
+		HWND h = CreateWindow2(WC_STATIC, Str, SS_OWNERDRAW, hFont, X, Y, Width, Height, (ENotificationId)0, exStyle);
+		SetWindowLong(h, GWLP_USERDATA, -1);
+		return h;
+	}
+
+	template<class Lambda>
+	int32 TraverseActivityPositions(int32 Width, int32 Height, const Lambda& InLamda)
+	{
+		Width -= 6; // Margins
+
+		int32 LeftCount = 0;
+		uint32 ActivityCount = 0;
+		for (int32 i = 0, e = Activities.Num(); i != e; ++i)
+		{
+			Activity& A = Activities[i];
+			if (A.Name.IsEmpty())
+				continue;
+			++ActivityCount;
+			if (A.bAlignLeft)
+				++LeftCount;
+		}
+		int32 RightCount = ActivityCount - LeftCount;
+
+		int32 LeftColumnMinWidth = 600;
+		int32 RightColumnWidth = 200;
+		int32 LeftColumnCount = 1;
+		int32 RightColumnCount = 1;
+
+		int32 RowCount = FMath::Max(LeftCount, RightCount);
+
+		if (RightCount > LeftCount)
+		{
+			if (Width - LeftColumnMinWidth >= RightColumnWidth * 2)
+			{
+				RightColumnCount = 2;
+				RowCount = FMath::Max(LeftCount, (RightCount+1)/2);
+			}
+		}
+
+
+		int32 RowHeight = 18;
+		int32 TotalHeight = RowCount * RowHeight + 8;
+
+
+		int32 StartY = Height - TotalHeight + 2;
+
+		int32 X;
+		int32 Y;
+		int32 RowIndex;
+		int32 ColWidth;
+		int32 ColOffset;
+
+		auto IterateActivities = [&](bool bAlignLeft)
+		{
+			for (int32 i = 0, e = Activities.Num(); i != e; ++i)
+			{
+				Activity& A = Activities[i];
+				if (A.Name.IsEmpty() || A.bAlignLeft != bAlignLeft)
+					continue;
+
+				InLamda(A, X, Y, ColWidth, RowIndex);
+				Y += RowHeight;
+				++RowIndex;
+				if (RowIndex < RowCount)
+					continue;
+				RowIndex = 0;
+				Y = StartY;
+				X += ColOffset;
+			}
+		};
+
+		X = 7;
+		Y = StartY;
+		RowIndex = 0;
+		int TotalLeftWidth = Width - RightColumnCount * RightColumnWidth - X;
+		ColWidth = TotalLeftWidth / LeftColumnCount;
+		ColOffset = ColWidth - 2;
+		IterateActivities(true);
+
+		X = Width - RightColumnWidth + 8;
+		Y = StartY;
+		RowIndex = 0;
+		ColWidth = RightColumnWidth - 8;
+		ColOffset = -RightColumnWidth;
+
+		IterateActivities(false);
+		return TotalHeight;
+	}
+
+	void UpdateSize(int32 ClientWidth, int32 ClientHeight, bool bRedraw)
+	{
+		if (bRedraw)
+			SendMessageW(LogHwnd, WM_SETREDRAW, false, 0);
+
+		int32 Flags = SWP_NOZORDER | SWP_NOACTIVATE | SWP_NOOWNERZORDER;
+		if (!bRedraw)
+			Flags |= SWP_NOREDRAW;
+
+		ActivitiesTotalHeight = UpdateActivityPositions(ClientWidth, ClientHeight, bRedraw);
+
+		int32 LogTop = 1;
+		if (bConsoleExpanded)
+			LogTop += 35;
+
+		int32 ButtonWidth = 90;
+		int32 CommandFlags = Flags | (bConsoleExpanded ? SWP_SHOWWINDOW : SWP_HIDEWINDOW);
+		int32 FilterFlags = SWP_NOMOVE | SWP_NOSIZE | Flags;
+
+		int32 CommandY = ClientHeight - ActivitiesTotalHeight - 20 - 5;
+		int32 LogHeight = CommandY - LogTop - 5;
+
+		if (bConsoleExpanded)
+		{
+			CommandFlags = Flags | SWP_SHOWWINDOW;
+			FilterFlags |= SWP_SHOWWINDOW;
+		}
+		else
+		{
+			CommandFlags = SWP_HIDEWINDOW;
+			FilterFlags |= SWP_HIDEWINDOW;
+			LogHeight += 25;
+		}
+
+		for (HWND Hwnd : FilterHwnd)
+			SetWindowPos(Hwnd, 0, 0, 0, 0, 0, FilterFlags);
+
+		SetWindowPos(LogHwnd, 0, 7, LogTop, ClientWidth - 33, LogHeight, Flags);
+		SetWindowPos(LogScrollHwnd, 0, ClientWidth - 26, LogTop, 19, LogHeight, Flags);
+		SetWindowPos(CommandHwnd, 0, 8, CommandY, ClientWidth - ButtonWidth - 20, 20, CommandFlags);
+		SetWindowPos(RunCommandHwnd, 0, ClientWidth - ButtonWidth - 7, CommandY - 1, ButtonWidth, 22, CommandFlags);
+
+		if (bRedraw)
+			SendMessageW(LogHwnd, WM_SETREDRAW, true, 0);
+	}
+
+	int32 UpdateActivityPositions(int32 ClientWidth, int32 ClientHeight, bool bRedraw)
+	{
+		int32 Flags = SWP_NOZORDER | SWP_NOACTIVATE | SWP_NOOWNERZORDER;
+		if (!bRedraw)
+			Flags |= SWP_NOREDRAW;
+		return TraverseActivityPositions(ClientWidth, ClientHeight, [&](Activity& A, int32 X, int32 Y, int32 Width, int32 RowIndex)
+			{
+				X += 16;
+				Width -= 16;
+				int32 NameWidth = FMath::Min(90, Width);
+				SetWindowPos(A.NameHwnd, 0, X, Y, NameWidth, 18, Flags);
+				SetWindowLong(A.NameHwnd, GWLP_USERDATA, RowIndex);
+				int32 StatusWidth = FMath::Max(Width - NameWidth, 0);
+				SetWindowPos(A.StatusHwnd, 0, X + NameWidth, Y, StatusWidth, 18, Flags);
+				SetWindowLong(A.StatusHwnd, GWLP_USERDATA, RowIndex);
+			});
+	}
+
+	void UpdateTime(uint64 seconds)
+	{
+		FTimespan span = FTimespan::FromSeconds(seconds);
+		FString Res;
+		if (int32 days = span.GetDays())
+			Res = FString::Printf(TEXT("%i."), FMath::Abs(days));
+		Res = FString::Printf(TEXT("%s%02i:%02i:%02i"), *Res, FMath::Abs(span.GetHours()), FMath::Abs(span.GetMinutes()), FMath::Abs(span.GetSeconds()));
+		SetActivity(TEXT("Time"), TEXT("Time"), *Res, 0, false);
+	}
+
+	struct LogEntry;
+
+	int32 AddEntryToLogHwnd(LogEntry& E, int32 LogVirtualIndex)
+	{
+		const TCHAR* Str = *E.String;
+		for (const FString& I : IncludeFilter)
+			if (FCString::Stristr(Str, *I) == 0)
+				return -1;
+		for (const FString& I : ExcludeFilter)
+			if (FCString::Stristr(Str, *I) != 0)
+				return -1;
+		AddedEntryLogVirtualIndex = LogVirtualIndex;
+		int32 ItemIndex = SendMessageW(LogHwnd, LB_ADDSTRING, 0, (LPARAM)Str);
+		SendMessageW(LogHwnd, LB_SETITEMDATA, ItemIndex, LogVirtualIndex);
+		AddedEntryLogVirtualIndex = -1;
+		return ItemIndex;
+	}
+
+	void RefreshLogHwnd()
+	{
+		int32 SelectedLogIndex = -1;
+		int32 SelectedItemOffsetFromTop = 0;
+		if (SendMessageW(LogHwnd, LB_GETSELCOUNT, 0, 0) == 1)
+		{
+			int32 SelectedItemIndex = -1;
+			SendMessageW(LogHwnd, LB_GETSELITEMS, 1, (LPARAM)&SelectedItemIndex);
+			SelectedLogIndex = SendMessageW(LogHwnd, LB_GETITEMDATA, SelectedItemIndex, 0);
+			SelectedItemOffsetFromTop = SelectedItemIndex - SendMessageW(LogHwnd, LB_GETTOPINDEX, 0, 0);
+		}
+
+		LogIndexOffset = 0;
+
+		SendMessageW(LogHwnd, WM_SETREDRAW, false, 0);
+		SendMessageW(LogHwnd, LB_RESETCONTENT, 0, 0);
+		SendMessageW(LogHwnd, LB_INITSTORAGE, Log.Num(), Log.Num()*256);
+		int32 NewSelectedItemIndex = -1;
+		int32 LogIndex = 0;
+		for (LogEntry& E : Log)
+		{
+			int32 ItemIndex = AddEntryToLogHwnd(E, LogIndex);
+			if (LogIndex == SelectedLogIndex)
+				NewSelectedItemIndex = ItemIndex;
+			++LogIndex;
+		}
+
+		int32 ScrollPos = -1;
+		if (NewSelectedItemIndex != -1)
+		{
+			ScrollPos = NewSelectedItemIndex - SelectedItemOffsetFromTop;
+			SendMessageW(LogHwnd, LB_SETTOPINDEX, ScrollPos, 0);
+			SendMessageW(LogHwnd, LB_SETSEL, true, NewSelectedItemIndex);
+		}
+		else
+			ScrollToBottom();
+
+		RedrawLogScrollbar();
+
+		SendMessageW(LogHwnd, WM_SETREDRAW, true, 0);
+		RedrawWindow(LogScrollHwnd, NULL, NULL, RDW_INVALIDATE);
+	}
+
+	void ScrollToBottom()
+	{
+		int32 ItemCount = (int)SendMessageW(LogHwnd, LB_GETCOUNT, 0, 0);
+
+		// This is here just to force last selection to be at the bottom
+		SendMessageW(LogHwnd, LB_SETSEL, true, ItemCount - 1);
+		SendMessageW(LogHwnd, LB_SETSEL, false, ItemCount - 1);
+
+		PostMessageW(LogHwnd, LB_SETTOPINDEX, ItemCount-1, 0); // Using post to remove weird glitches with smooth scrolling
+	}
+
+	static LRESULT CALLBACK StaticLogHwndWndProc(HWND hWnd, UINT Msg, WPARAM wParam, LPARAM lParam)
+	{
+		FConsoleWindow* ThisPtr = (FConsoleWindow*)GetWindowLongPtr(hWnd, GWLP_USERDATA);
+		if (Msg == WM_MOUSEWHEEL)
+			return ThisPtr->HandleMouseWheel(wParam);
+		return ThisPtr->LogHwndWndProc(hWnd, Msg, wParam, lParam);
+	}
+
+	static LRESULT CALLBACK StaticMainWinProc(HWND hWnd, UINT Msg, WPARAM wParam, LPARAM lParam)
+	{
+		FConsoleWindow* ThisPtr = (FConsoleWindow*)GetWindowLongPtr(hWnd, GWLP_USERDATA);
+		if (!ThisPtr && Msg == WM_CREATE)
+		{
+			#pragma comment( lib, "UxTheme")
+			SetWindowTheme(hWnd, L"", L""); // Needed to disable rounded edges on window
+			ThisPtr = (FConsoleWindow*)lParam;
+			SetWindowLongPtr(hWnd, GWLP_USERDATA, lParam);
+		}
+		return ThisPtr->MainWinProc(hWnd, Msg, wParam, lParam);
+	}
+
+	static LRESULT CALLBACK StaticScrollBarWinProc(HWND hWnd, UINT Msg, WPARAM wParam, LPARAM lParam)
+	{
+		if (FConsoleWindow* ThisPtr = (FConsoleWindow*)GetWindowLongPtr(hWnd, GWLP_USERDATA))
+			return ThisPtr->ScrollBarWinProc(hWnd, Msg, wParam, lParam);
+		return DefWindowProc(hWnd, Msg, wParam, lParam);
+	}
+
+	void GetThumbPosAndHeight(int32& OutPos, int32& OutHeight, const RECT& ClientRect)
+	{
+		int32 ButtonHeight = ClientRect.right + 1;
+
+		int32 TotalCount = (int)SendMessageW(LogHwnd, LB_GETCOUNT, 0, 0);
+		int32 PageSize = ClientRect.bottom / LogFontHeight;
+
+		int32 TotalScrollCount = TotalCount - PageSize;
+		int32 ScrollHeight = ClientRect.bottom - ButtonHeight * 2;
+
+		OutHeight = 0;
+
+		if (TotalScrollCount <= 0)
+			return;
+
+		int32 TopVisible = SendMessageW(LogHwnd, LB_GETTOPINDEX, 0, 0);
+
+		int32 ThumbHeight = FMath::Max((PageSize * ScrollHeight) / TotalCount, int32(ButtonHeight));
+
+		int32 MoveSpace = ScrollHeight - ThumbHeight;
+		if (MoveSpace <= 0)
+			return;
+
+		int32 ThumbPos = ButtonHeight + FMath::CeilToInt(float(TopVisible * MoveSpace) / TotalScrollCount);
+
+		OutPos = ThumbPos;
+		OutHeight = ThumbHeight;
+	}
+
+	int32 GetTopVisible(int32 ThumbPos, const RECT& ClientRect)
+	{
+		int32 ButtonHeight = ClientRect.right + 1;
+
+		int32 TotalCount = (int)SendMessageW(LogHwnd, LB_GETCOUNT, 0, 0);
+		int32 PageSize = ClientRect.bottom / LogFontHeight;
+
+		int32 TotalScrollCount = TotalCount - PageSize;
+		int32 ScrollHeight = ClientRect.bottom - ButtonHeight * 2;
+
+		int32 ThumbHeight = FMath::Max((PageSize * ScrollHeight) / TotalCount, int32(ButtonHeight));
+
+		int32 MoveSpace = ScrollHeight - ThumbHeight;
+
+		return FMath::FloorToInt(float(ThumbPos - ButtonHeight) * TotalScrollCount / MoveSpace);
+	}
+
+	void MoveTopVisible(const RECT& Rect, int32 Offset)
+	{
+		bAutoScrollLog = false;
+		int32 TopVisible = SendMessageW(LogHwnd, LB_GETTOPINDEX, 0, 0);
+		int32 PageSize = Rect.bottom / LogFontHeight;
+		TopVisible = FMath::Max(0, TopVisible + Offset);
+		SendMessageW(LogHwnd, LB_SETTOPINDEX, TopVisible, 0);
+		RedrawWindow(LogScrollHwnd, NULL, NULL, RDW_INVALIDATE);
+	}
+
+	int32 HandleMouseWheel(WPARAM wParam)
+	{
+		RECT Rect;
+		GetClientRect(LogHwnd, &Rect);
+		MoveTopVisible(Rect, -GET_WHEEL_DELTA_WPARAM(wParam) * 3 / WHEEL_DELTA);
+		return 0;
+	}
+
+	void RedrawLogScrollbar()
+	{
+		RedrawWindow(LogScrollHwnd, NULL, NULL, RDW_INVALIDATE);
+	}
+
+	LRESULT ScrollBarWinProc(HWND hWnd, UINT Msg, WPARAM wParam, LPARAM lParam)
+	{
+		switch (Msg)
+		{
+		case WM_ERASEBKGND:
+		{
+			HDC hdc = (HDC)wParam;
+			RECT Rect;
+			GetClientRect(hWnd, &Rect);
+			int32 ButtonHeight = Rect.right + 1;
+
+			RECT ScrollRect{ 0, ButtonHeight, Rect.right, Rect.bottom - ButtonHeight };
+			FillRect(hdc, &ScrollRect, ScrollBackgroundBrush);
+
+			RECT ButtonRect{ 0, 0, Rect.right, ButtonHeight };
+			FillRect(hdc, &ButtonRect, ButtonBrush);
+
+			RECT ButtonRect2{ 0, Rect.bottom - ButtonHeight, Rect.right, Rect.bottom };
+			FillRect(hdc, &ButtonRect2, ButtonBrush);
+
+			SelectObject(hdc, NoPen);
+			SelectObject(hdc, TextBrush);
+			int32 X = Rect.right / 2;
+			int32 Y = Rect.right / 2;
+			POINT vertices[] = { {X - 3, Y + 1}, {X, Y - 2}, {X + 3, Y + 1} };
+			Polygon(hdc, vertices, 3);
+
+			Y = Rect.bottom - Y;
+			POINT vertices2[] = { {X - 3, Y - 2}, {X, Y + 1}, {X + 3, Y - 2} };
+			Polygon(hdc, vertices2, 3);
+
+			return 1;
+		}
+		case WM_PAINT:
+		{
+			RECT Rect;
+			GetClientRect(hWnd, &Rect);
+			int32 ButtonHeight = Rect.right + 1;
+
+			PAINTSTRUCT ps;
+			BeginPaint(hWnd, &ps);
+
+			RECT ScrollRect{ 0, ButtonHeight, Rect.right, Rect.bottom - ButtonHeight };
+			FillRect(ps.hdc, &ScrollRect, ScrollBackgroundBrush);
+
+			int ThumbPos;
+			int ThumbHeight;
+			GetThumbPosAndHeight(ThumbPos, ThumbHeight, Rect);
+
+			if (ThumbHeight > 0)
+			{
+				RECT ThumbRect{ 1, ThumbPos, Rect.right-1, ThumbPos + ThumbHeight };
+				FillRect(ps.hdc, &ThumbRect, ThumbBrush);
+			}
+
+			EndPaint(hWnd, &ps);
+			break;
+		}
+		case WM_LBUTTONDOWN:
+		{
+			RECT Rect;
+			GetClientRect(hWnd, &Rect);
+			int32 Y = HIWORD(lParam);
+
+			int32 ButtonHeight = Rect.right + 1;
+			if (Y <= ButtonHeight)
+			{
+				MoveTopVisible(Rect, -1);
+				break;
+			}
+			if (Y > Rect.bottom - ButtonHeight)
+			{
+				MoveTopVisible(Rect, 1);
+				break;
+			}
+
+			int ThumbPos;
+			int ThumbHeight;
+			GetThumbPosAndHeight(ThumbPos, ThumbHeight, Rect);
+			if (ThumbHeight > 0 && Y > ThumbPos && Y < ThumbPos + ThumbHeight)
+			{
+				LogScrollGrabPos = Y - ThumbPos;
+				SetCapture(hWnd);
+				break;
+			}
+			if (Y <= ThumbPos)
+			{
+				int32 PageSize = Rect.bottom / LogFontHeight;
+				MoveTopVisible(Rect, -PageSize + 1);
+				break;
+			}
+			if (Y > ThumbPos + ThumbHeight)
+			{
+				int32 PageSize = Rect.bottom / LogFontHeight;
+				MoveTopVisible(Rect, PageSize - 1);
+				break;
+			}
+
+			break;
+		}
+		case WM_LBUTTONUP:
+			ReleaseCapture();
+			LogScrollGrabPos = -1;
+			break;
+
+		case WM_MOUSEMOVE:
+			if (LogScrollGrabPos != -1)
+			{
+				RECT Rect;
+				GetClientRect(hWnd, &Rect);
+
+				int32 Y = HIWORD(lParam);
+				int32 ThumbPos = Y - LogScrollGrabPos;
+				int32 TopVisible = GetTopVisible(ThumbPos, Rect);
+				SendMessageW(LogHwnd, LB_SETTOPINDEX, TopVisible, 0);
+				RedrawWindow(hWnd, NULL, NULL, RDW_INVALIDATE);
+			}
+			break;
+
+		case WM_MOUSEWHEEL:
+			return HandleMouseWheel(wParam);
+		}
+		return DefWindowProc(hWnd, Msg, wParam, lParam);
+	}
+
+	LRESULT HitTestNCA(HWND hWnd, LPARAM lParam)
+	{
+		POINT MousePt = { ((int)(short)LOWORD(lParam)), ((int)(short)HIWORD(lParam))};
+		RECT WindowRect;
+		GetWindowRect(hWnd, &WindowRect);
+		RECT FrameRect = { 0 };
+		AdjustWindowRectEx(&FrameRect, WS_OVERLAPPEDWINDOW & ~WS_CAPTION, false, NULL);
+		int32 Row = 1;
+		int32 Col = 1;
+		bool bOnResizeBorder = false;
+
+		int BorderWidth = 2;
+
+		if (MousePt.y >= WindowRect.top && MousePt.y < WindowRect.top + 27)
+		{
+			bOnResizeBorder = (MousePt.y < (WindowRect.top - FrameRect.top));
+			Row = 0;
+		}
+		else if (MousePt.y < WindowRect.bottom && MousePt.y >= WindowRect.bottom - BorderWidth)
+			Row = 2;
+
+		if (MousePt.x >= WindowRect.left && MousePt.x < WindowRect.left + BorderWidth)
+			Col = 0;
+		else if (MousePt.x < WindowRect.right && MousePt.x >= WindowRect.right - BorderWidth)
+			Col = 2;
+
+		LRESULT hitTests[3][3] =
+		{
+			{ HTTOPLEFT,    bOnResizeBorder ? HTTOP : HTCAPTION,    HTTOPRIGHT },
+			{ HTLEFT,       HTCLIENT,     HTRIGHT },
+			{ HTBOTTOMLEFT, HTBOTTOM, HTBOTTOMRIGHT },
+		};
+
+		return hitTests[Row][Col];
+	}
+
+	int32 HitTestNCB(int32 X, int32 Y)
+	{
+		RECT WindowRect;
+		GetWindowRect(MainHwnd, &WindowRect);
+		if (Y < WindowRect.top + 2 || Y > WindowRect.top + 27 || X > WindowRect.right - 2)
+			return HTNOWHERE;
+		int32 ButtonWidth = 27;
+		if (X < WindowRect.right - ButtonWidth*4)
+			return HTCAPTION;
+		int32 ButtonIndex = (WindowRect.right - X) / ButtonWidth;
+		int32 Buttons[] = { HTCLOSE, HTMAXBUTTON, HTMINBUTTON, HTEXPAND };
+		return Buttons[ButtonIndex];
+	}
+
+	int32 HitTestNCB(LPARAM lParam)
+	{
+		return HitTestNCB((int)(short)LOWORD(lParam), (int)(short)HIWORD(lParam));
+	}
+
+	void RedrawNC()
+	{
+		RedrawWindow(MainHwnd, NULL, NULL, RDW_INVALIDATE | RDW_FRAME);
+	}
+
+	LRESULT MainWinProc(HWND hWnd, UINT Msg, WPARAM wParam, LPARAM lParam)
+	{
+		switch (Msg)
+		{
+		case WM_NEWLOGENTRIES:
+		{
+			SendMessageW(LogHwnd, WM_SETREDRAW, false, 0);
+			{
+				{
+					FScopeLock _(&NewLogEntriesCs);
+					Swap(NewLogEntries, TempLogEntries);
+				}
+
+				int32 OldSize = Log.Num();
+
+				for (NewLogEntry& E : TempLogEntries)
+				{
+					const TCHAR* SearchStr = *E.String;
+					while (true)
+					{
+						if (const TCHAR* LineBreak = FCString::Strchr(SearchStr, '\n'))
+						{
+							int32 Len = LineBreak - SearchStr;
+							if (LineBreak > SearchStr && *(LineBreak - 1) == '\r')
+								--Len;
+							Log.Add({ FString(Len, SearchStr), E.TextAttribute, 1 });
+							SearchStr = LineBreak + 1;
+						}
+						else
+						{
+							if (int32 Len = FCString::Strlen(SearchStr))
+								Log.Add({ SearchStr, E.TextAttribute, 1 });
+							break;
+						}
+					}
+				}
+
+				constexpr int32 MaxSize = 20000;
+
+				int32 NewSize = Log.Num();
+				if (NewSize > MaxSize)
+				{
+					int32 ToRemove = NewSize - MaxSize;
+					Log.PopFront(ToRemove);
+					LogIndexOffset += ToRemove;
+
+					int32 LogIndex = SendMessageW(LogHwnd, LB_GETITEMDATA, 0, 0);
+					while (ToRemove--)
+					{
+						if (LogIndex > LogIndexOffset)
+							break;
+						SendMessageW(LogHwnd, LB_DELETESTRING, 0, 0);
+						++LogIndex;
+					}
+				}
+
+				for (int32 LogIndex = OldSize, LogNum = Log.Num(); LogIndex != LogNum; ++LogIndex)
+				{
+					LogEntry& E = Log[LogIndex];
+					AddEntryToLogHwnd(Log[LogIndex], LogIndexOffset + LogIndex);
+				}
+				TempLogEntries.Reset(0);
+			}
+
+			if (bAutoScrollLog)
+				ScrollToBottom();
+
+			RedrawLogScrollbar();
+
+			SendMessageW(LogHwnd, WM_SETREDRAW, true, 0);
+			break;
+		}
+
+		case WM_ACTIVITIESDIRTY:
+		{
+			FScopeLock _(&ActivitiesCs);
+			bActivitiesDirty = false;
+			bool bUpdatePositions = false;
+			for (int32 i = 0, e = Activities.Num(); i != e;)
+			{
+				Activity& A = Activities[i];
+
+				if (A.Name.IsEmpty())
+				{
+					DestroyWindow(A.NameHwnd);
+					DestroyWindow(A.StatusHwnd);
+					Activities.RemoveAt(i);
+					bUpdatePositions = true;
+					for (auto& KV : IdToActivityIndex)
+						if (KV.Value > i)
+							--KV.Value;
+					--e;
+					continue;
+				}
+
+				if (!A.NameHwnd)
+				{
+					A.NameHwnd = CreateTextHwnd(*A.Name, Font);
+					A.StatusHwnd = CreateTextHwnd(*A.Status, Font);
+					bUpdatePositions = true;
+				}
+				else if (A.bStatusDirty)
+				{
+					A.bStatusDirty = false;
+					SetWindowTextW(A.StatusHwnd, *A.Status);
+				}
+				++i;
+			}
+
+			if (bActivitiesLightDirty)
+			{
+				InvalidateRect(MainHwnd, NULL, false);
+				bActivitiesLightDirty = false;
+			}
+
+			if (bUpdatePositions)
+			{
+				RECT Rect;
+				GetClientRect(hWnd, &Rect);
+				int NewActivitiesTotalHeight = TraverseActivityPositions(Rect.right, Rect.bottom, [](Activity& A, int32 X, int32 Y, int32 Width, int32 RowIndex) {});
+				if (NewActivitiesTotalHeight == ActivitiesTotalHeight)
+				{
+					UpdateActivityPositions(Rect.right, Rect.bottom, false);
+				}
+				else
+				{
+					ActivitiesTotalHeight = NewActivitiesTotalHeight;
+					UpdateSize(Rect.right, Rect.bottom, false);
+				}
+				RedrawWindow(MainHwnd, NULL, NULL, RDW_ERASE | RDW_INVALIDATE | RDW_ALLCHILDREN | RDW_UPDATENOW); // Important with UPDATENOW to make sure we avoid glitches
+			}
+			break;
+		}
+
+		case WM_SIZE:
+		{
+			if (hWnd != MainHwnd)
+				break;
+			UpdateSize(LOWORD(lParam), HIWORD(lParam), true);
+			RECT WindowRect;
+			GetWindowRect(MainHwnd, &WindowRect);
+			ConsoleWidth = WindowRect.right - WindowRect.left;
+			ConsoleHeight = WindowRect.bottom - WindowRect.top;
+			return 0;
+		}
+		case WM_MOVE:
+		{
+			if (hWnd != MainHwnd)
+				break;
+			RECT WindowRect;
+			GetWindowRect(MainHwnd, &WindowRect);
+			ConsolePosX = WindowRect.left;
+			ConsolePosY = WindowRect.top;
+			break;
+		}
+
+		case WM_SETFOCUS:
+			if (hWnd != MainHwnd)
+				break;
+			SetFocus(CommandHwnd);
+			SendMessageW(CommandHwnd, EM_SETSEL, 0, -1);
+			break;
+
+		case WM_MOUSEWHEEL:
+			return HandleMouseWheel(wParam);
+
+		case WM_ERASEBKGND:
+		{
+			HDC Hdc = (HDC)wParam;
+			RECT Rect;
+			GetClientRect(hWnd, &Rect);
+			FillRect(Hdc, &Rect, BackgroundBrush);
+
+			// Borders around edit boxes
+			SelectObject(Hdc, BorderPen);
+			SelectObject(Hdc, GetStockObject(NULL_BRUSH));
+			for (int32 dlgId : {ID_LOG_INCLUDE_FILTER, ID_LOG_EXCLUDE_FILTER, ID_COMMAND})
+			{
+				HWND editHwnd = GetDlgItem(hWnd, dlgId);
+				if (!IsWindowVisible(editHwnd))
+					continue;
+				RECT EditRect;
+				GetClientRect(editHwnd, &EditRect);
+				MapWindowPoints(editHwnd, hWnd, (LPPOINT)&EditRect, 2);
+				Rectangle(Hdc, EditRect.left - 1, EditRect.top - 1, EditRect.right + 1, EditRect.bottom + 1);
+			}
+			return true;
+		}
+
+		case WM_CTLCOLORSTATIC:
+		{
+			HDC hdcStatic = (HDC)wParam;
+			SetTextColor(hdcStatic, TextColor);
+			SetBkColor(hdcStatic, BackgroundColor);
+			return (INT_PTR)BackgroundBrush;
+		}
+
+		case WM_CTLCOLOREDIT:
+		{
+			HDC hdcStatic = (HDC)wParam;
+			SetTextColor(hdcStatic, TextColor);
+			SetBkColor(hdcStatic, EditBackgroundColor);
+			return (INT_PTR)EditBackgroundBrush;
+		}
+
+		case WM_CTLCOLORLISTBOX:
+			if ((HWND)lParam == LogHwnd)
+			{
+				int32 ItemCount = (int)SendMessageW(LogHwnd, LB_GETCOUNT, 0, 0);
+				RECT Rect;
+				GetClientRect(LogHwnd, &Rect);
+				if (ItemCount * LogFontHeight > Rect.bottom)
+				{
+					// Draw the padding below the last line
+					if (int32 padding = Rect.bottom % LogFontHeight)
+					{
+						Rect.top = Rect.bottom - padding;
+						FillRect((HDC)wParam, &Rect, LogBackgroundBrush);
+					}
+					return (INT_PTR)GetStockObject(NULL_BRUSH);
+				}
+				else
+					return (INT_PTR)LogBackgroundBrush;
+			}
+			break;
+
+		case WM_MEASUREITEM:
+		{
+			PMEASUREITEMSTRUCT pmis = (PMEASUREITEMSTRUCT)lParam;
+			pmis->itemHeight = LogFontHeight;
+			return true;
+		}
+
+		case WM_NOTIFY:
+		{
+			LPNMHDR hdr = (LPNMHDR)lParam;
+			if (hdr->code != NM_CUSTOMDRAW)
+				break;
+			switch (hdr->idFrom)
+			{
+			case ID_COMMANDBUTTON:
+			case ID_CLEARLOGBUTTON:
+			case ID_ADDCHECKPOINTBUTTON:
+			{
+				NMCUSTOMDRAW& nmcd = *(NMCUSTOMDRAW*)hdr;
+				uint32 CtlId = hdr->idFrom;
+				int32 ItemState = nmcd.uItemState;
+				TCHAR Str[32];
+				int32 StrLen = GetDlgItemText(hWnd, CtlId, Str, 32);
+				HBRUSH Brush = ButtonBrush;
+				COLORREF Color = ButtonColor;
+				if (ItemState & CDIS_SELECTED)
+				{
+					Brush = ButtonPressedBrush;
+					Color = ButtonPressedColor;
+				}
+				else if (ItemState & CDIS_HOT)
+				{
+					Brush = ButtonHighlightBrush;
+					Color = ButtonHighlightColor;
+				}
+				SelectObject(nmcd.hdc, GetStockObject(NULL_BRUSH));
+				FillRect(nmcd.hdc, &nmcd.rc, Brush);
+				SetBkColor(nmcd.hdc, Color);
+				SetTextColor(nmcd.hdc, TextColor);
+				DrawTextW(nmcd.hdc, Str, StrLen, &nmcd.rc, DT_SINGLELINE | DT_VCENTER | DT_CENTER | DT_NOCLIP);
+				SelectObject(nmcd.hdc, BorderPen);
+				Rectangle(nmcd.hdc, nmcd.rc.left, nmcd.rc.top, nmcd.rc.right, nmcd.rc.bottom);
+				return CDRF_SKIPDEFAULT;
+			}
+			}
+			break;
+		}
+
+		case WM_DRAWITEM:
+		{
+			PDRAWITEMSTRUCT pdis = (PDRAWITEMSTRUCT)lParam;
+
+			if (pdis->CtlType == ODT_STATIC)
+			{
+				TCHAR Str[512];
+				int32 StrLen = GetWindowText(pdis->hwndItem, Str, UE_ARRAY_COUNT(Str));
+
+				HBRUSH Brush = BackgroundBrush;
+
+				int32 RowIndex = GetWindowLong(pdis->hwndItem, GWLP_USERDATA);
+				if (RowIndex != -1)
+				{
+					Brush = StatusBackgroundBrush[(RowIndex & 1)];
+					SetBkColor(pdis->hDC, StatusBackgroundColor[(RowIndex & 1)]);
+				}
+
+				SelectObject(pdis->hDC, Brush);
+
+				// All this stuff just to prevent flickering when text changes fast
+				int32 TextOffset = 1;
+				RECT ModifiedRect = pdis->rcItem;
+				DrawTextW(pdis->hDC, Str, StrLen, &ModifiedRect, DT_SINGLELINE | DT_END_ELLIPSIS | DT_CALCRECT | DT_MODIFYSTRING);
+				RECT TempRect = pdis->rcItem;
+				TempRect.bottom = TempRect.top + TextOffset;
+				TempRect.right = ModifiedRect.right;
+				FillRect(pdis->hDC, &TempRect, Brush);
+
+				TempRect = pdis->rcItem;
+				TempRect.left = ModifiedRect.right;
+				FillRect(pdis->hDC, &TempRect, Brush);
+				if (ModifiedRect.bottom < pdis->rcItem.bottom)
+				{
+					TempRect = pdis->rcItem;
+					TempRect.right = ModifiedRect.right;
+					TempRect.top = ModifiedRect.bottom;
+					FillRect(pdis->hDC, &TempRect, Brush);
+				}
+
+				StrLen = FCString::Strlen(Str);
+				ExtTextOut(pdis->hDC, pdis->rcItem.left, pdis->rcItem.top + TextOffset, 0, &pdis->rcItem, Str, StrLen, 0);
+				break;
+			}
+
+			if (pdis->CtlID != ID_LOG || pdis->itemID == -1)
+				break;
+			switch (pdis->itemAction)
+			{
+			case ODA_SELECT:
+			case ODA_DRAWENTIRE:
+			{
+				uint32 LogVirtualIndex = pdis->itemData;
+				LogEntry& Entry = Log[LogVirtualIndex - LogIndexOffset];
+
+
+				int Middle = (pdis->rcItem.top + pdis->rcItem.bottom) / 2;
+				int YPos = Middle - LogFontHeight / 2 - (Entry.LineCount - 1) * LogFontHeight / 2;
+
+				if (pdis->itemState & ODS_SELECTED)
+				{
+					FillRect(pdis->hDC, &pdis->rcItem, GetSysColorBrush(COLOR_HIGHLIGHT));
+					SetTextColor(pdis->hDC, GetSysColor(COLOR_HIGHLIGHTTEXT));
+				}
+				else
+				{
+					//if (pdis->itemAction != ODA_DRAWENTIRE)
+					{
+						FillRect(pdis->hDC, &pdis->rcItem, LogBackgroundBrush);
+					}
+
+					if (Entry.TextAttribute == (FOREGROUND_INTENSITY | FOREGROUND_RED))
+						SetTextColor(pdis->hDC, RGB(220, 0, 0));
+					else if (Entry.TextAttribute == (FOREGROUND_INTENSITY | FOREGROUND_RED | FOREGROUND_GREEN))
+						SetTextColor(pdis->hDC, RGB(220, 220, 0));
+					else
+						SetTextColor(pdis->hDC, RGB(204, 204, 204));
+				}
+
+				int32 XPos = 0;
+				const TCHAR* Str = *Entry.String;
+				while (true)
+				{
+					if (const TCHAR* Tab = FCString::Strchr(Str, '\t'))
+					{
+						if (int32 PartLen = int32(Tab - Str) - 1)
+						{
+							TextOut(pdis->hDC, XPos, YPos, Str, PartLen);
+							XPos += PartLen * LogFontWidth;
+						}
+						XPos += LogFontWidth * 3;
+						Str = Tab + 1;
+						continue;
+					}
+
+					TextOut(pdis->hDC, XPos, YPos, Str, FCString::Strlen(Str));
+					break;
+
+				}
+				return true;
+			}
+			case ODA_FOCUS: // We don't want the focus rectangle at all (dotted rectangle)
+				return true;
+			}
+			break;
+		}
+
+		case WM_PAINT:
+		{
+			if (hWnd != MainHwnd)
+				break;
+			PAINTSTRUCT ps;
+			HDC Hdc = BeginPaint(hWnd, &ps);
+			RECT Rect;
+			GetClientRect(hWnd, &Rect);
+
+			TraverseActivityPositions(Rect.right, Rect.bottom, [&](Activity& A, int32 X, int32 Y, int32 Width, int32 RowIndex)
+				{
+					if (A.Light == 0)
+						return;
+					SelectObject(Hdc, StatusLightBrush[A.Light - 1]);
+					int32 XOffset = 12;
+					int32 YOffset = 1;
+					int32 FontHeight = 16;
+					int32 CircleSize = 11;
+					int32 HalfDiff = (FontHeight - CircleSize) / 2;
+					RoundRect(Hdc, X - CircleSize + XOffset, Y + HalfDiff + YOffset, X + XOffset, Y + FontHeight - HalfDiff + YOffset, CircleSize, CircleSize);
+				});
+			EndPaint(hWnd, &ps);
+			return 0;
+		}
+
+		case WM_CLOSE:
+			PostQuitMessage(0);
+			break;
+
+		case WM_VKEYTOITEM:
+			if ((HWND)lParam == LogHwnd)
+			{
+				if ((LOWORD(wParam) == 0x43 || LOWORD(wParam) == VK_INSERT) && (::GetKeyState(VK_CONTROL) >> 15)) // Copy-paste
+				{
+					if (int32 SelectionCount = SendMessageW(LogHwnd, LB_GETSELCOUNT, 0, 0))
+					{
+						TArray<int> SelectedItems;
+						SelectedItems.SetNum(SelectionCount);
+						SelectionCount = SendMessageW(LogHwnd, LB_GETSELITEMS, SelectionCount, (LPARAM)SelectedItems.GetData());
+						TArray<TCHAR> Buffer;
+
+						TStringBuilder<512> StringBuilder;
+						for (int32 I = 0; I != SelectionCount; ++I)
+						{
+							if (I > 0)
+								StringBuilder.Append(TEXT("\r\n"));
+							int32 Index = SelectedItems[I];
+							int32 Len = SendMessageW(LogHwnd, LB_GETTEXTLEN, Index, 0);
+							Buffer.SetNum(Len + 1, false);
+							SendMessageW(LogHwnd, LB_GETTEXT, Index, (LPARAM)Buffer.GetData());
+							StringBuilder.Append(Buffer.GetData(), Buffer.Num() - 1);
+						}
+
+						FWindowsPlatformApplicationMisc::ClipboardCopy(*StringBuilder);
+					}
+				}
+				else if (LOWORD(wParam) == VK_END) // Enable auto scrolling and remove selection
+				{
+					bAutoScrollLog = true;
+					enum { MaxSelectionCount = 32 * 1024 - 1 };
+					int32 SelectedItems[MaxSelectionCount + 1];
+					int32 SelectionCount = SendMessageW(LogHwnd, LB_GETSELITEMS, MaxSelectionCount, (LPARAM)SelectedItems);
+					for (int32 I = 0; I != SelectionCount; ++I)
+						SendMessageW(LogHwnd, LB_SETSEL, false, SelectedItems[I]);
+					ScrollToBottom();
+					RedrawWindow(LogHwnd, NULL, NULL, RDW_ERASE|RDW_INVALIDATE);
+					RedrawLogScrollbar();
+					return -2;
+				}
+			}
+
+			break;
+		case WM_COMMAND:
+			if (LOWORD(wParam) == ID_LOG && HIWORD(wParam) == LBN_SELCHANGE)
+			{
+				if (SendMessageW(LogHwnd, LB_GETSELCOUNT, 0, 0) != 0)
+					bAutoScrollLog = false;
+			}
+			else if ((LOWORD(wParam) == ID_LOG_INCLUDE_FILTER || LOWORD(wParam) == ID_LOG_EXCLUDE_FILTER) && HIWORD(wParam) == EN_CHANGE)
+			{
+				TCHAR Str[1024];
+				GetDlgItemText(MainHwnd, LOWORD(wParam), Str, 1024);
+				TArray<FString>& Filter = LOWORD(wParam) == ID_LOG_INCLUDE_FILTER ? IncludeFilter : ExcludeFilter;
+				TArray<FString> ParsedFilter;
+				FString(Str).ParseIntoArray(ParsedFilter, TEXT(" "));
+				if (Filter == ParsedFilter)
+					break;
+				Filter = ParsedFilter;
+				RefreshLogHwnd();
+			}
+			else if (LOWORD(wParam) == ID_COMMANDBUTTON)
+			{
+				TCHAR Command[1024];
+				if (!GetWindowText(CommandHwnd, Command, 1024))
+					break;
+				SetWindowText(CommandHwnd, TEXT(""));
+				EConsoleColor NewConsoleColor = EConsoleColor(-1);
+				if (FCString::Stricmp(Command, TEXT("red")) == 0)
+					NewConsoleColor = EConsoleColor::Red;
+				else if (FCString::Stricmp(Command, TEXT("darkgray")) == 0)
+					NewConsoleColor = EConsoleColor::DarkGray;
+				else if (FCString::Stricmp(Command, TEXT("gray")) == 0)
+					NewConsoleColor = EConsoleColor::Gray;
+				else if (FCString::Stricmp(Command, TEXT("lightblue")) == 0)
+					NewConsoleColor = EConsoleColor::LightBlue;
+				else if (FCString::Stricmp(Command, TEXT("darkblue")) == 0)
+					NewConsoleColor = EConsoleColor::DarkBlue;
+				else if (FCString::Stricmp(Command, TEXT("load")) == 0)
+					NewConsoleColor = EConsoleColor::Load;
+
+				if (NewConsoleColor != EConsoleColor(-1))
+				{
+					ConsoleColor = NewConsoleColor;
+					DestroyColors();
+					CreateColors();
+					InvalidateRect(MainHwnd, NULL, true);
+					EnumChildWindows(MainHwnd, [](HWND hwnd, LPARAM lParam) -> BOOL { InvalidateRect(hwnd, NULL, true); return true;}, 0);
+					RedrawNC();
+					break;
+				}
+
+				Async(EAsyncExecution::TaskGraphMainThread, [Cmd = FString(Command)]
+					{
+						if (IModularFeatures::Get().IsModularFeatureAvailable(IConsoleCommandExecutor::ModularFeatureName()))
+						{
+							UE_LOG(LogExec, Log, TEXT("Executing console command: %s"), *Cmd);
+							IModularFeatures::Get().GetModularFeature<IConsoleCommandExecutor>(FName(IConsoleCommandExecutor::ModularFeatureName())).Exec(*Cmd);
+						}
+						else
+						{
+							UE_LOG(LogExec, Log, TEXT("Failed to execute console command: %s"), *Cmd);
+						}
+					});
+			}
+			else if (LOWORD(wParam) == ID_CLEARLOGBUTTON)
+			{
+				LogIndexOffset = 0;
+				Log.Reset();
+				RefreshLogHwnd();
+				CheckpointIndex = 0;
+				SetDlgItemText(MainHwnd, ID_ADDCHECKPOINTBUTTON, TEXT("Log CHECKPOINT0"));
+			}
+			else if (LOWORD(wParam) == ID_ADDCHECKPOINTBUTTON)
+			{
+				TStringBuilder<MAX_SPRINTF> TempString;
+				TempString.Appendf(TEXT("LOGCHECKPOINT%i"), CheckpointIndex);
+				AddLogEntry(TempString, FOREGROUND_BLUE | FOREGROUND_GREEN | FOREGROUND_RED);
+
+				++CheckpointIndex;
+				TCHAR ButtonString[MAX_SPRINTF] = TEXT("");
+				FCString::Sprintf(ButtonString, TEXT("Log CHECKPOINT%i"), CheckpointIndex);
+				SetDlgItemText(MainHwnd, ID_ADDCHECKPOINTBUTTON, ButtonString);
+			}
+			break;
+
+		case DM_GETDEFID:
+			return MAKEWPARAM(ID_COMMANDBUTTON, DC_HASDEFID);
+
+		case WM_NCHITTEST: // This actually draw things
+		{
+			int32 Res = HitTestNCB(lParam);
+			if (Res != HTNOWHERE)
+				return Res;
+			return HitTestNCA(hWnd, lParam);
+		}
+
+		case WM_MOUSEMOVE:
+		case WM_NCMOUSELEAVE:
+		{
+			if (NcButtonHot == -1 && NcButtonDown == -1)
+				break;
+			POINT p;
+			GetCursorPos(&p);
+			int32 Res = HitTestNCB(p.x, p.y);
+			if (Res != HTMINBUTTON && Res != HTMAXBUTTON && Res != HTCLOSE && Res != HTEXPAND)
+				Res = -1;
+			if (Res == NcButtonHot)
+				break;
+			NcButtonHot = Res;
+			RedrawNC();
+			break;
+		}
+
+		case WM_NCMOUSEMOVE:
+		{
+			int32 Res = HitTestNCB(lParam);
+			if (Res != HTMINBUTTON && Res != HTMAXBUTTON && Res != HTCLOSE && Res != HTEXPAND)
+				Res = -1;
+			if (Res == NcButtonHot)
+				break;
+			NcButtonHot = Res;
+			RedrawNC();
+			if (Res == -1)
+				break;
+			TRACKMOUSEEVENT tme;
+			tme.cbSize = sizeof(tme);
+			tme.dwFlags = TME_NONCLIENT | TME_LEAVE;
+			tme.hwndTrack = hWnd;
+			TrackMouseEvent(&tme);
+			break;
+		}
+
+		case WM_NCLBUTTONDOWN:
+		{
+			int32 Res = HitTestNCB(lParam);
+			if (Res == HTMINBUTTON || Res == HTMAXBUTTON || Res == HTCLOSE || Res == HTEXPAND)
+			{
+				SetCapture(MainHwnd);
+				NcButtonDown = Res;
+				RedrawNC();
+				return 0;
+			}
+			break;
+		}
+
+		case WM_LBUTTONUP:
+		{
+			if (NcButtonDown == -1)
+				break;
+			ReleaseCapture();
+			int32 Old = NcButtonDown;
+			NcButtonDown = -1;
+			POINT p;
+			GetCursorPos(&p);
+			int32 Res = HitTestNCB(p.x, p.y);
+			RedrawNC();
+			if (Res != Old)
+				break;
+			if (Res == HTEXPAND)
+			{
+				bConsoleExpanded = !bConsoleExpanded;
+				RedrawNC();
+				RECT Rect;
+				GetClientRect(MainHwnd, &Rect);
+				UpdateSize(Rect.right, Rect.bottom, true);
+				RedrawWindow(MainHwnd, NULL, NULL, RDW_INVALIDATE);
+			}
+			else if (Res == HTMINBUTTON)
+				ShowWindow(hWnd, SW_MINIMIZE);
+			else if (Res == HTMAXBUTTON)
+			{
+				WINDOWPLACEMENT Placement;
+				Placement.length = sizeof(Placement);
+				GetWindowPlacement(hWnd, &Placement);
+				ShowWindow(hWnd, Placement.showCmd == SW_MAXIMIZE ? SW_NORMAL : SW_MAXIMIZE);
+			}
+			else if (Res == HTCLOSE)
+				PostQuitMessage(0);
+			break;
+		}
+
+		case WM_NCCALCSIZE:
+		{
+			LPNCCALCSIZE_PARAMS ncParams = (LPNCCALCSIZE_PARAMS)lParam;
+			ncParams->rgrc[0].top += 27;
+			ncParams->rgrc[0].left += 1;
+			ncParams->rgrc[0].bottom -= 1;
+			ncParams->rgrc[0].right -= 1;
+			return 0;
+		}
+		case WM_NCPAINT:
+		{
+			RECT rect;
+			GetWindowRect(hWnd, &rect);
+			HRGN region = NULL;
+			if (wParam == NULLREGION)
+				region = CreateRectRgn(rect.left, rect.top, rect.right, rect.bottom);
+			else
+			{
+				HRGN copy = CreateRectRgn(0, 0, 0, 0);
+				if (CombineRgn(copy, (HRGN)wParam, NULL, RGN_COPY))
+					region = copy;
+				else
+					DeleteObject(copy);
+			}
+
+			HDC dc = GetDCEx(hWnd, region, DCX_WINDOW | DCX_CACHE | DCX_INTERSECTRGN | DCX_LOCKWINDOWUPDATE);
+			if (!dc && region)
+				DeleteObject(region);
+
+			int32 width = rect.right - rect.left;
+			int32 height = rect.bottom - rect.top;
+
+			int32 ButtonWidth = 27;
+
+			RECT CaptionRect{ 0,0,width, 27 };
+			FillRect(dc, &CaptionRect, BackgroundBrush);
+
+			SelectObject(dc, GetStockObject(NULL_BRUSH));
+			SelectObject(dc, WindowBorderPen);
+			Rectangle(dc, 0, 0, width, height);
+
+			DrawIconEx(dc, 3, 3, Icon, 21, 21, 0, NULL, DI_NORMAL);
+
+			SetBkColor(dc, BackgroundColor);
+			SetTextColor(dc, TextColor);
+			SelectObject(dc, Font);
+			CaptionRect.left += 32;
+			CaptionRect.right -= ButtonWidth * 4;
+			DrawTextW(dc, *ConsoleTitle, ConsoleTitle.Len(), &CaptionRect, DT_SINGLELINE | DT_VCENTER | DT_END_ELLIPSIS);
+
+			SelectObject(dc, TextPen);
+
+			if (NcButtonHot != -1 && (NcButtonDown == -1 || NcButtonHot == NcButtonDown))
+			{
+				int32 NcButton = NcButtonHot;
+				HBRUSH Brush = ButtonBrush;
+				if (NcButtonDown != -1)
+					Brush = ButtonPressedBrush;
+				int32 ButtonIndex = 0;
+				if (NcButton == HTMAXBUTTON)
+					ButtonIndex = 1;
+				else if (NcButton == HTMINBUTTON)
+					ButtonIndex = 2;
+				else if (NcButton == HTEXPAND)
+					ButtonIndex = 3;
+				RECT ButtonRect{ width - ButtonWidth*(ButtonIndex+1), 1, width - 1 - ButtonWidth*ButtonIndex, 27};
+				FillRect(dc, &ButtonRect, Brush);
+			}
+
+			int32 MiddleY = 27 / 2;
+			int32 MiddleX = width - ButtonWidth / 2 - ButtonWidth * 3 - 1;
+
+			// Double arrow
+			if (bConsoleExpanded)
+			{
+				MoveToEx(dc, MiddleX - 4, MiddleY - 4, NULL);
+				LineTo(dc, MiddleX, MiddleY);
+				LineTo(dc, MiddleX + 5, MiddleY - 5);
+				MoveToEx(dc, MiddleX - 4, MiddleY, NULL);
+				LineTo(dc, MiddleX, MiddleY + 4);
+				LineTo(dc, MiddleX + 5, MiddleY - 1);
+			}
+			else
+			{
+				MoveToEx(dc, MiddleX - 4, MiddleY, NULL);
+				LineTo(dc, MiddleX, MiddleY - 4);
+				LineTo(dc, MiddleX + 5, MiddleY + 1);
+				MoveToEx(dc, MiddleX - 4, MiddleY + 4, NULL);
+				LineTo(dc, MiddleX, MiddleY);
+				LineTo(dc, MiddleX + 5, MiddleY + 5);
+			}
+
+			// Minimize
+			MiddleX += ButtonWidth;
+			MoveToEx(dc, MiddleX - 4, MiddleY, NULL);
+			LineTo(dc, MiddleX + 5, MiddleY);
+			
+			// Maximize
+			MiddleX += ButtonWidth;
+			WINDOWPLACEMENT Placement;
+			Placement.length = sizeof(Placement);
+			GetWindowPlacement(hWnd, &Placement);
+			if (Placement.showCmd == SW_MAXIMIZE)
+			{
+				Rectangle(dc, MiddleX, MiddleY - 3, MiddleX + 8, MiddleY + 5);
+				SelectObject(dc, BackgroundBrush);
+				Rectangle(dc, MiddleX - 3, MiddleY, MiddleX + 5, MiddleY + 8);
+			}
+			else
+			{
+				Rectangle(dc, MiddleX - 5, MiddleY - 5, MiddleX + 5, MiddleY + 5);
+			}
+			
+			// Close
+			MiddleX += ButtonWidth;
+			MoveToEx(dc, MiddleX - 4, MiddleY - 4, NULL);
+			LineTo(dc, MiddleX + 5, MiddleY + 5);
+			MoveToEx(dc, MiddleX + 4, MiddleY - 4, NULL);
+			LineTo(dc, MiddleX - 5, MiddleY + 5);
+
+			ReleaseDC(hWnd, dc);
+			return 0;
+		}
+		case WM_NCACTIVATE:
+			RedrawWindow(hWnd, NULL, NULL, RDW_UPDATENOW);
+			return 0;
+		}
+		return DefWindowProc(hWnd, Msg, wParam, lParam);
+	}
+
+	enum class EConsoleColor : int32
+	{
+		DarkGray,
+		Gray,
+		Red,
+		LightBlue,
+		DarkBlue,
+		Load,
+	};
+
+	FWindowsConsoleOutputDevice2* Owner;
+
+	struct NewLogEntry { FString String; uint16 TextAttribute; };
+	struct LogEntry { FString String; uint16 TextAttribute; uint16 LineCount; };
+	FCriticalSection NewLogEntriesCs;
+	TArray<NewLogEntry> NewLogEntries;
+	TArray<NewLogEntry> TempLogEntries;
+	TMap<FString, int> IdToActivityIndex;
+	TArray<FString> IncludeFilter;
+	TArray<FString> ExcludeFilter;
+	TRingBuffer<LogEntry> Log;
+	int32 AddedEntryLogVirtualIndex = -1;
+	HICON Icon;
+	HFONT Font;
+	HWND MainHwnd;
+	HWND LogHwnd;
+	HWND LogScrollHwnd;
+	HWND CommandHwnd;
+	HWND RunCommandHwnd;
+	HWND FilterHwnd[7];
+	WNDPROC LogHwndWndProc;
+	COLORREF TextColor;
+	COLORREF BackgroundColor;
+	COLORREF EditBackgroundColor;
+	COLORREF StatusBackgroundColor[2];
+	COLORREF ButtonColor;
+	COLORREF ButtonHighlightColor;
+	COLORREF ButtonPressedColor;
+	COLORREF ThumbColor;
+	COLORREF ScrollBackgroundColor;
+	HBRUSH TextBrush;
+	HBRUSH BackgroundBrush;
+	HBRUSH LogBackgroundBrush;
+	HBRUSH EditBackgroundBrush;
+	HBRUSH StatusBackgroundBrush[2];
+	HBRUSH ButtonBrush;
+	HBRUSH ButtonHighlightBrush;
+	HBRUSH ButtonPressedBrush;
+	HBRUSH ThumbBrush;
+	HBRUSH ScrollBackgroundBrush;
+	HBRUSH StatusLightBrush[3];
+	HPEN BorderPen;
+	HPEN WindowBorderPen;
+	HPEN TextPen;
+	HPEN NoPen;
+	int32 LogIndexOffset = 0;
+	int32 LogFontWidth = 8;
+	int32 LogFontHeight = 14;
+	int32 LogScrollGrabPos = -1;
+	int32 NcButtonHot = -1;
+	int32 NcButtonDown = -1;
+
+	struct Activity { HWND NameHwnd = 0; HWND StatusHwnd = 0; FString Name; FString Status; int32 Light = 0; bool bStatusDirty = false; bool bAlignLeft = false; };
+	FCriticalSection ActivitiesCs;
+	TArray<Activity> Activities;
+	bool bActivitiesDirty = false;
+	bool bActivitiesLightDirty = false;
+	int32 ActivitiesTotalHeight = 0;
+	
+	HANDLE Thread;
+	bool bAutoScrollLog = true;
+	int32 CheckpointIndex = 0;
+
+	bool bIsVisible = false;
+	int32 ConsoleWidth = 160;
+	int32 ConsoleHeight = 4000;
+	int32 ConsolePosX = 0;
+	int32 ConsolePosY = 0;
+	EConsoleColor ConsoleColor = EConsoleColor::DarkGray;
+	FString ConsoleTitle;
+	bool bConsoleExpanded = false;
+};
+
+FName AppStatusTrackerName(TEXT("AppStatusTracker"));
+
+FWindowsConsoleOutputDevice2::FWindowsConsoleOutputDevice2()
+	: OverrideColorSet(false)
+	, Window(nullptr)
+{
+#if !UE_BUILD_SHIPPING
+	auto ParseColorStr = [](const FString& ColorStr) -> const TCHAR*
+	{
+		const TCHAR* ReturnVal = nullptr;
+
+		if (ColorStr == TEXTVIEW("Black"))
+		{
+			ReturnVal = COLOR_BLACK;
+		}
+		else if (ColorStr == TEXTVIEW("DarkRed"))
+		{
+			ReturnVal = COLOR_DARK_RED;
+		}
+		else if (ColorStr == TEXTVIEW("DarkGreen"))
+		{
+			ReturnVal = COLOR_DARK_GREEN;
+		}
+		else if (ColorStr == TEXTVIEW("DarkBlue"))
+		{
+			ReturnVal = COLOR_DARK_BLUE;
+		}
+		else if (ColorStr == TEXTVIEW("DarkYellow"))
+		{
+			ReturnVal = COLOR_DARK_YELLOW;
+		}
+		else if (ColorStr == TEXTVIEW("DarkCyan"))
+		{
+			ReturnVal = COLOR_DARK_CYAN;
+		}
+		else if (ColorStr == TEXTVIEW("DarkPurple"))
+		{
+			ReturnVal = COLOR_DARK_PURPLE;
+		}
+		else if (ColorStr == TEXTVIEW("Gray"))
+		{
+			ReturnVal = COLOR_DARK_WHITE;
+		}
+		else if (ColorStr == TEXTVIEW("Red"))
+		{
+			ReturnVal = COLOR_RED;
+		}
+		else if (ColorStr == TEXTVIEW("Green"))
+		{
+			ReturnVal = COLOR_GREEN;
+		}
+		else if (ColorStr == TEXTVIEW("Blue"))
+		{
+			ReturnVal = COLOR_BLUE;
+		}
+		else if (ColorStr == TEXTVIEW("Yellow"))
+		{
+			ReturnVal = COLOR_YELLOW;
+		}
+		else if (ColorStr == TEXTVIEW("Cyan"))
+		{
+			ReturnVal = COLOR_CYAN;
+		}
+		else if (ColorStr == TEXTVIEW("Purple"))
+		{
+			ReturnVal = COLOR_PURPLE;
+		}
+		else if (ColorStr == TEXTVIEW("White"))
+		{
+			ReturnVal = COLOR_WHITE;
+		}
+
+		return ReturnVal;
+	};
+
+	FString HighlightsStr;
+
+	// -LogHighlights="LogNet Cyan, LogTemp Green"
+	if (FParse::Value(FCommandLine::Get(), TEXT("LogHighlights="), HighlightsStr))
+	{
+		TArray<FString> HighlightsList;
+
+		HighlightsStr.ParseIntoArray(HighlightsList, TEXT(","));
+
+		for (const FString& CurHighlightEntry : HighlightsList)
+		{
+			TArray<FString> CategoryAndColor;
+
+			if (CurHighlightEntry.TrimStartAndEnd().ParseIntoArray(CategoryAndColor, TEXT(" ")) && CategoryAndColor.Num() == 2)
+			{
+				const TCHAR* ColorStr = ParseColorStr(CategoryAndColor[1].TrimStartAndEnd());
+
+				if (ColorStr != nullptr)
+				{
+					FLogHighlight& NewEntry = LogHighlights.AddDefaulted_GetRef();
+
+					NewEntry.Category = FName(CategoryAndColor[0]);
+					NewEntry.Color = ColorStr;
+				}
+			}
+		}
+	}
+
+	FString StringHighlights;
+
+	// -LogStringHighlights="UNetConnection::Close=Purple, NotifyAcceptingConnection accepted from=DarkGreen"
+	if (FParse::Value(FCommandLine::Get(), TEXT("LogStringHighlights="), StringHighlights))
+	{
+		TArray<FString> StringHighlightsList;
+
+		StringHighlights.ParseIntoArray(StringHighlightsList, TEXT(","));
+
+		for (const FString& CurStringHighlightEntry : StringHighlightsList)
+		{
+			TArray<FString> StringAndColor;
+
+			if (CurStringHighlightEntry.ParseIntoArray(StringAndColor, TEXT("=")) && StringAndColor.Num() == 2)
+			{
+				const TCHAR* ColorStr = ParseColorStr(StringAndColor[1].TrimStartAndEnd());
+
+				if (ColorStr != nullptr)
+				{
+					FLogStringHighlight& NewEntry = LogStringHighlights.AddDefaulted_GetRef();
+
+					NewEntry.SearchString = StringAndColor[0].GetCharArray();
+					NewEntry.Color = ColorStr;
+				}
+			}
+		}
+	}
+#endif
+}
+
+FWindowsConsoleOutputDevice2::~FWindowsConsoleOutputDevice2()
+{
+	delete Window;
+}
+
+void FWindowsConsoleOutputDevice2::SaveToINI()
+{
+	FWriteScopeLock lock(WindowRWLock);
+
+	if (!Window)
+		return;
+	
+	FString Filename = TEXT("DebugConsole.ini");
+	const TCHAR* Selection = TEXT("ConsoleWindows");
+	if (IsRunningDedicatedServer())
+		Selection = TEXT("ServerConsoleWindows");
+	else if (IsRunningGame())
+		Selection = TEXT("GameConsoleWindows");
+
+	FConfigCacheIni Config(EConfigCacheType::DiskBacked);
+	Config.LoadFile(Filename);
+
+	Config.SetInt(Selection, TEXT("ConsoleX"), Window->ConsolePosX, Filename);
+	Config.SetInt(Selection, TEXT("ConsoleY"), Window->ConsolePosY, Filename);
+
+	Config.SetInt(Selection, TEXT("ConsoleWidth"), Window->ConsoleWidth, Filename);
+	Config.SetInt(Selection, TEXT("ConsoleHeight"), Window->ConsoleHeight, Filename);
+
+	Config.SetInt(Selection, TEXT("ConsoleColor"), int32(Window->ConsoleColor), Filename);
+
+	Config.SetBool(Selection, TEXT("ConsoleExpanded"), Window->bConsoleExpanded, Filename);
+
+	Config.Flush(false, Filename);
+}
+
+void FWindowsConsoleOutputDevice2::Show( bool ShowWindow )
+{
+	if (ShowWindow)
+	{
+		check(IsInGameThread());
+
+		int32 ConsoleWidth = 1000;
+		int32 ConsoleHeight = 700;
+		int32 ConsolePosX = 0;
+		int32 ConsolePosY = 0;
+		int32 ConsoleColor = 0;
+		bool bConsoleExpanded = false;
+		bool bHasX = false;
+		bool bHasY = false;
+
+		FString Filename = TEXT("DebugConsole.ini");
+		const TCHAR* Selection = TEXT("ConsoleWindows");
+		if (IsRunningDedicatedServer())
+			Selection = TEXT("ServerConsoleWindows");
+		else if (IsRunningGame())
+			Selection = TEXT("GameConsoleWindows");
+
+		FConfigCacheIni Config(EConfigCacheType::Temporary);
+		Config.LoadFile(Filename);
+		Config.GetInt(Selection, TEXT("ConsoleWidth"), ConsoleWidth, Filename);
+		Config.GetInt(Selection, TEXT("ConsoleHeight"), ConsoleHeight, Filename);
+		bHasX = Config.GetInt(Selection, TEXT("ConsoleX"), ConsolePosX, Filename);
+		bHasY = Config.GetInt(Selection, TEXT("ConsoleY"), ConsolePosY, Filename);
+		Config.GetInt(Selection, TEXT("ConsoleColor"), ConsoleColor, Filename);
+		Config.GetBool(Selection, TEXT("ConsoleExpanded"), bConsoleExpanded, Filename);
+
+		if (!FParse::Value(FCommandLine::Get(), TEXT("ConsoleX="), ConsolePosX) && !bHasX)
+		{
+			ConsolePosX = CW_USEDEFAULT;
+		}
+
+		if (!FParse::Value(FCommandLine::Get(), TEXT("ConsoleY="), ConsolePosY) && !bHasY)
+		{
+			ConsolePosY = CW_USEDEFAULT;
+		}
+
+		FString ConsoleTitle;
+
+		if (!FParse::Value(FCommandLine::Get(), TEXT("ConsoleTitle="), ConsoleTitle))
+		{
+			// Setting text so they are easily identifyable in taskbar
+			const TCHAR* ConsoleType = TEXT("Unreal");
+			if (IsRunningDedicatedServer())
+				ConsoleType = ConsoleType = TEXT("Server");
+			else if (IsRunningGame())
+				ConsoleType = ConsoleType = TEXT("Client");
+			ConsoleTitle = FString::Printf(TEXT("%s Console (%s) - %s"), ConsoleType, FApp::GetProjectName(), FPlatformProcess::ExecutablePath());
+		}
+
+		FDisplayMetrics DisplayMetrics;
+		FDisplayMetrics::RebuildDisplayMetrics(DisplayMetrics);
+
+		// Make sure that the positions specified by INI/CMDLINE are proper
+		static const int32 ActualConsoleWidth = ConsoleWidth;
+		static const int32 ActualConsoleHeight = ConsoleHeight;
+
+		static const int32 ActualScreenWidth = DisplayMetrics.VirtualDisplayRect.Right - DisplayMetrics.VirtualDisplayRect.Left;
+		static const int32 ActualScreenHeight = DisplayMetrics.VirtualDisplayRect.Bottom - DisplayMetrics.VirtualDisplayRect.Top;
+
+		static const int32 RightPadding = FMath::Max(50, FMath::Min((ActualConsoleWidth / 2), ActualScreenWidth / 2));
+		static const int32 BottomPadding = FMath::Max(50, FMath::Min((ActualConsoleHeight / 2), ActualScreenHeight / 2));
+				
+		ConsolePosX = FMath::Min(FMath::Max(ConsolePosX, DisplayMetrics.VirtualDisplayRect.Left), DisplayMetrics.VirtualDisplayRect.Right - RightPadding);
+		ConsolePosY = FMath::Min(FMath::Max(ConsolePosY, DisplayMetrics.VirtualDisplayRect.Top), DisplayMetrics.VirtualDisplayRect.Bottom - BottomPadding);
+
+		FWriteScopeLock lock(WindowRWLock);
+
+		bool bFirstCall = !Window;
+		if (bFirstCall)
+		{
+			Window = new FConsoleWindow(this);
+		}
+
+		Window->ConsolePosX = ConsolePosX;
+		Window->ConsolePosY = ConsolePosY;
+		Window->ConsoleWidth = ConsoleWidth;
+		Window->ConsoleHeight = ConsoleHeight;
+		Window->ConsoleColor = (FConsoleWindow::EConsoleColor)ConsoleColor;
+		Window->ConsoleTitle = ConsoleTitle;
+		Window->bConsoleExpanded = bConsoleExpanded;
+		Window->bIsVisible = true;
+
+		if (bFirstCall)
+		{
+			Window->Start();
+		}
+
+	}
+	else
+	{
+		SaveToINI();
+		Window->bIsVisible = false;
+	}
+}
+
+bool FWindowsConsoleOutputDevice2::IsShown()
+{
+	return Window && Window->bIsVisible;
+}
+
+void FWindowsConsoleOutputDevice2::Serialize( const TCHAR* Data, ELogVerbosity::Type Verbosity, const class FName& Category, const double Time )
+{
+	using namespace UE::String;
+
+	const double RealTime = Time == -1.0f ? FPlatformTime::Seconds() - GStartTime : Time;
+
+	if (Verbosity == ELogVerbosity::SetColor)
+	{
+		SetColor( Data );
+		OverrideColorSet = FCString::Strcmp(COLOR_NONE, Data) != 0;
+	}
+	else
+	{
+		bool bNeedToResetColor = false;
+		if( !OverrideColorSet )
+		{
+			if( Verbosity == ELogVerbosity::Error )
+			{
+				SetColor( COLOR_RED );
+				bNeedToResetColor = true;
+			}
+			else if( Verbosity == ELogVerbosity::Warning )
+			{
+				SetColor( COLOR_YELLOW );
+				bNeedToResetColor = true;
+			}
+#if !UE_BUILD_SHIPPING
+			else
+			{
+				if (LogHighlights.Num() > 0)
+				{
+					if (const FLogHighlight* CurHighlight = LogHighlights.FindByKey(Category))
+					{
+						SetColor(CurHighlight->Color);
+						bNeedToResetColor = true;
+					}
+				}
+
+				if (LogStringHighlights.Num() > 0)
+				{
+					FStringView DataView = TStringView(Data);
+
+					for (const FLogStringHighlight& CurStringHighlight : LogStringHighlights)
+					{
+						if (FindFirst(DataView, CurStringHighlight.SearchString, ESearchCase::IgnoreCase) != INDEX_NONE)
+						{
+							SetColor(CurStringHighlight.Color);
+							bNeedToResetColor = true;
+
+							break;
+						}
+					}
+				}
+			}
+#endif
+		}
+		TStringBuilder<MAX_SPRINTF> OutputString;
+		FOutputDeviceHelper::AppendFormatLogLine(OutputString, Verbosity, Category, Data, GPrintLogTimes, RealTime);
+
+		{
+			FReadScopeLock lock(WindowRWLock);
+			if (Window)
+			{
+				Window->AddLogEntry(OutputString, TextAttribute);
+			}
+		}
+
+		if (bNeedToResetColor)
+		{
+			SetColor( COLOR_NONE );
+		}
+	}
+}
+
+void FWindowsConsoleOutputDevice2::Serialize( const TCHAR* Data, ELogVerbosity::Type Verbosity, const class FName& Category )
+{
+	Serialize( Data, Verbosity, Category, -1.0 );
+}
+
+void FWindowsConsoleOutputDevice2::SetColor( const TCHAR* Color )
+{
+	// here we can change the color of the text to display, it's in the format:
+	// ForegroundRed | ForegroundGreen | ForegroundBlue | ForegroundBright | BackgroundRed | BackgroundGreen | BackgroundBlue | BackgroundBright
+	// where each value is either 0 or 1 (can leave off trailing 0's), so 
+	// blue on bright yellow is "00101101" and red on black is "1"
+	// An empty string reverts to the normal gray on black
+	if (FCString::Stricmp(Color, TEXT("")) == 0)
+	{
+		TextAttribute = FOREGROUND_BLUE | FOREGROUND_GREEN | FOREGROUND_RED;
+	}
+	else
+	{
+		// turn the string into a bunch of 0's and 1's
+		TCHAR String[9];
+		FMemory::Memset(String, 0, sizeof(TCHAR) * UE_ARRAY_COUNT(String));
+		FCString::Strncpy(String, Color, UE_ARRAY_COUNT(String));
+		for (TCHAR* S = String; *S; S++)
+		{
+			*S -= '0';
+		}
+		// make the color
+		TextAttribute = 
+			(String[0] ? FOREGROUND_RED			: 0) | 
+			(String[1] ? FOREGROUND_GREEN		: 0) | 
+			(String[2] ? FOREGROUND_BLUE		: 0) | 
+			(String[3] ? FOREGROUND_INTENSITY	: 0) | 
+			(String[4] ? BACKGROUND_RED			: 0) | 
+			(String[5] ? BACKGROUND_GREEN		: 0) | 
+			(String[6] ? BACKGROUND_BLUE		: 0) | 
+			(String[7] ? BACKGROUND_INTENSITY	: 0);
+	}
+}
+
+bool FWindowsConsoleOutputDevice2::IsAttached()
+{
+	return Window != nullptr;
+}
+
+bool FWindowsConsoleOutputDevice2::CanBeUsedOnAnyThread() const
+{
+	return true;
+}
+
diff --git a/Engine/Source/Runtime/ApplicationCore/Private/Windows/WindowsPlatformApplicationMisc.cpp b/Engine/Source/Runtime/ApplicationCore/Private/Windows/WindowsPlatformApplicationMisc.cpp
--- a/Engine/Source/Runtime/ApplicationCore/Private/Windows/WindowsPlatformApplicationMisc.cpp
+++ b/Engine/Source/Runtime/ApplicationCore/Private/Windows/WindowsPlatformApplicationMisc.cpp
@@ -4,6 +4,7 @@
 #include "Windows/WindowsApplication.h"
 #include "Windows/WindowsApplicationErrorOutputDevice.h"
 #include "Windows/WindowsConsoleOutputDevice.h"
+#include "Windows/WindowsConsoleOutputDevice2.h"
 #include "Windows/WindowsFeedbackContext.h"
 #include "HAL/FeedbackContextAnsi.h"
 #include "Misc/App.h"
@@ -51,7 +52,10 @@ void FWindowsPlatformApplicationMisc::LoadStartupModules()
 class FOutputDeviceConsole* FWindowsPlatformApplicationMisc::CreateConsoleOutputDevice()
 {
 	// this is a slightly different kind of singleton that gives ownership to the caller and should not be called more than once
-	return new FWindowsConsoleOutputDevice();
+	if (FParse::Param(FCommandLine::Get(), TEXT("NewConsole")))
+		return new FWindowsConsoleOutputDevice2();
+	else
+		return new FWindowsConsoleOutputDevice();
 }
 
 class FOutputDeviceError* FWindowsPlatformApplicationMisc::GetErrorOutputDevice()
diff --git a/Engine/Source/Runtime/ApplicationCore/Public/Windows/WindowsConsoleOutputDevice2.h b/Engine/Source/Runtime/ApplicationCore/Public/Windows/WindowsConsoleOutputDevice2.h
new file mode 100644
--- /dev/null
+++ b/Engine/Source/Runtime/ApplicationCore/Public/Windows/WindowsConsoleOutputDevice2.h
@@ -0,0 +1,102 @@
+// Copyright Epic Games, Inc. All Rights Reserved.
+
+#pragma once
+
+#include "CoreTypes.h"
+#include "Misc/OutputDeviceConsole.h"
+#include "UObject/NameTypes.h"
+
+/**
+ * Windows implementation of console log window, utilizing the Win32 console API
+ */
+class APPLICATIONCORE_API FWindowsConsoleOutputDevice2 final : public FOutputDeviceConsole
+{
+private:
+	uint16 TextAttribute;
+
+	/** true if the color is currently set by the caller */
+	bool OverrideColorSet;
+
+	class FConsoleWindow;
+	FConsoleWindow* Window;
+	FRWLock WindowRWLock;
+
+#if !UE_BUILD_SHIPPING
+	/** An entry for log category highlighting */
+	struct FLogHighlight
+	{
+		/** The category to highlight */
+		FName Category;
+
+		/** The 4 digit color code to highlight with */
+		const TCHAR* Color = nullptr;
+
+		bool operator == (FName InCategory) const
+		{
+			return Category == InCategory;
+		}
+	};
+
+	/** Log categories to be highlighted */
+	TArray<FLogHighlight, TInlineAllocator<8>> LogHighlights;
+
+	/** An entry for log string highlighting */
+	struct FLogStringHighlight
+	{
+		/** The string to search for and highlight */
+		TStringBuilder<128> SearchString;
+
+		/** The 4 digit color code to highlight with */
+		const TCHAR* Color = nullptr;
+	};
+
+	/** Log strings to be highlighted */
+	TArray<FLogStringHighlight, TInlineAllocator<8>> LogStringHighlights;
+#endif
+
+	/**
+	 * Saves the console window's position and size to the game .ini
+	 */
+	void SaveToINI();
+
+public:
+
+	/** 
+	 * Constructor, setting console control handler.
+	 */
+	FWindowsConsoleOutputDevice2();
+	~FWindowsConsoleOutputDevice2();
+
+	/**
+	 * Shows or hides the console window. 
+	 *
+	 * @param ShowWindow	Whether to show (true) or hide (false) the console window.
+	 */
+	virtual void Show( bool ShowWindow );
+
+	/** 
+	 * Returns whether console is currently shown or not
+	 *
+	 * @return true if console is shown, false otherwise
+	 */
+	virtual bool IsShown();
+
+	virtual bool IsAttached();
+
+	virtual bool CanBeUsedOnAnyThread() const override;
+
+	/**
+	 * Displays text on the console and scrolls if necessary.
+	 *
+	 * @param Data	Text to display
+	 * @param Event	Event type, used for filtering/ suppression
+	 */
+	virtual void Serialize( const TCHAR* Data, ELogVerbosity::Type Verbosity, const class FName& Category, const double Time ) override;
+	virtual void Serialize( const TCHAR* Data, ELogVerbosity::Type Verbosity, const class FName& Category ) override;
+
+	void SetColor( const TCHAR* Color );
+
+private:
+	FWindowsConsoleOutputDevice2(const FWindowsConsoleOutputDevice2&) = delete;
+	FWindowsConsoleOutputDevice2& operator=(const FWindowsConsoleOutputDevice2&) = delete;
+};
diff --git a/Engine/Source/Runtime/Core/Private/Windows/WindowsPlatformMisc.cpp b/Engine/Source/Runtime/Core/Private/Windows/WindowsPlatformMisc.cpp
--- a/Engine/Source/Runtime/Core/Private/Windows/WindowsPlatformMisc.cpp
+++ b/Engine/Source/Runtime/Core/Private/Windows/WindowsPlatformMisc.cpp
@@ -682,6 +682,11 @@ void FWindowsPlatformMisc::SetGracefulTerminationHandler()
 #endif // !UE_BUILD_SHIPPING && PLATFORM_CPU_X86_FAMILY
 }
 
+void FWindowsPlatformMisc::CallGracefulTerminationHandler()
+{
+	ConsoleCtrlHandler(CTRL_CLOSE_EVENT);
+}
+
 static ECrashHandlingType GCrashHandlingType; /* = ECrashHandlingType::Default */
 
 ECrashHandlingType FWindowsPlatformMisc::GetCrashHandlingType()
diff --git a/Engine/Source/Runtime/Core/Public/Windows/MinWindows.h b/Engine/Source/Runtime/Core/Public/Windows/MinWindows.h
--- a/Engine/Source/Runtime/Core/Public/Windows/MinWindows.h
+++ b/Engine/Source/Runtime/Core/Public/Windows/MinWindows.h
@@ -39,7 +39,7 @@
 //#define NOCLIPBOARD			// Clipboard routines
 //#define NOCOLOR				// Screen colors
 //#define NOCTLMGR				// Control and Dialog routines
-#define NODRAWTEXT				// DrawText() and DT_*
+//#define NODRAWTEXT				// DrawText() and DT_*
 //#define NOGDI					// All GDI #defines and routines
 #define NOKERNEL				// All KERNEL #defines and routines
 //#define NOUSER				// All USER #defines and routines
diff --git a/Engine/Source/Runtime/Core/Public/Windows/WindowsPlatformMisc.h b/Engine/Source/Runtime/Core/Public/Windows/WindowsPlatformMisc.h
--- a/Engine/Source/Runtime/Core/Public/Windows/WindowsPlatformMisc.h
+++ b/Engine/Source/Runtime/Core/Public/Windows/WindowsPlatformMisc.h
@@ -47,6 +47,7 @@ struct CORE_API FWindowsPlatformMisc
 	static void PlatformPreInit();
 	static void PlatformInit();
 	static void SetGracefulTerminationHandler();
+	static void CallGracefulTerminationHandler();
 	static ECrashHandlingType GetCrashHandlingType();
 	static ECrashHandlingType SetCrashHandlingType(ECrashHandlingType);
 	static int32 GetMaxPathLength();
