[ApplicationCore - WindowsConsoleOutputDevice2]
* Fixed deadlock that could happen when setting info in ini file triggered log which in turn was outputted to console (this could happen if runtime early out during startup)
* Fixed bad visual behavior when log becomes bigger than log list entry max count

diff --git a/Engine/Source/Runtime/ApplicationCore/Private/Windows/WindowsConsoleOutputDevice2.cpp b/Engine/Source/Runtime/ApplicationCore/Private/Windows/WindowsConsoleOutputDevice2.cpp
--- a/Engine/Source/Runtime/ApplicationCore/Private/Windows/WindowsConsoleOutputDevice2.cpp
+++ b/Engine/Source/Runtime/ApplicationCore/Private/Windows/WindowsConsoleOutputDevice2.cpp
@@ -1044,33 +1044,38 @@ public:
 				constexpr int32 MaxSize = 20000;
 
 				int32 NewSize = Log.Num();
+
+				int32 ToAddToHwnd = NewSize - OldSize;
+				int32 LogIndex = OldSize;
+
 				if (NewSize > MaxSize)
 				{
 					int32 ToRemove = NewSize - MaxSize;
 					Log.PopFront(ToRemove);
 					LogIndexOffset += ToRemove;
 
-					int32 LogIndex = SendMessageW(LogHwnd, LB_GETITEMDATA, 0, 0);
-					while (ToRemove--)
-					{
-						if (LogIndex > LogIndexOffset)
-							break;
+					int32 ToChangeInHwnd = FMath::Min(MaxSize, ToRemove);
+
+					int32 ToRemoveFromHwnd = ToChangeInHwnd;
+					while (ToRemoveFromHwnd--)
 						SendMessageW(LogHwnd, LB_DELETESTRING, 0, 0);
-						++LogIndex;
-					}
+
+					ToAddToHwnd = ToChangeInHwnd;
+					LogIndex = MaxSize - ToAddToHwnd;
 				}
 
-				for (int32 LogIndex = OldSize, LogNum = Log.Num(); LogIndex != LogNum; ++LogIndex)
+				while (ToAddToHwnd--)
 				{
-					LogEntry& E = Log[LogIndex];
 					AddEntryToLogHwnd(Log[LogIndex], LogIndexOffset + LogIndex);
+					++LogIndex;
 				}
 				TempLogEntries.Reset(0);
 			}
 
 			if (bAutoScrollLog)
+			{
 				ScrollToBottom();
-
+			}
 			RedrawLogScrollbar();
 
 			SendMessageW(LogHwnd, WM_SETREDRAW, true, 0);
@@ -1989,11 +1994,20 @@ FWindowsConsoleOutputDevice2::~FWindowsConsoleOutputDevice2()
 
 void FWindowsConsoleOutputDevice2::SaveToINI()
 {
-	FWriteScopeLock lock(WindowRWLock);
-
+	WindowRWLock.WriteLock();
 	if (!Window)
+	{
+		WindowRWLock.WriteUnlock();
 		return;
-	
+	}
+	int32 ConsoleWidth = Window->ConsoleWidth;
+	int32 ConsoleHeight = Window->ConsoleHeight;
+	int32 ConsolePosX = Window->ConsolePosX;
+	int32 ConsolePosY = Window->ConsolePosY;
+	int32 ConsoleColor = (int32)Window->ConsoleColor;
+	bool bConsoleExpanded = Window->bConsoleExpanded;
+	WindowRWLock.WriteUnlock();
+
 	FString Filename = TEXT("DebugConsole.ini");
 	const TCHAR* Selection = TEXT("ConsoleWindows");
 	if (IsRunningDedicatedServer())
@@ -2004,15 +2018,15 @@ void FWindowsConsoleOutputDevice2::SaveToINI()
 	FConfigCacheIni Config(EConfigCacheType::DiskBacked);
 	Config.LoadFile(Filename);
 
-	Config.SetInt(Selection, TEXT("ConsoleX"), Window->ConsolePosX, Filename);
-	Config.SetInt(Selection, TEXT("ConsoleY"), Window->ConsolePosY, Filename);
+	Config.SetInt(Selection, TEXT("ConsoleX"), ConsolePosX, Filename);
+	Config.SetInt(Selection, TEXT("ConsoleY"), ConsolePosY, Filename);
 
-	Config.SetInt(Selection, TEXT("ConsoleWidth"), Window->ConsoleWidth, Filename);
-	Config.SetInt(Selection, TEXT("ConsoleHeight"), Window->ConsoleHeight, Filename);
+	Config.SetInt(Selection, TEXT("ConsoleWidth"), ConsoleWidth, Filename);
+	Config.SetInt(Selection, TEXT("ConsoleHeight"), ConsoleHeight, Filename);
 
-	Config.SetInt(Selection, TEXT("ConsoleColor"), int32(Window->ConsoleColor), Filename);
+	Config.SetInt(Selection, TEXT("ConsoleColor"), ConsoleColor, Filename);
 
-	Config.SetBool(Selection, TEXT("ConsoleExpanded"), Window->bConsoleExpanded, Filename);
+	Config.SetBool(Selection, TEXT("ConsoleExpanded"), bConsoleExpanded, Filename);
 
 	Config.Flush(false, Filename);
 }
