[ApplicationCore]
* WindowsConsoleOutputDevice2 - Fixed warning from static analysis (ConsoleType assigned twice)

diff --git a/Engine/Source/Runtime/ApplicationCore/Private/Windows/WindowsConsoleOutputDevice2.cpp b/Engine/Source/Runtime/ApplicationCore/Private/Windows/WindowsConsoleOutputDevice2.cpp
--- a/Engine/Source/Runtime/ApplicationCore/Private/Windows/WindowsConsoleOutputDevice2.cpp
+++ b/Engine/Source/Runtime/ApplicationCore/Private/Windows/WindowsConsoleOutputDevice2.cpp
@@ -2079,9 +2079,9 @@ void FWindowsConsoleOutputDevice2::Show( bool ShowWindow )
 			// Setting text so they are easily identifyable in taskbar
 			const TCHAR* ConsoleType = TEXT("Unreal");
 			if (IsRunningDedicatedServer())
-				ConsoleType = ConsoleType = TEXT("Server");
+				ConsoleType = TEXT("Server");
 			else if (IsRunningGame())
-				ConsoleType = ConsoleType = TEXT("Client");
+				ConsoleType = TEXT("Client");
 			ConsoleTitle = FString::Printf(TEXT("%s Console (%s) - %s"), ConsoleType, FApp::GetProjectName(), FPlatformProcess::ExecutablePath());
 		}
 
