Add functions to IPlatformInputDeviceMapper for converting to and from legacy user index that should work the same as the ones on FPlatformMisc
Fix the user login callback to use user index instead of raw platform id because the raw platform ids are not necessarily sequential

diff --git a/Engine/Source/Runtime/ApplicationCore/Private/GenericPlatform/GenericPlatformInputDeviceMapper.cpp b/Engine/Source/Runtime/ApplicationCore/Private/GenericPlatform/GenericPlatformInputDeviceMapper.cpp
--- a/Engine/Source/Runtime/ApplicationCore/Private/GenericPlatform/GenericPlatformInputDeviceMapper.cpp
+++ b/Engine/Source/Runtime/ApplicationCore/Private/GenericPlatform/GenericPlatformInputDeviceMapper.cpp
@@ -353,6 +353,20 @@ bool FGenericPlatformInputDeviceMapper::RemapControllerIdToPlatformUserAndDevice
 	return false;
 }
 
+FPlatformUserId FGenericPlatformInputDeviceMapper::GetPlatformUserForUserIndex(int32 LocalUserIndex)
+{
+	FPlatformUserId OutUserId;
+	FInputDeviceId TempDeviceId;
+
+	// The platform user index is equivalent to ControllerId in most legacy code
+	if (RemapControllerIdToPlatformUserAndDevice(LocalUserIndex, OutUserId, TempDeviceId))
+	{
+		return OutUserId;
+	}
+
+	return PLATFORMUSERID_NONE;
+}
+
 bool FGenericPlatformInputDeviceMapper::RemapUserAndDeviceToControllerId(FPlatformUserId UserId, int32& OutControllerId, FInputDeviceId OptionalDeviceId /* = INPUTDEVICEID_NONE */)
 {
 	// It's just a 1:1 mapping of the old ControllerId to PlatformId if this is true
@@ -364,6 +378,17 @@ bool FGenericPlatformInputDeviceMapper::RemapUserAndDeviceToControllerId(FPlatfo
 	return false;
 }
 
+int32 FGenericPlatformInputDeviceMapper::GetUserIndexForPlatformUser(FPlatformUserId UserId)
+{
+	// The platform user index is equivalent to ControllerId in most legacy code
+	int32 OutControllerId = INDEX_NONE;
+	if (RemapUserAndDeviceToControllerId(UserId, OutControllerId))
+	{
+		return OutControllerId;
+	}
+	return INDEX_NONE;
+}
+
 bool FGenericPlatformInputDeviceMapper::IsUsingControllerIdAsUserId() const
 {
 	return bUsingControllerIdAsUserId;
@@ -376,10 +401,8 @@ bool FGenericPlatformInputDeviceMapper::ShouldBroadcastLegacyDelegates() const
 
 void FGenericPlatformInputDeviceMapper::OnUserLoginChangedEvent(bool bLoggedIn, int32 RawPlatformUserId, int32 UserIndex)
 {
-	// Attain the FPlatformUserId from the raw int32 that will be given by the Platform code
-	FPlatformUserId LoggedOutPlatformUserId = PLATFORMUSERID_NONE;
-	FInputDeviceId DummyInputDevice = INPUTDEVICEID_NONE;
-	RemapControllerIdToPlatformUserAndDevice(RawPlatformUserId, LoggedOutPlatformUserId, DummyInputDevice);
+	// Attain the FPlatformUserId from the user index that will be given by the Platform code
+	FPlatformUserId LoggedOutPlatformUserId = GetPlatformUserForUserIndex(UserIndex);
 	
 	// A new user has logged in
 	if (bLoggedIn)
diff --git a/Engine/Source/Runtime/ApplicationCore/Public/GenericPlatform/GenericPlatformInputDeviceMapper.h b/Engine/Source/Runtime/ApplicationCore/Public/GenericPlatform/GenericPlatformInputDeviceMapper.h
--- a/Engine/Source/Runtime/ApplicationCore/Public/GenericPlatform/GenericPlatformInputDeviceMapper.h
+++ b/Engine/Source/Runtime/ApplicationCore/Public/GenericPlatform/GenericPlatformInputDeviceMapper.h
@@ -167,8 +167,26 @@ public:
 	// Functions to provide compatibility between the old "int32 ControllerId"
 	// and the new FPlatformUserId and FInputDeviceId structs.
 	
+	/**
+	 * Remap a FPlatformUserId to a ControllerId that is needed by legacy code. Use this function to
+	 * add compatibility to platforms that may not have implemented this device mapper yet.
+	 *
+	 * @param UserId				The platform user that should be converted
+	 * @param OutControllerId		Set to the old plain "int32" that represented gamepad id or user index depending on the context.
+	 * @param OptionalDeviceId		Set this to a specific device id if known
+	 * @return						True if this maps to a real user
+	 */
 	virtual bool RemapUserAndDeviceToControllerId(FPlatformUserId UserId, int32& OutControllerId, FInputDeviceId OptionalDeviceId = INPUTDEVICEID_NONE) = 0;
 	
+	/*
+	 * Remap a FPlatformUserId to the legacy platform LocalUserIndex (often the same as ControllerId).
+	 * This is useful in cases where the input device id is not needed.
+	 *
+	 * @param UserId			The platform user that should be converted
+	 * @returns					A valid platform user index, or INDEX_NONE
+	 */
+	virtual int32 GetUserIndexForPlatformUser(FPlatformUserId UserId) = 0;
+
 	/**
 	 * Remap the legacy "int32 ControllerId" to the updated FPlatformUserId and FInputDeviceId. Use this function to
 	 * add compatibility to platforms that may not have implemented this device mapper yet.
@@ -177,13 +195,22 @@ public:
 	 * use the "int32 ControllerId" as a parameter so that you can call the new FGenericApplicationMessageHandler
 	 * that take in a PlatformUserId and an InputDeviceId.
 	 * 
-	 * @param ControllerId			The old plain "int32" that represented gamepad id or user id depending on the context.
+	 * @param ControllerId			The old plain "int32" that represented gamepad id or user index depending on the context.
 	 * @param InOutUserId			If the old function provides a PlatformId then pass it here, otherwise pass PLATFORMUSERID_NONE.
 	 * @param OutInputDeviceId		The best guess for an InputDeviceId based on the legacy int32 ControllerId. This may be INPUTDEVICEID_NONE.
 	 * @return						True if this maps to a real user
 	 */
 	virtual bool RemapControllerIdToPlatformUserAndDevice(int32 ControllerId, FPlatformUserId& InOutUserId, FInputDeviceId& OutInputDeviceId) = 0;
 
+	/*
+	 * Remap the legacy platform LocalUserIndex (often the same as ControllerId) to a FPlatformUserId.
+	 * This is useful in cases where the input device id is not needed.
+	 *
+	 * @param LocalUserIndex	A platform user index where 0 should always be a valid user
+	 * @returns					A platform-allocated user id, or PLATFORMUSERID_NONE
+	 */
+	virtual FPlatformUserId GetPlatformUserForUserIndex(int32 LocalUserIndex) = 0;
+
 protected:
 
 	/**
@@ -260,7 +287,9 @@ public:
 	virtual FInputDeviceId GetDefaultInputDevice() const override;
 
 	virtual bool RemapControllerIdToPlatformUserAndDevice(int32 ControllerId, FPlatformUserId& InOutUserId, FInputDeviceId& OutInputDeviceId) override;
+	virtual FPlatformUserId GetPlatformUserForUserIndex(int32 LocalUserIndex) override;
 	virtual bool RemapUserAndDeviceToControllerId(FPlatformUserId UserId, int32& OutControllerId, FInputDeviceId OptionalDeviceId = INPUTDEVICEID_NONE) override;
+	virtual int32 GetUserIndexForPlatformUser(FPlatformUserId UserId) override;
 	virtual bool IsUsingControllerIdAsUserId() const override;
 	virtual bool ShouldBroadcastLegacyDelegates() const override;
 	
