[ApplicationCore - WindowsConsoleOutputDevice2]
* Fixed so selection/caret position is updated when number of log entries becomes larger than history buffer size
* Fixed so selection using keyboard and mouse works well number of log entries becomes larger than history buffer size
* Added tool tip hover feature for log listbox items
* Added help dialog
* Added popup menu with lots of features. Press right mouse button over log list view
* Added tooltip window that extends hovered log entry outside of log window if needed.
* Added so date/time is hidden by default and can be shown with commands showtime and showdatetime
* Fixed so filter include/exclude edit boxes size with log window size
* Added so include/exclude/command edit boxes show describable text when empty and removed static text in front of edit boxes
* Fixed so shift-end works properly to expand selection
* Fixed bug with log window default position is CW_USEDEFAULT when no stored position exists

diff --git a/Engine/Source/Runtime/ApplicationCore/Private/Windows/WindowsConsoleOutputDevice2.cpp b/Engine/Source/Runtime/ApplicationCore/Private/Windows/WindowsConsoleOutputDevice2.cpp
--- a/Engine/Source/Runtime/ApplicationCore/Private/Windows/WindowsConsoleOutputDevice2.cpp
+++ b/Engine/Source/Runtime/ApplicationCore/Private/Windows/WindowsConsoleOutputDevice2.cpp
@@ -21,6 +21,7 @@
 #include "Misc/OutputDeviceHelper.h"
 #include "Misc/Parse.h"
 #include "Misc/ScopeRWLock.h"
+#include <Misc/Timespan.h>
 #include "Misc/TrackedActivity.h"
 #include "String/Find.h"
 #include "Templates/UnrealTemplate.h"
@@ -32,10 +33,30 @@
 #include <uxtheme.h>
 #include "Windows/HideWindowsPlatformTypes.h"
 
-
 class FWindowsConsoleOutputDevice2::FConsoleWindow
 {
 public:
+	static constexpr const TCHAR* HelpTitle = TEXT("Unreal Console Help");
+	static constexpr const TCHAR* HelpText = TEXT(
+		"Welcome to Unreal Console!\r\n"
+		"\r\n"
+		"Unreal Console is a replacement of the built-in console. If you still want to use windows console, remove '-NewConsole'\r\n"
+		"\r\n"
+		"Use double down-arrow in top right corner to expand console to show command line and other features.\r\n"
+		"\r\n"
+		"Console commands:\r\n"
+		"  showdate\t - Show time in log\r\n"
+		"  showdatetime\t - Show date/time in log\r\n"
+		"  hidedatetime\t - Hide date/time from log\r\n"
+		"  <color>\t - Change color of console\r\n"
+		"\r\n"
+		"Supported colors: red,gray,darkgray,lightblue,darkblue\r\n"
+		"\r\n"
+		"Console settings are saved separately for client/server/other in basedir folder\r\n"
+		"\r\n"
+		"Enjoy!\r\n"
+	);
+
 	enum
 	{
 		WM_NEWLOGENTRIES = WM_USER + 1,
@@ -85,13 +106,13 @@ public:
 			DeleteObject(StatusLightBrush[I]);
 	}
 
-	void AddLogEntry(const FStringView& Text, uint16 TextAttribute)
+	void AddLogEntry(const FStringView& Text, ELogVerbosity::Type Verbosity, const class FName& Category, double Time, uint16 TextAttribute)
 	{
 		bool bWasEmpty;
 		{
 			FScopeLock _(&NewLogEntriesCs);
 			bWasEmpty = NewLogEntries.Num() == 0;
-			NewLogEntries.Add({FString(Text), TextAttribute});
+			NewLogEntries.Add({FString(Text), Verbosity, Category, Time, TextAttribute});
 		}
 		if (bWasEmpty)
 			PostMessageW(MainHwnd, WM_NEWLOGENTRIES, 0, 0);
@@ -192,35 +213,38 @@ public:
 		SystemParametersInfo(SPI_GETNONCLIENTMETRICS, sizeof(NonClientMetrics), &NonClientMetrics, 0);
 		Font = (HFONT)CreateFontIndirect(&NonClientMetrics.lfMessageFont);
 
+		const TCHAR* FontName = TEXT("Cascadia Mono"); //TEXT("Courier New"); // TEXT("Consolas");
 		int32 FontHeight = -MulDiv(8, GetDeviceCaps(GetDC(NULL), LOGPIXELSY), 72);
-		HFONT hLogFont = (HFONT)CreateFontW(FontHeight, 0, 0, 0, FW_NORMAL, false, false, false, ANSI_CHARSET, OUT_DEFAULT_PRECIS, CLIP_DEFAULT_PRECIS, ANTIALIASED_QUALITY, FIXED_PITCH | FF_MODERN, TEXT("Courier New"));
+		LogFont = (HFONT)CreateFontW(FontHeight, 0, 0, 0, FW_NORMAL, false, false, false, ANSI_CHARSET, OUT_DEFAULT_PRECIS, CLIP_DEFAULT_PRECIS, ANTIALIASED_QUALITY, FIXED_PITCH | FF_MODERN, FontName);
 
 		DWORD WindowStyle = WS_CAPTION | WS_SYSMENU | WS_THICKFRAME | WS_MAXIMIZEBOX | WS_MINIMIZEBOX | WS_CLIPCHILDREN;
 		const TCHAR* WindowClassName = MAKEINTATOM(WndClassAtom);
 		MainHwnd = CreateWindow(WindowClassName, *ConsoleTitle, WindowStyle, ConsolePosX, ConsolePosY, ConsoleWidth, ConsoleHeight, NULL, NULL, HInstance, this);
 		SetWindowLongPtr(MainHwnd, GWLP_USERDATA, (LONG_PTR)this);
+		RECT MainRect;
+		GetWindowRect(MainHwnd, &MainRect);
+		ConsolePosX = MainRect.left;
+		ConsolePosY = MainRect.top;
+		ConsoleWidth = MainRect.right - ConsolePosX;
+		ConsoleHeight = MainRect.bottom - ConsolePosY;
 
-		int32 X = 20;
-		int32 Y = 10;
-		int32 Height = 24;
-		FilterHwnd[0] = CreateTextHwnd(TEXT("Log Filter"), Font, X, Y, 70, 20);
-		X += 100;
-		FilterHwnd[1] = CreateTextHwnd(TEXT("Include:"), Font, X, Y, 47, 20);
-		X += 50;
-		FilterHwnd[2] = CreateWindow2(WC_EDIT, NULL, ES_AUTOHSCROLL, Font, X, Y - 1, 200, 20, ID_LOG_INCLUDE_FILTER);
-		X += 240;
-		FilterHwnd[3] = CreateTextHwnd(TEXT("Exclude:"), Font, X, Y, 47, 20);
-		X += 50;
-		FilterHwnd[4] = CreateWindow2(WC_EDIT, NULL, ES_AUTOHSCROLL, Font, X, Y - 1, 200, 20, ID_LOG_EXCLUDE_FILTER);
-		X += 250;
-		FilterHwnd[5] = CreateWindow2(WC_BUTTON, TEXT("Clear Log"), 0, Font, X, Y - 2, 80, 22, ID_CLEARLOGBUTTON);
-		X += 90;
-		FilterHwnd[6] = CreateWindow2(WC_BUTTON, TEXT("Log CHECKPOINT0"), 0, Font, X, Y - 2, 140, 22, ID_ADDCHECKPOINTBUTTON);
+		SIZE FontSize;
+		HDC Hdc = GetDC(MainHwnd);
+		SelectObject(Hdc, LogFont);
+		GetTextExtentPoint32(Hdc, TEXT("A"), 1, &FontSize);
+		LogFontWidth = FontSize.cx;
+		LogFontHeight = FontSize.cy;
+		ReleaseDC(MainHwnd, Hdc);
+
+		IncludeFilterHwnd = CreateWindow2(WC_EDIT, NULL, ES_AUTOHSCROLL, Font, 0, 0, 0, 0, ID_LOG_INCLUDE_FILTER);
+		ExcludeFilterHwnd = CreateWindow2(WC_EDIT, NULL, ES_AUTOHSCROLL, Font, 0, 0, 0, 0, ID_LOG_EXCLUDE_FILTER);
+		ClearLogButtonHwnd = CreateWindow2(WC_BUTTON, TEXT("Clear Log"), 0, Font, 0, 0, 0, 0, ID_CLEARLOGBUTTON);
+		CheckpointButtonHwnd = CreateWindow2(WC_BUTTON, TEXT("Log CHECKPOINT0"), 0, Font, 0, 0, 0, 0, ID_ADDCHECKPOINTBUTTON);
 
 		DWORD LogHwndStyle = LBS_NOINTEGRALHEIGHT | LBS_EXTENDEDSEL | LBS_WANTKEYBOARDINPUT | LBS_NOTIFY | LBS_HASSTRINGS | LBS_OWNERDRAWFIXED;
-		LogHwnd = CreateWindow2(WC_LISTBOX, NULL, LogHwndStyle, hLogFont, 0, 0, 1, 1, ID_LOG);
+		LogHwnd = CreateWindow2(WC_LISTBOX, NULL, LogHwndStyle, LogFont, 0, 0, 1, 1, ID_LOG);
 		SetWindowLongPtr(LogHwnd, GWLP_USERDATA, (LONG_PTR)this);
-		LogHwndWndProc = (WNDPROC)SetWindowLongPtr(LogHwnd, GWLP_WNDPROC, (LONG_PTR)StaticLogHwndWndProc);
+		LogHwndWndProcPtr = (WNDPROC)SetWindowLongPtr(LogHwnd, GWLP_WNDPROC, (LONG_PTR)StaticLogHwndWndProc);
 
 		const TCHAR* ScrollBarClassName = MAKEINTATOM(ScrollBarClassAtom);
 		LogScrollHwnd = CreateWindow2(ScrollBarClassName, NULL, SBS_VERT, 0);
@@ -231,11 +255,6 @@ public:
 
 		SetFocus(CommandHwnd);
 
-		SIZE FontSize;
-		GetTextExtentPoint32(GetDC(LogHwnd), TEXT("A"), 1, &FontSize);
-		LogFontWidth = FontSize.cx;
-		LogFontHeight = FontSize.cy - 2; // We want more compact listbox
-
 		RECT Rect;
 		GetClientRect(MainHwnd, &Rect);
 		UpdateSize(Rect.right, Rect.bottom, false);
@@ -266,6 +285,10 @@ public:
 
 		SetForegroundWindow(MainHwnd);
 
+
+		TipHwnd = CreateWindowEx(WS_EX_TOOLWINDOW | WS_EX_TOPMOST | WS_EX_NOACTIVATE, WC_STATIC, TEXT(""), SS_OWNERDRAW | WS_VISIBLE | WS_POPUP | WS_DISABLED | WS_CLIPCHILDREN | WS_CLIPSIBLINGS, 1, 1, 1, 1, MainHwnd, NULL, HInstance, NULL);
+		SendMessageW(TipHwnd, WM_SETFONT, (WPARAM)Font, 0);
+
 		PostMessageW(MainHwnd, WM_NEWLOGENTRIES, 0, 0);
 
 		FDateTime startTime(FDateTime::Now());
@@ -277,12 +300,21 @@ public:
 			if (bWasVisible != bIsVisible)
 			{
 				if (bIsVisible)
+				{
 					MoveWindow(MainHwnd, ConsolePosX, ConsolePosY, ConsoleWidth, ConsoleHeight, true);
+				}
+
 				ShowWindow(MainHwnd, bIsVisible ? SW_SHOW : SW_HIDE);
 				bWasVisible = bIsVisible;
+
+				if (bIsVisible)
+				{
+					UpdateIncludeFilter(*ConsoleIncludeFilterStr);
+					UpdateExcludeFilter(*ConsoleExcludeFilterStr);
+				}
 			}
 
-			uint64 seconds = uint64((FDateTime::Now() - startTime).GetTotalSeconds());
+			uint64 seconds = uint64(FPlatformTime::Seconds() - GStartTime);
 			if (lastSeconds != seconds)
 			{
 				lastSeconds = seconds;
@@ -303,7 +335,7 @@ public:
 					if (Msg.message == WM_QUIT)
 					{
 						DestroyWindow(MainHwnd);
-						DeleteObject(hLogFont);
+						DeleteObject(LogFont);
 						DeleteObject(Font);
 						UnregisterClass(WindowClassName, HInstance);
 						MainHwnd = 0;
@@ -593,7 +625,7 @@ public:
 
 		int32 ButtonWidth = 90;
 		int32 CommandFlags = Flags | (bConsoleExpanded ? SWP_SHOWWINDOW : SWP_HIDEWINDOW);
-		int32 FilterFlags = SWP_NOMOVE | SWP_NOSIZE | Flags;
+		int32 FilterFlags = 0;
 
 		int32 CommandY = ClientHeight - ActivitiesTotalHeight - 20 - 5;
 		int32 LogHeight = CommandY - LogTop - 5;
@@ -601,17 +633,29 @@ public:
 		if (bConsoleExpanded)
 		{
 			CommandFlags = Flags | SWP_SHOWWINDOW;
-			FilterFlags |= SWP_SHOWWINDOW;
+			FilterFlags = Flags | SWP_SHOWWINDOW;
 		}
 		else
 		{
 			CommandFlags = SWP_HIDEWINDOW;
-			FilterFlags |= SWP_HIDEWINDOW;
+			FilterFlags = Flags | SWP_HIDEWINDOW;
 			LogHeight += 25;
 		}
 
-		for (HWND Hwnd : FilterHwnd)
-			SetWindowPos(Hwnd, 0, 0, 0, 0, 0, FilterFlags);
+		int32 Y = 10;
+		int32 ButtonStart = ClientWidth - 235;
+
+		SetWindowPos(ClearLogButtonHwnd, 0, ButtonStart, Y - 2, 80, 22, FilterFlags);
+		SetWindowPos(CheckpointButtonHwnd, 0, ButtonStart + 94, Y - 2, 135, 22, FilterFlags);
+
+		int32 WidthForFilters = ButtonStart - 45;
+		if (WidthForFilters < 100)
+			FilterFlags = Flags | SWP_HIDEWINDOW;
+		int32 X = 8;
+		int32 WidthForEditBoxes = WidthForFilters / 2;
+		SetWindowPos(IncludeFilterHwnd, 0, X, Y - 1, WidthForEditBoxes, 20, FilterFlags);
+		X += WidthForEditBoxes + 20;
+		SetWindowPos(ExcludeFilterHwnd, 0, X, Y - 1, WidthForEditBoxes, 20, FilterFlags);
 
 		SetWindowPos(LogHwnd, 0, 7, LogTop, ClientWidth - 33, LogHeight, Flags);
 		SetWindowPos(LogScrollHwnd, 0, ClientWidth - 26, LogTop, 19, LogHeight, Flags);
@@ -654,7 +698,10 @@ public:
 
 	int32 AddEntryToLogHwnd(LogEntry& E, int32 LogVirtualIndex)
 	{
-		const TCHAR* Str = *E.String;
+		TStringBuilder<1024> OutString;
+		CreateLogEntryText(OutString, E, true);
+
+		const TCHAR* Str = *OutString;
 		for (const FString& I : IncludeFilter)
 			if (FCString::Stristr(Str, *I) == 0)
 				return -1;
@@ -662,7 +709,8 @@ public:
 			if (FCString::Stristr(Str, *I) != 0)
 				return -1;
 		AddedEntryLogVirtualIndex = LogVirtualIndex;
-		int32 ItemIndex = SendMessageW(LogHwnd, LB_ADDSTRING, 0, (LPARAM)Str);
+		const TCHAR* Empty = TEXT("");
+		int32 ItemIndex = SendMessageW(LogHwnd, LB_ADDSTRING, 0, (LPARAM)Empty);
 		SendMessageW(LogHwnd, LB_SETITEMDATA, ItemIndex, LogVirtualIndex);
 		AddedEntryLogVirtualIndex = -1;
 		return ItemIndex;
@@ -676,7 +724,7 @@ public:
 		{
 			int32 SelectedItemIndex = -1;
 			SendMessageW(LogHwnd, LB_GETSELITEMS, 1, (LPARAM)&SelectedItemIndex);
-			SelectedLogIndex = SendMessageW(LogHwnd, LB_GETITEMDATA, SelectedItemIndex, 0);
+			SelectedLogIndex = SendMessageW(LogHwnd, LB_GETITEMDATA, SelectedItemIndex, 0) - LogIndexOffset;
 			SelectedItemOffsetFromTop = SelectedItemIndex - SendMessageW(LogHwnd, LB_GETTOPINDEX, 0, 0);
 		}
 
@@ -699,7 +747,7 @@ public:
 		if (NewSelectedItemIndex != -1)
 		{
 			ScrollPos = NewSelectedItemIndex - SelectedItemOffsetFromTop;
-			SendMessageW(LogHwnd, LB_SETTOPINDEX, ScrollPos, 0);
+			SetTopVisible(ScrollPos, false);
 			SendMessageW(LogHwnd, LB_SETSEL, true, NewSelectedItemIndex);
 		}
 		else
@@ -711,6 +759,14 @@ public:
 		RedrawWindow(LogScrollHwnd, NULL, NULL, RDW_INVALIDATE);
 	}
 
+	void ClearSelection()
+	{
+		TArray<int> SelectedItems;
+		GetSelectedItems(SelectedItems);
+		for (int SelectedItem : SelectedItems)
+			SendMessageW(LogHwnd, LB_SETSEL, false, SelectedItem);;
+	}
+
 	void ScrollToBottom()
 	{
 		int32 ItemCount = (int)SendMessageW(LogHwnd, LB_GETCOUNT, 0, 0);
@@ -718,15 +774,25 @@ public:
 		// This is here just to force last selection to be at the bottom
 		SendMessageW(LogHwnd, LB_SETSEL, true, ItemCount - 1);
 		SendMessageW(LogHwnd, LB_SETSEL, false, ItemCount - 1);
+		SetTopVisible(ItemCount - 1, true); // Using post to remove weird glitches with smooth scrolling
+	}
 
-		PostMessageW(LogHwnd, LB_SETTOPINDEX, ItemCount-1, 0); // Using post to remove weird glitches with smooth scrolling
+	void SuspendAddingEntries()
+	{
+		bSuspendAddingEntries = true;
+	}
+
+	void ResumeAddingEntries()
+	{
+		if (!bSuspendAddingEntries)
+			return;
+		bSuspendAddingEntries = false;
+		PostMessageW(MainHwnd, WM_NEWLOGENTRIES, 0, 0);
 	}
 
 	static LRESULT CALLBACK StaticLogHwndWndProc(HWND hWnd, UINT Msg, WPARAM wParam, LPARAM lParam)
 	{
 		FConsoleWindow* ThisPtr = (FConsoleWindow*)GetWindowLongPtr(hWnd, GWLP_USERDATA);
-		if (Msg == WM_MOUSEWHEEL)
-			return ThisPtr->HandleMouseWheel(wParam);
 		return ThisPtr->LogHwndWndProc(hWnd, Msg, wParam, lParam);
 	}
 
@@ -740,7 +806,10 @@ public:
 			ThisPtr = (FConsoleWindow*)lParam;
 			SetWindowLongPtr(hWnd, GWLP_USERDATA, lParam);
 		}
-		return ThisPtr->MainWinProc(hWnd, Msg, wParam, lParam);
+		if (ThisPtr && hWnd == ThisPtr->MainHwnd)
+			return ThisPtr->MainWinProc(hWnd, Msg, wParam, lParam);
+		else
+			return DefWindowProc(hWnd, Msg, wParam, lParam);
 	}
 
 	static LRESULT CALLBACK StaticScrollBarWinProc(HWND hWnd, UINT Msg, WPARAM wParam, LPARAM lParam)
@@ -796,13 +865,22 @@ public:
 		return FMath::FloorToInt(float(ThumbPos - ButtonHeight) * TotalScrollCount / MoveSpace);
 	}
 
+	void SetTopVisible(int32 TopVisible, bool ShouldPost)
+	{
+		SetWindowPos(TipHwnd, MainHwnd, 1, 1, 1, 1, SWP_NOACTIVATE|SWP_NOSENDCHANGING|SWP_NOZORDER|SWP_NOOWNERZORDER);
+		if (ShouldPost)
+			PostMessageW(LogHwnd, LB_SETTOPINDEX, TopVisible, 0);
+		else
+			SendMessageW(LogHwnd, LB_SETTOPINDEX, TopVisible, 0);
+	}
+
 	void MoveTopVisible(const RECT& Rect, int32 Offset)
 	{
 		bAutoScrollLog = false;
 		int32 TopVisible = SendMessageW(LogHwnd, LB_GETTOPINDEX, 0, 0);
 		int32 PageSize = Rect.bottom / LogFontHeight;
 		TopVisible = FMath::Max(0, TopVisible + Offset);
-		SendMessageW(LogHwnd, LB_SETTOPINDEX, TopVisible, 0);
+		SetTopVisible(TopVisible, false);
 		RedrawWindow(LogScrollHwnd, NULL, NULL, RDW_INVALIDATE);
 	}
 
@@ -902,6 +980,7 @@ public:
 			{
 				LogScrollGrabPos = Y - ThumbPos;
 				SetCapture(hWnd);
+				SuspendAddingEntries();
 				break;
 			}
 			if (Y <= ThumbPos)
@@ -921,6 +1000,7 @@ public:
 		}
 		case WM_LBUTTONUP:
 			ReleaseCapture();
+			ResumeAddingEntries();
 			LogScrollGrabPos = -1;
 			break;
 
@@ -933,7 +1013,7 @@ public:
 				int32 Y = HIWORD(lParam);
 				int32 ThumbPos = Y - LogScrollGrabPos;
 				int32 TopVisible = GetTopVisible(ThumbPos, Rect);
-				SendMessageW(LogHwnd, LB_SETTOPINDEX, TopVisible, 0);
+				SetTopVisible(TopVisible, false);
 				RedrawWindow(hWnd, NULL, NULL, RDW_INVALIDATE);
 			}
 			break;
@@ -1004,78 +1084,483 @@ public:
 		RedrawWindow(MainHwnd, NULL, NULL, RDW_INVALIDATE | RDW_FRAME);
 	}
 
-	LRESULT MainWinProc(HWND hWnd, UINT Msg, WPARAM wParam, LPARAM lParam)
+	void CreateLogEntryText(TStringBuilderBase<TCHAR>& OutString, const LogEntry& E, bool ForFilter)
+	{
+		TStringBuilder<128> CategoryBuilder;
+		E.Category.AppendString(CategoryBuilder);
+
+
+		if (ConsoleShowDateTime == 1)
+		{
+			FTimespan Ts = FTimespan::FromSeconds(FPlatformTime::Seconds() - GStartTime);
+			OutString << '[';
+			OutString.Appendf(TEXT("%02i:"), Ts.GetHours() + Ts.GetDays()*24);
+			OutString.Appendf(TEXT("%02i:"), Ts.GetMinutes());
+			OutString.Appendf(TEXT("%02i."), Ts.GetSeconds());
+			OutString.Appendf(TEXT("%03i]"), Ts.GetFractionMilli());
+			OutString.Appendf(TEXT("[%3llu] "), GFrameCounter % 1000);
+		}
+		else if (ConsoleShowDateTime == 2)
+		{
+			FDateTime::Now().ToString(TEXT("[%Y.%m.%d-%H.%M.%S:%s]"), OutString);
+			OutString.Appendf(TEXT("[%3llu] "), GFrameCounter % 1000);
+		}
+
+		if (ForFilter)
+		{
+			OutString << '[' << CategoryBuilder << TEXT("][") << ToString(E.Verbosity) << ']';
+		}
+
+
+
+		const TCHAR* Category = *CategoryBuilder;
+		if (FCString::Strstr(Category, TEXT("Log")) == Category)
+			Category += 3;
+		const TCHAR CategorySpace[] = TEXT("                 ");
+		if (FCString::Strlen(Category) >= UE_ARRAY_COUNT(CategorySpace))
+			OutString.Append(Category, UE_ARRAY_COUNT(CategorySpace) - 1);
+		else
+			OutString.Append(Category).Append(CategorySpace + FCString::Strlen(Category));
+		OutString.Append(": ");
+		for (TCHAR C : E.String)
+		{
+			if (C == '\t')
+				OutString.Append(TEXT("  "));
+			else
+				OutString.AppendChar(C);
+		}
+	}
+
+	void EnableAutoScroll()
+	{
+		bAutoScrollLog = true;
+		ClearSelection();
+		ScrollToBottom();
+		RedrawWindow(LogHwnd, NULL, NULL, RDW_ERASE | RDW_INVALIDATE);
+		RedrawLogScrollbar();
+	}
+
+	void ClearLog()
+	{
+		LogIndexOffset = 0;
+		Log.Reset();
+		RefreshLogHwnd();
+		CheckpointIndex = 0;
+		SetDlgItemText(MainHwnd, ID_ADDCHECKPOINTBUTTON, TEXT("Log CHECKPOINT0"));
+	}
+
+	int32 GetSelectedItems(TArray<int>& OutSelectedItems)
+	{
+		int32 SelectionCount = SendMessageW(LogHwnd, LB_GETSELCOUNT, 0, 0);
+		OutSelectedItems.SetNum(SelectionCount);
+		SendMessageW(LogHwnd, LB_GETSELITEMS, SelectionCount, (LPARAM)OutSelectedItems.GetData());
+		return SelectionCount;
+	}
+
+	void CopySelectionToClipboard()
+	{
+		TArray<int> SelectedItems;
+		int32 SelectionCount = GetSelectedItems(SelectedItems);
+		if (!SelectionCount)
+			return;
+		TArray<TCHAR> Buffer;
+		TStringBuilder<512> StringBuilder;
+		for (int32 I = 0; I != SelectionCount; ++I)
+		{
+			int32 Index = SelectedItems[I];
+			int32 LogIndex = SendMessageW(LogHwnd, LB_GETITEMDATA, Index, 0) - LogIndexOffset;
+			StringBuilder.Append(Log[LogIndex].String);
+			StringBuilder.Append(TEXT("\r\n"));
+		}
+
+		FWindowsPlatformApplicationMisc::ClipboardCopy(*StringBuilder);
+	}
+
+	void UpdateIncludeFilter(const TCHAR* Str)
+	{
+		ConsoleIncludeFilterStr = Str;
+		SetDlgItemText(MainHwnd, ID_LOG_INCLUDE_FILTER, *ConsoleIncludeFilterStr);
+		ConsoleIncludeFilterStr.ParseIntoArray(IncludeFilter, TEXT(" "));
+	}
+
+	void UpdateExcludeFilter(const TCHAR* Str)
+	{
+		ConsoleExcludeFilterStr = Str;
+		SetDlgItemText(MainHwnd, ID_LOG_EXCLUDE_FILTER, *ConsoleExcludeFilterStr);
+		ConsoleExcludeFilterStr.ParseIntoArray(ExcludeFilter, TEXT(" "));
+	}
+
+	LRESULT LogHwndWndProc(HWND hWnd, UINT Msg, WPARAM wParam, LPARAM lParam)
 	{
 		switch (Msg)
 		{
-		case WM_NEWLOGENTRIES:
+		case WM_MOUSEWHEEL:
+			return HandleMouseWheel(wParam);
+
+		case WM_MOUSEMOVE:
 		{
-			SendMessageW(LogHwnd, WM_SETREDRAW, false, 0);
+			POINT Pos{ ((int)(short)LOWORD(lParam)), ((int)(short)HIWORD(lParam)) };
+			if (Pos.x != TipPosition.x || Pos.y != TipPosition.y)
 			{
+				SetWindowPos(TipHwnd, MainHwnd, 1, 1, 1, 1, SWP_NOACTIVATE | SWP_NOSENDCHANGING | SWP_NOZORDER | SWP_NOOWNERZORDER);
+				TRACKMOUSEEVENT tme;
+				tme.cbSize = sizeof(tme);
+				tme.dwFlags = TME_HOVER;
+				tme.hwndTrack = hWnd;
+				tme.dwHoverTime = HOVER_DEFAULT;
+				TrackMouseEvent(&tme);
+			}
+			break;
+		}
+
+		case WM_LBUTTONDOWN:
+			SuspendAddingEntries();
+			break;
+
+		case WM_LBUTTONUP:
+			ResumeAddingEntries();
+			break;
+
+		case WM_RBUTTONDOWN:
+		{
+			bAutoScrollLog = false;
+			SuspendAddingEntries();
+			RightClickedItem = SendMessageW(hWnd, LB_ITEMFROMPOINT, 0, lParam);
+
+			bool AlreadySelected = false;
+			TArray<int> SelectedItems;
+			GetSelectedItems(SelectedItems);
+			for (int SelectedItem : SelectedItems)
+				AlreadySelected |= RightClickedItem == SelectedItem;
+			if (AlreadySelected)
+				break;
+			ClearSelection();
+			SendMessageW(LogHwnd, LB_SETSEL, true, RightClickedItem);
+			break;
+		}
+
+		case WM_RBUTTONUP:
+		{
+			ResumeAddingEntries();
+
+			POINT MousePos{ ((int)(short)LOWORD(lParam)), ((int)(short)HIWORD(lParam)) };
+			TStringBuilder<1024> OutString;
+
+			int32 LogIndex = SendMessageW(LogHwnd, LB_GETITEMDATA, RightClickedItem, 0) - LogIndexOffset;
+			RightClickedItem = -1;
+			if (LogIndex == LB_ERR)
+			{
+				RECT Rect;
+				GetClientRect(hWnd, &Rect);
+				if (!PtInRect(&Rect, MousePos))
+					break;
+			}
+			else
+			{
+				CreateLogEntryText(OutString, Log[LogIndex], false);
+			}
+
+			HMENU Menu = CreatePopupMenu();
+			MENUITEMINFO Info;
+			Info.cbSize = sizeof(Info);
+			Info.dwItemData = 0;
+			auto AddItem = [&](const TCHAR* Str, int32 id)
+			{
+				Info.fMask = MIIM_STRING | MIIM_ID | MIIM_DATA;
+				Info.fType = MFT_STRING;
+				Info.dwTypeData = (TCHAR*)Str;
+				Info.wID = id;
+				InsertMenuItem(Menu, Info.dwItemData++, 1, &Info);
+			};
+
+			auto AddSeparator = [&]()
+			{
+				Info.fMask = 0;
+				Info.fType = MFT_SEPARATOR;
+				InsertMenuItem(Menu, Info.dwItemData++, 1, &Info);
+			};
+
+			AddItem(TEXT("Activate auto scroll"), 101);
+			AddItem(TEXT("Copy selected line(s)"), 102);
+			AddSeparator();
+			AddItem(TEXT("Clear all"), 103);
+			AddItem(TEXT("Clear all entries above line"), 104);
+			AddItem(TEXT("Clear filters"), 121);
+
+			SelectedCategory = Log[LogIndex].Category.ToString();
+
+			AddSeparator();
+			AddItem(*TStringBuilder<64>().Append(TEXT("Include category '")).Append(*SelectedCategory).Append(TEXT("' in filter")), 124);
+			AddItem(*TStringBuilder<64>().Append(TEXT("Exclude category '")).Append(*SelectedCategory).Append(TEXT("' in filter")), 125);
+
+
+			int CharIndex = (MousePos.x + (LogFontWidth / 2)) / LogFontWidth;
+			if (CharIndex < OutString.Len())
+			{
+				auto ValidChar = [](TCHAR C) { return FChar::IsAlnum(C) || C == '_' || C == '-'; };
+
+				const TCHAR* Str = *OutString;
+				const TCHAR* Start = Str + CharIndex;
+				while (Start != Str)
 				{
-					FScopeLock _(&NewLogEntriesCs);
-					Swap(NewLogEntries, TempLogEntries);
+					TCHAR C = *(Start - 1);
+					if (!ValidChar(C))
+						break;
+					--Start;
 				}
+				const TCHAR* End = Str + CharIndex;
+				while (true)
+ 				{
+					if (!ValidChar(*End))
+						break;
+					++End;
+				}
+				if (Start != End)
+				{
+					SelectedWord = FStringView(Start, End - Start);
 
-				int32 OldSize = Log.Num();
+					AddSeparator();
+					AddItem(*TStringBuilder<64>().Append("Copy word '").Append(SelectedWord).Append("'"), 120);
+					AddSeparator();
+					AddItem(*TStringBuilder<64>().Append("Include word '").Append(SelectedWord).Append("' in filter"), 122);
+					AddItem(*TStringBuilder<64>().Append("Exclude word '").Append(SelectedWord).Append("' in filter"), 123);
+				}
+			}
+
+			MapWindowPoints(hWnd, HWND_DESKTOP, &MousePos, 1);
+			TrackPopupMenu(Menu, 0, MousePos.x, MousePos.y, 0, hWnd, NULL);
+			break;
+		}
+
+		case WM_COMMAND:
+		{
+			switch (LOWORD(wParam))
+			{
+			case 101:
+				EnableAutoScroll();
+				break;
 
-				for (NewLogEntry& E : TempLogEntries)
+			case 102:
+				CopySelectionToClipboard();
+				break;
+
+			case 103:
+				ClearLog();
+				break;
+
+			case 104:
+			{
+				int32 SelectedItemIndex = -1;
+				if (SendMessageW(LogHwnd, LB_GETSELITEMS, 1, (LPARAM)&SelectedItemIndex) != 1)
+					break;
+				int32 LogIndex = SendMessageW(LogHwnd, LB_GETITEMDATA, SelectedItemIndex, 0) - LogIndexOffset;
+				LogIndexOffset += LogIndex;
+				Log.PopFront(LogIndex);
+				RefreshLogHwnd();
+				break;
+			}
+			case 120:
+				FWindowsPlatformApplicationMisc::ClipboardCopy(*SelectedWord);
+				break;
+
+			case 121:
+				UpdateIncludeFilter(TEXT(""));
+				UpdateExcludeFilter(TEXT(""));
+				RefreshLogHwnd();
+				break;
+
+			case 122:
+				UpdateIncludeFilter(*(ConsoleIncludeFilterStr.IsEmpty() ? SelectedWord : (ConsoleIncludeFilterStr + " " + SelectedWord)));
+				RefreshLogHwnd();
+				break;
+
+			case 123:
+				UpdateExcludeFilter(*(ConsoleExcludeFilterStr.IsEmpty() ? SelectedWord : (ConsoleExcludeFilterStr + " " + SelectedWord)));
+				RefreshLogHwnd();
+				break;
+
+			case 124:
+			case 125:
+			{
+				TStringBuilder<128> Str;
+				Str << (LOWORD(wParam) == 124 ? *ConsoleIncludeFilterStr : *ConsoleExcludeFilterStr);
+				if (Str.Len() != 0)
+					Str << ' ';
+				Str << '[' << SelectedCategory << ']';
+				if (LOWORD(wParam) == 124)
+					UpdateIncludeFilter(*Str);
+				else
+					UpdateExcludeFilter(*Str);
+				RefreshLogHwnd();
+				break;
+			}
+			}
+			break;
+		}
+
+		case WM_KEYDOWN:
+			if (wParam == VK_SHIFT)
+				SuspendAddingEntries();
+			break;
+
+		case WM_KEYUP:
+			if (wParam == VK_SHIFT)
+				ResumeAddingEntries();
+			break;
+
+		case WM_KILLFOCUS:
+			ResumeAddingEntries();
+			break;
+
+		case WM_MOUSEHOVER:
+			{
+				POINT MousePos{ ((int)(short)LOWORD(lParam)), ((int)(short)HIWORD(lParam)) };
+				TipPosition = MousePos;
+				int32 Res = SendMessageW(hWnd, LB_ITEMFROMPOINT, 0, lParam);
+				int32 ItemIndex = LOWORD(Res);
+				int32 LogIndex = SendMessageW(hWnd, LB_GETITEMDATA, ItemIndex, 0) - LogIndexOffset;
+				if (LogIndex <= 0 || LogIndex >= Log.Num())
+					break;
+				TStringBuilder<1024> OutString;
+				CreateLogEntryText(OutString, Log[LogIndex], false);
+				int32 EntryStrLen = OutString.Len();
+				RECT LogRect;
+				GetClientRect(LogHwnd, &LogRect);
+				int32 VisibleCharCount = LogRect.right / LogFontWidth;
+				if (EntryStrLen <= VisibleCharCount)
+					break;
+				const TCHAR* Str = *OutString + VisibleCharCount;
+				int32 ExtraLen = FCString::Strlen(Str) * LogFontWidth + 3;
+				HDC dc = GetDC(TipHwnd);
+				SelectObject(dc, LogFont);
+				ReleaseDC(TipHwnd, dc);
+				POINT ExtraPos{VisibleCharCount*LogFontWidth,(MousePos.y/LogFontHeight)*LogFontHeight};
+				MapWindowPoints(LogHwnd, HWND_DESKTOP, &ExtraPos, 1);
+				SetWindowPos(TipHwnd, MainHwnd, ExtraPos.x, ExtraPos.y-1, ExtraLen, LogFontHeight+2, SWP_NOACTIVATE | SWP_NOSENDCHANGING | SWP_NOZORDER | SWP_NOOWNERZORDER);
+				SetWindowText(TipHwnd, Str);
+				RedrawWindow(hWnd, NULL, NULL, RDW_INVALIDATE | RDW_UPDATENOW);
+				break;
+			}
+		}
+
+		return LogHwndWndProcPtr(hWnd, Msg, wParam, lParam);
+	}
+
+	static LRESULT CALLBACK StaticMessageBoxHookProc(int nCode, WPARAM wParam, LPARAM lParam)
+	{
+		if (nCode == HCBT_ACTIVATE)
+		{
+			HWND hwnd = (HWND)wParam;
+			TCHAR Title[128];
+			GetWindowText(hwnd, Title, 64);
+			if (FCString::Strncmp(Title, HelpTitle, 63) == 0)
+			{
+				HWND parentHwnd = GetParent(hwnd);
+				int32 Sx, Sy, Px, Py;
+				RECT R1;
+				RECT R2;
+				GetWindowRect(parentHwnd, &R1); GetWindowRect(hwnd, &R2);
+				Sx = R2.right - R2.left, Px = R1.left + (R1.right - R1.left) / 2 - Sx / 2;
+				Sy = R2.bottom - R2.top, Py = R1.top + (R1.bottom - R1.top) / 2 - Sy / 2;
+				MoveWindow(hwnd, Px, Py, Sx, Sy, 1);
+			}
+		}
+		return CallNextHookEx(0, nCode, wParam, lParam);
+	}
+
+	LRESULT MainWinProc(HWND hWnd, UINT Msg, WPARAM wParam, LPARAM lParam)
+	{
+		switch (Msg)
+		{
+		case WM_NEWLOGENTRIES:
+		{
+			if (bSuspendAddingEntries)
+				break;
+
+			{
+				FScopeLock _(&NewLogEntriesCs);
+				Swap(NewLogEntries, TempLogEntries);
+			}
+
+			if (TempLogEntries.IsEmpty())
+				break;
+
+			int32 OldSize = Log.Num();
+
+			for (NewLogEntry& E : TempLogEntries)
+			{
+				const TCHAR* SearchStr = *E.String;
+				while (true)
 				{
-					const TCHAR* SearchStr = *E.String;
-					while (true)
+					if (const TCHAR* LineBreak = FCString::Strchr(SearchStr, '\n'))
 					{
-						if (const TCHAR* LineBreak = FCString::Strchr(SearchStr, '\n'))
-						{
-							int32 Len = LineBreak - SearchStr;
-							if (LineBreak > SearchStr && *(LineBreak - 1) == '\r')
-								--Len;
-							Log.Add({ FString(Len, SearchStr), E.TextAttribute, 1 });
-							SearchStr = LineBreak + 1;
-						}
-						else
-						{
-							if (int32 Len = FCString::Strlen(SearchStr))
-								Log.Add({ SearchStr, E.TextAttribute, 1 });
-							break;
-						}
+						int32 Len = LineBreak - SearchStr;
+						if (LineBreak > SearchStr && *(LineBreak - 1) == '\r')
+							--Len;
+						Log.Add({ FString(Len, SearchStr), E.Verbosity, E.Category, E.Time, E.TextAttribute, 1 });
+						SearchStr = LineBreak + 1;
+					}
+					else
+					{
+						if (int32 Len = FCString::Strlen(SearchStr))
+							Log.Add({ SearchStr, E.Verbosity, E.Category, E.Time, E.TextAttribute, 1 });
+						break;
 					}
 				}
+			}
 
-				constexpr int32 MaxSize = 20000;
-
-				int32 NewSize = Log.Num();
+			constexpr int32 MaxSize = 20000;
 
-				int32 ToAddToHwnd = NewSize - OldSize;
-				int32 LogIndex = OldSize;
+			int32 NewSize = Log.Num();
 
-				if (NewSize > MaxSize)
-				{
-					int32 ToRemove = NewSize - MaxSize;
-					Log.PopFront(ToRemove);
-					LogIndexOffset += ToRemove;
+			int32 ToAddToHwnd = NewSize - OldSize;
+			int32 LogIndex = OldSize;
+			int32 NewTopIndex = -1;
+			int32 NewCaretIndex = -1;
 
-					int32 ToChangeInHwnd = FMath::Min(MaxSize, ToRemove);
+			SendMessageW(LogHwnd, WM_SETREDRAW, false, 0);
 
-					int32 ToRemoveFromHwnd = ToChangeInHwnd;
-					while (ToRemoveFromHwnd--)
-						SendMessageW(LogHwnd, LB_DELETESTRING, 0, 0);
+			if (NewSize > MaxSize)
+			{
+				int32 ToRemove = NewSize - MaxSize;
+				Log.PopFront(ToRemove);
+				LogIndexOffset += ToRemove;
 
-					ToAddToHwnd = ToChangeInHwnd;
-					LogIndex = MaxSize - ToAddToHwnd;
-				}
+				int32 ToChangeInHwnd = FMath::Min(MaxSize, ToRemove);
 
-				while (ToAddToHwnd--)
+				int32 ToRemoveFromHwnd = ToChangeInHwnd;
+				if (!bAutoScrollLog)
 				{
-					AddEntryToLogHwnd(Log[LogIndex], LogIndexOffset + LogIndex);
-					++LogIndex;
+					int32 TopIndex = SendMessageW(LogHwnd, LB_GETTOPINDEX, 0, 0);
+					NewTopIndex = FMath::Max(TopIndex - ToRemoveFromHwnd, 0);
+					int32 CaretIndex = SendMessageW(LogHwnd, LB_GETCARETINDEX, 0, 0);
+					NewCaretIndex = FMath::Max(CaretIndex - ToRemoveFromHwnd, 0);
 				}
-				TempLogEntries.Reset(0);
+
+				while (ToRemoveFromHwnd--)
+					SendMessageW(LogHwnd, LB_DELETESTRING, 0, 0);
+
+				ToAddToHwnd = ToChangeInHwnd;
+				LogIndex = MaxSize - ToAddToHwnd;
 			}
 
+			while (ToAddToHwnd--)
+			{
+				AddEntryToLogHwnd(Log[LogIndex], LogIndexOffset + LogIndex);
+				++LogIndex;
+			}
+			TempLogEntries.Reset(0);
+
 			if (bAutoScrollLog)
 			{
 				ScrollToBottom();
 			}
+			else if (NewTopIndex != -1)
+			{
+				SendMessageW(LogHwnd, LB_SETCARETINDEX, NewCaretIndex, 0);
+				SendMessageW(LogHwnd, LB_SETSEL, true, NewCaretIndex);
+				SetTopVisible(NewTopIndex, false);
+			}
 			RedrawLogScrollbar();
 
 			SendMessageW(LogHwnd, WM_SETREDRAW, true, 0);
@@ -1143,10 +1628,16 @@ public:
 			break;
 		}
 
+		case WM_HELP:
+			{
+				HHOOK Hook = SetWindowsHookEx(WH_CBT, StaticMessageBoxHookProc, NULL, GetCurrentThreadId());
+				MessageBox(MainHwnd, HelpText, HelpTitle, MB_OK | MB_ICONINFORMATION);
+				UnhookWindowsHookEx(Hook);
+		}
+			break;
+
 		case WM_SIZE:
 		{
-			if (hWnd != MainHwnd)
-				break;
 			UpdateSize(LOWORD(lParam), HIWORD(lParam), true);
 			RECT WindowRect;
 			GetWindowRect(MainHwnd, &WindowRect);
@@ -1156,8 +1647,6 @@ public:
 		}
 		case WM_MOVE:
 		{
-			if (hWnd != MainHwnd)
-				break;
 			RECT WindowRect;
 			GetWindowRect(MainHwnd, &WindowRect);
 			ConsolePosX = WindowRect.left;
@@ -1166,8 +1655,6 @@ public:
 		}
 
 		case WM_SETFOCUS:
-			if (hWnd != MainHwnd)
-				break;
 			SetFocus(CommandHwnd);
 			SendMessageW(CommandHwnd, EM_SETSEL, 0, -1);
 			break;
@@ -1200,18 +1687,27 @@ public:
 
 		case WM_CTLCOLORSTATIC:
 		{
-			HDC hdcStatic = (HDC)wParam;
-			SetTextColor(hdcStatic, TextColor);
-			SetBkColor(hdcStatic, BackgroundColor);
+			HDC Hdc = (HDC)wParam;
+			SetTextColor(Hdc, TextColor);
+			SetBkColor(Hdc, BackgroundColor);
 			return (INT_PTR)BackgroundBrush;
 		}
 
 		case WM_CTLCOLOREDIT:
 		{
-			HDC hdcStatic = (HDC)wParam;
-			SetTextColor(hdcStatic, TextColor);
-			SetBkColor(hdcStatic, EditBackgroundColor);
-			return (INT_PTR)EditBackgroundBrush;
+			HWND filterHwnd = (HWND)lParam;
+			HDC Hdc = (HDC)wParam;
+			SetBkColor(Hdc, EditBackgroundColor);
+			RECT Rect;
+			GetClientRect(filterHwnd, &Rect);
+			FillRect(Hdc, &Rect, EditBackgroundBrush);
+			SetTextColor(Hdc, TextColor);
+			if (!GetWindowTextLength(filterHwnd))
+			{
+				const TCHAR* Str = filterHwnd == CommandHwnd ? TEXT("Type command here") : (filterHwnd == IncludeFilterHwnd ? TEXT("Add include filter here") : TEXT("Add exclude filter here"));
+				TextOut(Hdc, 0, 0, Str, FCString::Strlen(Str));
+			}
+			return (INT_PTR)GetStockObject(NULL_BRUSH);
 		}
 
 		case WM_CTLCOLORLISTBOX:
@@ -1220,18 +1716,14 @@ public:
 				int32 ItemCount = (int)SendMessageW(LogHwnd, LB_GETCOUNT, 0, 0);
 				RECT Rect;
 				GetClientRect(LogHwnd, &Rect);
-				if (ItemCount * LogFontHeight > Rect.bottom)
+				int32 VisibleCount = FMath::CeilToInt(float(Rect.bottom) / LogFontHeight);
+				int32 TopIndex = SendMessageW(LogHwnd, LB_GETTOPINDEX, 0, 0);
+				if (ItemCount - TopIndex < VisibleCount)
 				{
-					// Draw the padding below the last line
-					if (int32 padding = Rect.bottom % LogFontHeight)
-					{
-						Rect.top = Rect.bottom - padding;
-						FillRect((HDC)wParam, &Rect, LogBackgroundBrush);
-					}
-					return (INT_PTR)GetStockObject(NULL_BRUSH);
+					Rect.top = (ItemCount - TopIndex) * LogFontHeight;
+					FillRect((HDC)wParam, &Rect, LogBackgroundBrush); // Draw the padding below the last line
 				}
-				else
-					return (INT_PTR)LogBackgroundBrush;
+				return (INT_PTR)GetStockObject(NULL_BRUSH);
 			}
 			break;
 
@@ -1286,45 +1778,59 @@ public:
 		case WM_DRAWITEM:
 		{
 			PDRAWITEMSTRUCT pdis = (PDRAWITEMSTRUCT)lParam;
+			HDC Hdc = pdis->hDC;
 
 			if (pdis->CtlType == ODT_STATIC)
 			{
 				TCHAR Str[512];
 				int32 StrLen = GetWindowText(pdis->hwndItem, Str, UE_ARRAY_COUNT(Str));
 
+				if (pdis->hwndItem == TipHwnd)
+				{
+					if (pdis->rcItem.bottom != 1)
+					{
+						FillRect(Hdc, &pdis->rcItem, BackgroundBrush);
+						SetTextColor(Hdc, GetSysColor(COLOR_HIGHLIGHTTEXT));
+						SelectObject(Hdc, LogFont);
+						DrawTextW(Hdc, Str, StrLen, &pdis->rcItem, DT_SINGLELINE | DT_VCENTER);
+					}
+					break;
+				}
+
+
 				HBRUSH Brush = BackgroundBrush;
 
 				int32 RowIndex = GetWindowLong(pdis->hwndItem, GWLP_USERDATA);
 				if (RowIndex != -1)
 				{
 					Brush = StatusBackgroundBrush[(RowIndex & 1)];
-					SetBkColor(pdis->hDC, StatusBackgroundColor[(RowIndex & 1)]);
+					SetBkColor(Hdc, StatusBackgroundColor[(RowIndex & 1)]);
 				}
 
-				SelectObject(pdis->hDC, Brush);
+				SelectObject(Hdc, Brush);
 
 				// All this stuff just to prevent flickering when text changes fast
 				int32 TextOffset = 1;
 				RECT ModifiedRect = pdis->rcItem;
-				DrawTextW(pdis->hDC, Str, StrLen, &ModifiedRect, DT_SINGLELINE | DT_END_ELLIPSIS | DT_CALCRECT | DT_MODIFYSTRING);
+				DrawTextW(Hdc, Str, StrLen, &ModifiedRect, DT_SINGLELINE | DT_END_ELLIPSIS | DT_CALCRECT | DT_MODIFYSTRING);
 				RECT TempRect = pdis->rcItem;
 				TempRect.bottom = TempRect.top + TextOffset;
 				TempRect.right = ModifiedRect.right;
-				FillRect(pdis->hDC, &TempRect, Brush);
+				FillRect(Hdc, &TempRect, Brush);
 
 				TempRect = pdis->rcItem;
 				TempRect.left = ModifiedRect.right;
-				FillRect(pdis->hDC, &TempRect, Brush);
+				FillRect(Hdc, &TempRect, Brush);
 				if (ModifiedRect.bottom < pdis->rcItem.bottom)
 				{
 					TempRect = pdis->rcItem;
 					TempRect.right = ModifiedRect.right;
 					TempRect.top = ModifiedRect.bottom;
-					FillRect(pdis->hDC, &TempRect, Brush);
+					FillRect(Hdc, &TempRect, Brush);
 				}
 
 				StrLen = FCString::Strlen(Str);
-				ExtTextOut(pdis->hDC, pdis->rcItem.left, pdis->rcItem.top + TextOffset, 0, &pdis->rcItem, Str, StrLen, 0);
+				ExtTextOut(Hdc, pdis->rcItem.left, pdis->rcItem.top + TextOffset, 0, &pdis->rcItem, Str, StrLen, 0);
 				break;
 			}
 
@@ -1344,44 +1850,24 @@ public:
 
 				if (pdis->itemState & ODS_SELECTED)
 				{
-					FillRect(pdis->hDC, &pdis->rcItem, GetSysColorBrush(COLOR_HIGHLIGHT));
-					SetTextColor(pdis->hDC, GetSysColor(COLOR_HIGHLIGHTTEXT));
+					FillRect(Hdc, &pdis->rcItem, GetSysColorBrush(COLOR_HIGHLIGHT));
+					SetTextColor(Hdc, GetSysColor(COLOR_HIGHLIGHTTEXT));
 				}
 				else
 				{
-					//if (pdis->itemAction != ODA_DRAWENTIRE)
-					{
-						FillRect(pdis->hDC, &pdis->rcItem, LogBackgroundBrush);
-					}
-
+					FillRect(Hdc, &pdis->rcItem, LogBackgroundBrush);
 					if (Entry.TextAttribute == (FOREGROUND_INTENSITY | FOREGROUND_RED))
-						SetTextColor(pdis->hDC, RGB(220, 0, 0));
+						SetTextColor(Hdc, RGB(220, 0, 0));
 					else if (Entry.TextAttribute == (FOREGROUND_INTENSITY | FOREGROUND_RED | FOREGROUND_GREEN))
-						SetTextColor(pdis->hDC, RGB(220, 220, 0));
+						SetTextColor(Hdc, RGB(220, 220, 0));
 					else
-						SetTextColor(pdis->hDC, RGB(204, 204, 204));
+						SetTextColor(Hdc, RGB(204, 204, 204));
 				}
 
+				TStringBuilder<1024> OutString;
+				CreateLogEntryText(OutString, Entry, false);
 				int32 XPos = 0;
-				const TCHAR* Str = *Entry.String;
-				while (true)
-				{
-					if (const TCHAR* Tab = FCString::Strchr(Str, '\t'))
-					{
-						if (int32 PartLen = int32(Tab - Str) - 1)
-						{
-							TextOut(pdis->hDC, XPos, YPos, Str, PartLen);
-							XPos += PartLen * LogFontWidth;
-						}
-						XPos += LogFontWidth * 3;
-						Str = Tab + 1;
-						continue;
-					}
-
-					TextOut(pdis->hDC, XPos, YPos, Str, FCString::Strlen(Str));
-					break;
-
-				}
+				TextOut(Hdc, XPos, YPos, *OutString, OutString.Len());
 				return true;
 			}
 			case ODA_FOCUS: // We don't want the focus rectangle at all (dotted rectangle)
@@ -1392,8 +1878,6 @@ public:
 
 		case WM_PAINT:
 		{
-			if (hWnd != MainHwnd)
-				break;
 			PAINTSTRUCT ps;
 			HDC Hdc = BeginPaint(hWnd, &ps);
 			RECT Rect;
@@ -1424,39 +1908,25 @@ public:
 			{
 				if ((LOWORD(wParam) == 0x43 || LOWORD(wParam) == VK_INSERT) && (::GetKeyState(VK_CONTROL) >> 15)) // Copy-paste
 				{
-					if (int32 SelectionCount = SendMessageW(LogHwnd, LB_GETSELCOUNT, 0, 0))
+					TArray<int> SelectedItems;
+					if (int32 SelectionCount = GetSelectedItems(SelectedItems))
 					{
-						TArray<int> SelectedItems;
-						SelectedItems.SetNum(SelectionCount);
-						SelectionCount = SendMessageW(LogHwnd, LB_GETSELITEMS, SelectionCount, (LPARAM)SelectedItems.GetData());
 						TArray<TCHAR> Buffer;
-
 						TStringBuilder<512> StringBuilder;
 						for (int32 I = 0; I != SelectionCount; ++I)
 						{
-							if (I > 0)
-								StringBuilder.Append(TEXT("\r\n"));
 							int32 Index = SelectedItems[I];
-							int32 Len = SendMessageW(LogHwnd, LB_GETTEXTLEN, Index, 0);
-							Buffer.SetNum(Len + 1, false);
-							SendMessageW(LogHwnd, LB_GETTEXT, Index, (LPARAM)Buffer.GetData());
-							StringBuilder.Append(Buffer.GetData(), Buffer.Num() - 1);
+							int32 LogIndex = SendMessageW(LogHwnd, LB_GETITEMDATA, Index, 0) - LogIndexOffset;
+							StringBuilder.Append(Log[LogIndex].String);
+							StringBuilder.Append(TEXT("\r\n"));
 						}
 
 						FWindowsPlatformApplicationMisc::ClipboardCopy(*StringBuilder);
 					}
 				}
-				else if (LOWORD(wParam) == VK_END) // Enable auto scrolling and remove selection
+				else if (LOWORD(wParam) == VK_END && !(::GetKeyState(VK_SHIFT) >> 15)) // Enable auto scrolling and remove selection
 				{
-					bAutoScrollLog = true;
-					enum { MaxSelectionCount = 32 * 1024 - 1 };
-					int32 SelectedItems[MaxSelectionCount + 1];
-					int32 SelectionCount = SendMessageW(LogHwnd, LB_GETSELITEMS, MaxSelectionCount, (LPARAM)SelectedItems);
-					for (int32 I = 0; I != SelectionCount; ++I)
-						SendMessageW(LogHwnd, LB_SETSEL, false, SelectedItems[I]);
-					ScrollToBottom();
-					RedrawWindow(LogHwnd, NULL, NULL, RDW_ERASE|RDW_INVALIDATE);
-					RedrawLogScrollbar();
+					EnableAutoScroll();
 					return -2;
 				}
 			}
@@ -1467,14 +1937,17 @@ public:
 			{
 				if (SendMessageW(LogHwnd, LB_GETSELCOUNT, 0, 0) != 0)
 					bAutoScrollLog = false;
+				RedrawLogScrollbar();
 			}
 			else if ((LOWORD(wParam) == ID_LOG_INCLUDE_FILTER || LOWORD(wParam) == ID_LOG_EXCLUDE_FILTER) && HIWORD(wParam) == EN_CHANGE)
 			{
 				TCHAR Str[1024];
 				GetDlgItemText(MainHwnd, LOWORD(wParam), Str, 1024);
+				FString& FilterStr = LOWORD(wParam) == ID_LOG_INCLUDE_FILTER ? ConsoleIncludeFilterStr : ConsoleExcludeFilterStr;
 				TArray<FString>& Filter = LOWORD(wParam) == ID_LOG_INCLUDE_FILTER ? IncludeFilter : ExcludeFilter;
+				FilterStr = Str;
 				TArray<FString> ParsedFilter;
-				FString(Str).ParseIntoArray(ParsedFilter, TEXT(" "));
+				FilterStr.ParseIntoArray(ParsedFilter, TEXT(" "));
 				if (Filter == ParsedFilter)
 					break;
 				Filter = ParsedFilter;
@@ -1487,6 +1960,7 @@ public:
 					break;
 				SetWindowText(CommandHwnd, TEXT(""));
 				EConsoleColor NewConsoleColor = EConsoleColor(-1);
+				int32 OldConsoleShowDateTime = ConsoleShowDateTime;
 				if (FCString::Stricmp(Command, TEXT("red")) == 0)
 					NewConsoleColor = EConsoleColor::Red;
 				else if (FCString::Stricmp(Command, TEXT("darkgray")) == 0)
@@ -1499,6 +1973,15 @@ public:
 					NewConsoleColor = EConsoleColor::DarkBlue;
 				else if (FCString::Stricmp(Command, TEXT("load")) == 0)
 					NewConsoleColor = EConsoleColor::Load;
+				else if (FCString::Stricmp(Command, TEXT("showtime")) == 0)
+					ConsoleShowDateTime = 1;
+				else if (FCString::Stricmp(Command, TEXT("showdatetime")) == 0)
+					ConsoleShowDateTime = 2;
+				else if (FCString::Stricmp(Command, TEXT("hidedatetime")) == 0)
+					ConsoleShowDateTime = 0;
+
+				if (ConsoleShowDateTime != OldConsoleShowDateTime)
+					RefreshLogHwnd();
 
 				if (NewConsoleColor != EConsoleColor(-1))
 				{
@@ -1526,20 +2009,16 @@ public:
 			}
 			else if (LOWORD(wParam) == ID_CLEARLOGBUTTON)
 			{
-				LogIndexOffset = 0;
-				Log.Reset();
-				RefreshLogHwnd();
-				CheckpointIndex = 0;
-				SetDlgItemText(MainHwnd, ID_ADDCHECKPOINTBUTTON, TEXT("Log CHECKPOINT0"));
+				ClearLog();
 			}
 			else if (LOWORD(wParam) == ID_ADDCHECKPOINTBUTTON)
 			{
-				TStringBuilder<MAX_SPRINTF> TempString;
+				TStringBuilder<64> TempString;
 				TempString.Appendf(TEXT("LOGCHECKPOINT%i"), CheckpointIndex);
-				AddLogEntry(TempString, FOREGROUND_BLUE | FOREGROUND_GREEN | FOREGROUND_RED);
+				AddLogEntry(TempString, ELogVerbosity::Display, FName(TEXT("LogCheckpoint")), 0, FOREGROUND_BLUE | FOREGROUND_GREEN | FOREGROUND_RED);
 
 				++CheckpointIndex;
-				TCHAR ButtonString[MAX_SPRINTF] = TEXT("");
+				TCHAR ButtonString[64] = TEXT("");
 				FCString::Sprintf(ButtonString, TEXT("Log CHECKPOINT%i"), CheckpointIndex);
 				SetDlgItemText(MainHwnd, ID_ADDCHECKPOINTBUTTON, ButtonString);
 			}
@@ -1626,6 +2105,8 @@ public:
 				GetClientRect(MainHwnd, &Rect);
 				UpdateSize(Rect.right, Rect.bottom, true);
 				RedrawWindow(MainHwnd, NULL, NULL, RDW_INVALIDATE);
+				if (bConsoleExpanded)
+					SetFocus(CommandHwnd);
 			}
 			else if (Res == HTMINBUTTON)
 				ShowWindow(hWnd, SW_MINIMIZE);
@@ -1766,7 +2247,7 @@ public:
 		}
 		case WM_NCACTIVATE:
 			RedrawWindow(hWnd, NULL, NULL, RDW_UPDATENOW);
-			return 0;
+			return 1;
 		}
 		return DefWindowProc(hWnd, Msg, wParam, lParam);
 	}
@@ -1783,8 +2264,8 @@ public:
 
 	FWindowsConsoleOutputDevice2* Owner;
 
-	struct NewLogEntry { FString String; uint16 TextAttribute; };
-	struct LogEntry { FString String; uint16 TextAttribute; uint16 LineCount; };
+	struct NewLogEntry { FString String; ELogVerbosity::Type Verbosity; const class FName Category; double Time; uint16 TextAttribute; };
+	struct LogEntry { FString String; ELogVerbosity::Type Verbosity; const class FName Category; double Time; uint16 TextAttribute; uint16 LineCount; };
 	FCriticalSection NewLogEntriesCs;
 	TArray<NewLogEntry> NewLogEntries;
 	TArray<NewLogEntry> TempLogEntries;
@@ -1795,13 +2276,19 @@ public:
 	int32 AddedEntryLogVirtualIndex = -1;
 	HICON Icon;
 	HFONT Font;
+	HFONT LogFont;
 	HWND MainHwnd;
 	HWND LogHwnd;
 	HWND LogScrollHwnd;
 	HWND CommandHwnd;
 	HWND RunCommandHwnd;
-	HWND FilterHwnd[7];
-	WNDPROC LogHwndWndProc;
+	HWND IncludeFilterHwnd;
+	HWND ExcludeFilterHwnd;
+	HWND ClearLogButtonHwnd;
+	HWND CheckpointButtonHwnd;
+	HWND TipHwnd;
+	POINT TipPosition;
+	WNDPROC LogHwndWndProcPtr;
 	COLORREF TextColor;
 	COLORREF BackgroundColor;
 	COLORREF EditBackgroundColor;
@@ -1827,11 +2314,14 @@ public:
 	HPEN TextPen;
 	HPEN NoPen;
 	int32 LogIndexOffset = 0;
-	int32 LogFontWidth = 8;
-	int32 LogFontHeight = 14;
+	int32 LogFontWidth = 0;
+	int32 LogFontHeight = 0;
 	int32 LogScrollGrabPos = -1;
 	int32 NcButtonHot = -1;
 	int32 NcButtonDown = -1;
+	int32 RightClickedItem = -1;
+	FString SelectedWord;
+	FString SelectedCategory;
 
 	struct Activity { HWND NameHwnd = 0; HWND StatusHwnd = 0; FString Name; FString Status; int32 Light = 0; bool bStatusDirty = false; bool bAlignLeft = false; };
 	FCriticalSection ActivitiesCs;
@@ -1842,6 +2332,8 @@ public:
 	
 	HANDLE Thread;
 	bool bAutoScrollLog = true;
+	bool bSuspendAddingEntries = false;
+
 	int32 CheckpointIndex = 0;
 
 	bool bIsVisible = false;
@@ -1850,6 +2342,10 @@ public:
 	int32 ConsolePosX = 0;
 	int32 ConsolePosY = 0;
 	EConsoleColor ConsoleColor = EConsoleColor::DarkGray;
+	int32 ConsoleShowDateTime = 0; // 0 show nothing, 1 show only time, 2 show date and time
+	FString ConsoleIncludeFilterStr;
+	FString ConsoleExcludeFilterStr;
+
 	FString ConsoleTitle;
 	bool bConsoleExpanded = false;
 };
@@ -2005,6 +2501,9 @@ void FWindowsConsoleOutputDevice2::SaveToINI()
 	int32 ConsolePosX = Window->ConsolePosX;
 	int32 ConsolePosY = Window->ConsolePosY;
 	int32 ConsoleColor = (int32)Window->ConsoleColor;
+	int32 ConsoleShowDateTime = Window->ConsoleShowDateTime;
+	FString ConsoleIncludeFilterStr = Window->ConsoleIncludeFilterStr;
+	FString ConsoleExcludeFilterStr = Window->ConsoleExcludeFilterStr;
 	bool bConsoleExpanded = Window->bConsoleExpanded;
 	WindowRWLock.WriteUnlock();
 
@@ -2028,6 +2527,11 @@ void FWindowsConsoleOutputDevice2::SaveToINI()
 
 	Config.SetBool(Selection, TEXT("ConsoleExpanded"), bConsoleExpanded, Filename);
 
+	Config.SetInt(Selection, TEXT("ConsoleShowDateTime"), ConsoleShowDateTime, Filename);
+
+	Config.SetString(Selection, TEXT("IncludeFilter"), *ConsoleIncludeFilterStr, Filename);
+	Config.SetString(Selection, TEXT("ExcludeFilter"), *ConsoleExcludeFilterStr, Filename);
+
 	Config.Flush(false, Filename);
 }
 
@@ -2042,6 +2546,10 @@ void FWindowsConsoleOutputDevice2::Show( bool ShowWindow )
 		int32 ConsolePosX = 0;
 		int32 ConsolePosY = 0;
 		int32 ConsoleColor = 0;
+		int32 ConsoleShowDateTime = 0;
+		FString ConsoleIncludeFilterStr;
+		FString ConsoleExcludeFilterStr;
+
 		bool bConsoleExpanded = false;
 		bool bHasX = false;
 		bool bHasY = false;
@@ -2061,6 +2569,9 @@ void FWindowsConsoleOutputDevice2::Show( bool ShowWindow )
 		bHasY = Config.GetInt(Selection, TEXT("ConsoleY"), ConsolePosY, Filename);
 		Config.GetInt(Selection, TEXT("ConsoleColor"), ConsoleColor, Filename);
 		Config.GetBool(Selection, TEXT("ConsoleExpanded"), bConsoleExpanded, Filename);
+		Config.GetInt(Selection, TEXT("ConsoleShowDateTime"), ConsoleShowDateTime, Filename);
+		Config.GetString(Selection, TEXT("IncludeFilter"), ConsoleIncludeFilterStr, Filename);
+		Config.GetString(Selection, TEXT("ExcludeFilter"), ConsoleExcludeFilterStr, Filename);
 
 		if (!FParse::Value(FCommandLine::Get(), TEXT("ConsoleX="), ConsolePosX) && !bHasX)
 		{
@@ -2082,7 +2593,7 @@ void FWindowsConsoleOutputDevice2::Show( bool ShowWindow )
 				ConsoleType = TEXT("Server");
 			else if (IsRunningGame())
 				ConsoleType = TEXT("Client");
-			ConsoleTitle = FString::Printf(TEXT("%s Console (%s) - %s"), ConsoleType, FApp::GetProjectName(), FPlatformProcess::ExecutablePath());
+			ConsoleTitle = FString::Printf(TEXT("%s Console (%s) - %s  (F1 for help)"), ConsoleType, FApp::GetProjectName(), FPlatformProcess::ExecutablePath());
 		}
 
 		FDisplayMetrics DisplayMetrics;
@@ -2098,8 +2609,10 @@ void FWindowsConsoleOutputDevice2::Show( bool ShowWindow )
 		static const int32 RightPadding = FMath::Max(50, FMath::Min((ActualConsoleWidth / 2), ActualScreenWidth / 2));
 		static const int32 BottomPadding = FMath::Max(50, FMath::Min((ActualConsoleHeight / 2), ActualScreenHeight / 2));
 				
-		ConsolePosX = FMath::Min(FMath::Max(ConsolePosX, DisplayMetrics.VirtualDisplayRect.Left), DisplayMetrics.VirtualDisplayRect.Right - RightPadding);
-		ConsolePosY = FMath::Min(FMath::Max(ConsolePosY, DisplayMetrics.VirtualDisplayRect.Top), DisplayMetrics.VirtualDisplayRect.Bottom - BottomPadding);
+		if (ConsolePosX != CW_USEDEFAULT)
+			ConsolePosX = FMath::Min(FMath::Max(ConsolePosX, DisplayMetrics.VirtualDisplayRect.Left), DisplayMetrics.VirtualDisplayRect.Right - RightPadding);
+		if (ConsolePosY != CW_USEDEFAULT)
+			ConsolePosY = FMath::Min(FMath::Max(ConsolePosY, DisplayMetrics.VirtualDisplayRect.Top), DisplayMetrics.VirtualDisplayRect.Bottom - BottomPadding);
 
 		FWriteScopeLock lock(WindowRWLock);
 
@@ -2114,6 +2627,10 @@ void FWindowsConsoleOutputDevice2::Show( bool ShowWindow )
 		Window->ConsoleWidth = ConsoleWidth;
 		Window->ConsoleHeight = ConsoleHeight;
 		Window->ConsoleColor = (FConsoleWindow::EConsoleColor)ConsoleColor;
+		Window->ConsoleShowDateTime = ConsoleShowDateTime;
+		Window->ConsoleIncludeFilterStr = ConsoleIncludeFilterStr;
+		Window->ConsoleExcludeFilterStr = ConsoleExcludeFilterStr;
+
 		Window->ConsoleTitle = ConsoleTitle;
 		Window->bConsoleExpanded = bConsoleExpanded;
 		Window->bIsVisible = true;
@@ -2192,14 +2709,11 @@ void FWindowsConsoleOutputDevice2::Serialize( const TCHAR* Data, ELogVerbosity::
 			}
 #endif
 		}
-		TStringBuilder<MAX_SPRINTF> OutputString;
-		FOutputDeviceHelper::AppendFormatLogLine(OutputString, Verbosity, Category, Data, GPrintLogTimes, RealTime);
-
 		{
 			FReadScopeLock lock(WindowRWLock);
 			if (Window)
 			{
-				Window->AddLogEntry(OutputString, TextAttribute);
+				Window->AddLogEntry(Data, Verbosity, Category, FPlatformTime::Seconds() - GStartTime, TextAttribute);
 			}
 		}
 
