tvOS doesn't support UIInterfaceOrientations. Also, ViewController was accessing incorrect View for supported orientations.

diff --git a/Engine/Source/Runtime/ApplicationCore/Private/IOS/IOSView.cpp b/Engine/Source/Runtime/ApplicationCore/Private/IOS/IOSView.cpp
--- a/Engine/Source/Runtime/ApplicationCore/Private/IOS/IOSView.cpp
+++ b/Engine/Source/Runtime/ApplicationCore/Private/IOS/IOSView.cpp
@@ -188,6 +188,7 @@ id<MTLDevice> GMetalDevice = nil;
 	
 	
 #if !PLATFORM_TVOS
+	SupportedInterfaceOrientations = UIInterfaceOrientationMaskAll;
 	self.multipleTouchEnabled = YES;
 #endif
 
@@ -1039,11 +1040,19 @@ self.accessibilityElements = @[Window.accessibilityContainer];
 }
 
 /**
- * Tell the OS what the default supported orientations are
+ * Tell the OS about the default supported orientations
  */
 - (NSUInteger)supportedInterfaceOrientations
 {
-	return SupportedInterfaceOrientations;
+	const FIOSView *View = [[IOSAppDelegate GetDelegate] IOSView];
+	if (View != nil)
+	{
+		return View->SupportedInterfaceOrientations;
+	}
+	else
+	{	
+		return 0;
+	}	
 }
 
 /**
diff --git a/Engine/Source/Runtime/ApplicationCore/Public/IOS/IOSView.h b/Engine/Source/Runtime/ApplicationCore/Public/IOS/IOSView.h
--- a/Engine/Source/Runtime/ApplicationCore/Public/IOS/IOSView.h
+++ b/Engine/Source/Runtime/ApplicationCore/Public/IOS/IOSView.h
@@ -66,7 +66,7 @@ APPLICATIONCORE_API
 	BOOL bSecureTextEntry;
 	
 	volatile int32 KeyboardShowCount;
-	NSUInteger SupportedInterfaceOrientations = UIInterfaceOrientationMaskAll;
+	NSUInteger SupportedInterfaceOrientations;
 }
 
 #if WITH_ACCESSIBILITY
