[ApplicationCore - WindowsConsoleOutputDevice2]
* Reduced lock time between game thread and hwnd thread
* Added name to thread for debugger
* Changed so ini-file is stored in same folder as binary so it doesnt get wiped when getting new builds
* Fixed crash during shutdown if Show(false) while window was never visible

diff --git a/Engine/Source/Runtime/ApplicationCore/Private/Windows/WindowsConsoleOutputDevice2.cpp b/Engine/Source/Runtime/ApplicationCore/Private/Windows/WindowsConsoleOutputDevice2.cpp
--- a/Engine/Source/Runtime/ApplicationCore/Private/Windows/WindowsConsoleOutputDevice2.cpp
+++ b/Engine/Source/Runtime/ApplicationCore/Private/Windows/WindowsConsoleOutputDevice2.cpp
@@ -36,6 +36,7 @@
 class FWindowsConsoleOutputDevice2::FConsoleWindow
 {
 public:
+	static constexpr int32 LogHistoryMaxSize = 200000;
 	static constexpr const TCHAR* HelpTitle = TEXT("Unreal Console Help");
 	static constexpr const TCHAR* HelpText = TEXT(
 		"Welcome to Unreal Console!\r\n"
@@ -57,11 +58,6 @@ public:
 		"Enjoy!\r\n"
 	);
 
-	enum
-	{
-		WM_NEWLOGENTRIES = WM_USER + 1,
-		WM_ACTIVITIESDIRTY,
-	};
 	enum ENotificationId : int
 	{
 		ID_LOG = 0x8801,
@@ -85,11 +81,14 @@ public:
 		for (int32 I = 0; I != 3; ++I)
 			StatusLightBrush[I] = CreateSolidBrush(Colors[I]);
 
+		// Create Event used to trigger update in log window thread
+		DirtyEvent = CreateEvent(NULL, false, false, NULL);
+
 		// Can't remove this in dtor since it happens after OnExit delegate is destroyed`
 		FCoreDelegates::OnExit.AddLambda([this]
 			{
 				bool bSaveIni = bIsVisible;
-				bIsVisible = false;
+				//bIsVisible = false;
 				if (bSaveIni)
 					Owner->SaveToINI();
 			});
@@ -102,20 +101,17 @@ public:
 		CloseHandle(Thread);
 		Thread = nullptr;
 
+		CloseHandle(DirtyEvent);
+
 		for (int32 I=0; I!=3; ++I)
 			DeleteObject(StatusLightBrush[I]);
 	}
 
 	void AddLogEntry(const FStringView& Text, ELogVerbosity::Type Verbosity, const class FName& Category, double Time, uint16 TextAttribute)
 	{
-		bool bWasEmpty;
-		{
-			FScopeLock _(&NewLogEntriesCs);
-			bWasEmpty = NewLogEntries.Num() == 0;
-			NewLogEntries.Add({FString(Text), Verbosity, Category, Time, TextAttribute});
-		}
-		if (bWasEmpty)
-			PostMessageW(MainHwnd, WM_NEWLOGENTRIES, 0, 0);
+		FScopeLock _(&NewLogEntriesCs);
+		NewLogEntries.Add({FString(Text), Verbosity, Category, Time, TextAttribute});
+		SetEvent(DirtyEvent);
 	}
 
 	// Light 0 = none, 1 = Red, 2 = Yellow, 3 = Green
@@ -124,59 +120,48 @@ public:
 		if (!*Name)
 			return;
 
-		FScopeLock _(&ActivitiesCs);
+		FScopeLock _(&ActivityModificationsCs);
 
-		bool bDirty = false;
-		if (int* Index = IdToActivityIndex.Find(Id))
+		bool IsNew = false;
+		int Index;
 		{
-			Activity& A = Activities[*Index];
-			if (A.Status != Status)
+			if (int* IndexPtr = IdToActivityIndex.Find(Id))
 			{
-				A.Status = Status;
-				A.bStatusDirty = true;
-				bDirty = true;
+				Index = *IndexPtr;
 			}
-			if (A.Light != Light)
+			else
 			{
-				A.Light = Light;
-				bActivitiesLightDirty = true;
-				bDirty = true;
+				IsNew = true;
+				Index = IdToActivityIndex.Num();
+				IdToActivityIndex.Add(Id, Index);
 			}
 		}
-		else
+
+		ActivityModification& Modification = ActivityModifications.FindOrAdd(Index);
+		if (IsNew)
 		{
-			IdToActivityIndex.Add(Id, Activities.Num());
-			Activity& A = Activities.AddDefaulted_GetRef();;
-			A.bAlignLeft = bAlignLeft;
-			A.Name = Name;
-			A.Status = Status;
-			A.Light = Light;
-			if (Light != 0)
-				bActivitiesLightDirty = true;
-			bDirty = true;
+			Modification.Name = Name;
+			Modification.bAlignLeft = bAlignLeft;
 		}
-
-		if (!bDirty || bActivitiesDirty)
-			return;
-		bActivitiesDirty = true;
-		PostMessageW(MainHwnd, WM_ACTIVITIESDIRTY, 0, 0);
+		Modification.Status = Status;
+		Modification.Light = Light;
+		SetEvent(DirtyEvent);
 	}
 
 	void RemoveStatus(const TCHAR* Id)
 	{
-		FScopeLock _(&ActivitiesCs);
-
-		int I;
-		if (!IdToActivityIndex.RemoveAndCopyValue(Id, I))
-			return;
-		Activities[I].Name.Reset();
-		bActivitiesDirty = true;
-		PostMessageW(MainHwnd, WM_ACTIVITIESDIRTY, 0, 0);
+		FScopeLock _(&ActivityModificationsCs);
+		if (int* IndexPtr = IdToActivityIndex.Find(Id))
+		{
+			ActivityModifications.FindOrAdd(*IndexPtr).bRemove = true;
+			SetEvent(DirtyEvent);
+		}
 	}
 
 	void Start()
 	{
 		Thread = CreateThread(NULL, 0, StaticThreadProc, this, 0, NULL);
+		SetThreadDescription(Thread, TEXT("LogConsoleHwnd"));
 	}
 
 	static DWORD WINAPI StaticThreadProc(LPVOID lpParameter)
@@ -289,8 +274,6 @@ public:
 		TipHwnd = CreateWindowEx(WS_EX_TOOLWINDOW | WS_EX_TOPMOST | WS_EX_NOACTIVATE, WC_STATIC, TEXT(""), SS_OWNERDRAW | WS_VISIBLE | WS_POPUP | WS_DISABLED | WS_CLIPCHILDREN | WS_CLIPSIBLINGS, 1, 1, 1, 1, MainHwnd, NULL, HInstance, NULL);
 		SendMessageW(TipHwnd, WM_SETFONT, (WPARAM)Font, 0);
 
-		PostMessageW(MainHwnd, WM_NEWLOGENTRIES, 0, 0);
-
 		FDateTime startTime(FDateTime::Now());
 		uint64 lastSeconds = 0;
 		bool bLoop = true;
@@ -321,27 +304,32 @@ public:
 				UpdateTime(seconds);
 			}
 
-			DWORD timeout = 200;
-			if (MsgWaitForMultipleObjects(0, NULL, false, timeout, QS_ALLINPUT) == WAIT_OBJECT_0)
+			DWORD Timeout = 200;
+			DWORD WaitResult = MsgWaitForMultipleObjects(1, &DirtyEvent, false, Timeout, QS_ALLINPUT);
+			if (WaitResult == WAIT_TIMEOUT)
+				continue;
+
+			HandleActivityModifications();
+
+			HandleNewLogEntries();
+
+			MSG Msg;
+			while (PeekMessage(&Msg, NULL, 0, 0, PM_REMOVE))
 			{
-				MSG Msg;
-				while (PeekMessage(&Msg, NULL, 0, 0, PM_REMOVE))
+				if (!IsDialogMessage(MainHwnd, &Msg))
 				{
-					if (!IsDialogMessage(MainHwnd, &Msg))
-					{
-						TranslateMessage(&Msg);
-						DispatchMessage(&Msg);
-					}
-					if (Msg.message == WM_QUIT)
-					{
-						DestroyWindow(MainHwnd);
-						DeleteObject(LogFont);
-						DeleteObject(Font);
-						UnregisterClass(WindowClassName, HInstance);
-						MainHwnd = 0;
-						bLoop = false;
-						break;
-					}
+					TranslateMessage(&Msg);
+					DispatchMessage(&Msg);
+				}
+				if (Msg.message == WM_QUIT)
+				{
+					DestroyWindow(MainHwnd);
+					DeleteObject(LogFont);
+					DeleteObject(Font);
+					UnregisterClass(WindowClassName, HInstance);
+					MainHwnd = 0;
+					bLoop = false;
+					break;
 				}
 			}
 		}
@@ -354,6 +342,12 @@ public:
 
 			FWindowsPlatformMisc::CallGracefulTerminationHandler();
 		}
+		else
+		{
+			// This means that we've pressed close after the process has already been exited. It is likely the shutdown process is hung and we should just terminate the process
+			Sleep(500);
+			TerminateProcess(GetCurrentProcess(), -1);
+		}
 
 		return 0;
 	}
@@ -787,7 +781,7 @@ public:
 		if (!bSuspendAddingEntries)
 			return;
 		bSuspendAddingEntries = false;
-		PostMessageW(MainHwnd, WM_NEWLOGENTRIES, 0, 0);
+		SetEvent(DirtyEvent);
 	}
 
 	static LRESULT CALLBACK StaticLogHwndWndProc(HWND hWnd, UINT Msg, WPARAM wParam, LPARAM lParam)
@@ -1248,7 +1242,7 @@ public:
 
 			int32 LogIndex = SendMessageW(LogHwnd, LB_GETITEMDATA, RightClickedItem, 0) - LogIndexOffset;
 			RightClickedItem = -1;
-			if (LogIndex == LB_ERR)
+			if (LogIndex < 0)
 			{
 				RECT Rect;
 				GetClientRect(hWnd, &Rect);
@@ -1287,12 +1281,14 @@ public:
 			AddItem(TEXT("Clear all entries above line"), 104);
 			AddItem(TEXT("Clear filters"), 121);
 
-			SelectedCategory = Log[LogIndex].Category.ToString();
-
-			AddSeparator();
-			AddItem(*TStringBuilder<64>().Append(TEXT("Include category '")).Append(*SelectedCategory).Append(TEXT("' in filter")), 124);
-			AddItem(*TStringBuilder<64>().Append(TEXT("Exclude category '")).Append(*SelectedCategory).Append(TEXT("' in filter")), 125);
+			if (LogIndex >= 0)
+			{
+				SelectedCategory = Log[LogIndex].Category.ToString();
 
+				AddSeparator();
+				AddItem(*TStringBuilder<64>().Append(TEXT("Include category '")).Append(*SelectedCategory).Append(TEXT("' in filter")), 124);
+				AddItem(*TStringBuilder<64>().Append(TEXT("Exclude category '")).Append(*SelectedCategory).Append(TEXT("' in filter")), 125);
+			}
 
 			int CharIndex = (MousePos.x + (LogFontWidth / 2)) / LogFontWidth;
 			if (CharIndex < OutString.Len())
@@ -1468,166 +1464,221 @@ public:
 		return CallNextHookEx(0, nCode, wParam, lParam);
 	}
 
-	LRESULT MainWinProc(HWND hWnd, UINT Msg, WPARAM wParam, LPARAM lParam)
+	void HandleActivityModifications()
 	{
-		switch (Msg)
-		{
-		case WM_NEWLOGENTRIES:
+		bool bActivitiesLightDirty = false;
+
+		// Transfer modifications to Activities and update IdToActivity
 		{
-			if (bSuspendAddingEntries)
-				break;
+			TArray<int> IndicesToRemove;
 
-			{
-				FScopeLock _(&NewLogEntriesCs);
-				Swap(NewLogEntries, TempLogEntries);
-			}
+			FScopeLock _(&ActivityModificationsCs);
 
-			if (TempLogEntries.IsEmpty())
-				break;
+			for (auto& ModPair : ActivityModifications)
+			{
+				int Index = ModPair.Key;
+				ActivityModification& Mod = ModPair.Value;
+				if (Mod.bRemove)
+				{
+					IndicesToRemove.Add(Index);
+					continue;
+				}
+				if (Index >= Activities.Num())
+				{
+					Activities.SetNum(Index+1);
+					Activity& A = Activities[Index];
+					A.Name = Mod.Name;
+					A.Light = Mod.Light;
+					A.Status = Mod.Status;
+					A.bAlignLeft = Mod.bAlignLeft;
+					continue;
+				}
 
-			int32 OldSize = Log.Num();
+				Activity& A = Activities[Index];
+				if (A.Light != Mod.Light)
+				{
+					A.Light = Mod.Light;
+					bActivitiesLightDirty = true;
+				}
+				if (A.Status != Mod.Status)
+				{
+					A.Status = Mod.Status;
+					A.bStatusDirty = true;
+				}
+			}
 
-			for (NewLogEntry& E : TempLogEntries)
+			for (int I = IndicesToRemove.Num() - 1; I >= 0; --I)
 			{
-				const TCHAR* SearchStr = *E.String;
-				while (true)
+				int ToRemove = IndicesToRemove[I];
+				if (ToRemove < Activities.Num())
+					Activities[ToRemove].Name.Empty();
+				for (auto It = IdToActivityIndex.CreateIterator(); It; ++It)
 				{
-					if (const TCHAR* LineBreak = FCString::Strchr(SearchStr, '\n'))
+					if (It.Value() == ToRemove)
 					{
-						int32 Len = LineBreak - SearchStr;
-						if (LineBreak > SearchStr && *(LineBreak - 1) == '\r')
-							--Len;
-						Log.Add({ FString(Len, SearchStr), E.Verbosity, E.Category, E.Time, E.TextAttribute, 1 });
-						SearchStr = LineBreak + 1;
+						It.RemoveCurrent();
 					}
-					else
+					else if (It.Value() > ToRemove)
 					{
-						if (int32 Len = FCString::Strlen(SearchStr))
-							Log.Add({ SearchStr, E.Verbosity, E.Category, E.Time, E.TextAttribute, 1 });
-						break;
+						--It.Value();
 					}
 				}
 			}
 
-			constexpr int32 MaxSize = 20000;
-
-			int32 NewSize = Log.Num();
-
-			int32 ToAddToHwnd = NewSize - OldSize;
-			int32 LogIndex = OldSize;
-			int32 NewTopIndex = -1;
-			int32 NewCaretIndex = -1;
+			ActivityModifications.Reset();
+		}
 
-			SendMessageW(LogHwnd, WM_SETREDRAW, false, 0);
+		bool bUpdatePositions = false;
+		for (int32 i = 0, e = Activities.Num(); i != e;)
+		{
+			Activity& A = Activities[i];
 
-			if (NewSize > MaxSize)
+			if (A.Name.IsEmpty())
 			{
-				int32 ToRemove = NewSize - MaxSize;
-				Log.PopFront(ToRemove);
-				LogIndexOffset += ToRemove;
-
-				int32 ToChangeInHwnd = FMath::Min(MaxSize, ToRemove);
-
-				int32 ToRemoveFromHwnd = ToChangeInHwnd;
-				if (!bAutoScrollLog)
-				{
-					int32 TopIndex = SendMessageW(LogHwnd, LB_GETTOPINDEX, 0, 0);
-					NewTopIndex = FMath::Max(TopIndex - ToRemoveFromHwnd, 0);
-					int32 CaretIndex = SendMessageW(LogHwnd, LB_GETCARETINDEX, 0, 0);
-					NewCaretIndex = FMath::Max(CaretIndex - ToRemoveFromHwnd, 0);
-				}
-
-				while (ToRemoveFromHwnd--)
-					SendMessageW(LogHwnd, LB_DELETESTRING, 0, 0);
-
-				ToAddToHwnd = ToChangeInHwnd;
-				LogIndex = MaxSize - ToAddToHwnd;
+				DestroyWindow(A.NameHwnd);
+				DestroyWindow(A.StatusHwnd);
+				Activities.RemoveAt(i);
+				bUpdatePositions = true;
+				--e;
+				continue;
 			}
 
-			while (ToAddToHwnd--)
+			if (!A.NameHwnd)
+			{
+				A.NameHwnd = CreateTextHwnd(*A.Name, Font);
+				A.StatusHwnd = CreateTextHwnd(*A.Status, Font);
+				bUpdatePositions = true;
+			}
+			else if (A.bStatusDirty)
 			{
-				AddEntryToLogHwnd(Log[LogIndex], LogIndexOffset + LogIndex);
-				++LogIndex;
+				A.bStatusDirty = false;
+				SetWindowTextW(A.StatusHwnd, *A.Status);
 			}
-			TempLogEntries.Reset(0);
+			++i;
+		}
 
-			if (bAutoScrollLog)
+		if (bActivitiesLightDirty)
+		{
+			InvalidateRect(MainHwnd, NULL, false);
+			bActivitiesLightDirty = false;
+		}
+
+		if (bUpdatePositions)
+		{
+			SendMessageW(LogHwnd, WM_SETREDRAW, false, 0);
+			RECT Rect;
+			GetClientRect(MainHwnd, &Rect);
+			int NewActivitiesTotalHeight = TraverseActivityPositions(Rect.right, Rect.bottom, [](Activity& A, int32 X, int32 Y, int32 Width, int32 RowIndex) {});
+			if (NewActivitiesTotalHeight == ActivitiesTotalHeight)
 			{
-				ScrollToBottom();
+				UpdateActivityPositions(Rect.right, Rect.bottom, false);
 			}
-			else if (NewTopIndex != -1)
+			else
 			{
-				SendMessageW(LogHwnd, LB_SETCARETINDEX, NewCaretIndex, 0);
-				SendMessageW(LogHwnd, LB_SETSEL, true, NewCaretIndex);
-				SetTopVisible(NewTopIndex, false);
+				ActivitiesTotalHeight = NewActivitiesTotalHeight;
+				UpdateSize(Rect.right, Rect.bottom, false);
 			}
-			RedrawLogScrollbar();
-
+			RedrawWindow(MainHwnd, NULL, NULL, RDW_ERASE | RDW_NOINTERNALPAINT | RDW_INVALIDATE | RDW_ALLCHILDREN | RDW_UPDATENOW); // Don't change flags, very sensitive with glitches during resizing
 			SendMessageW(LogHwnd, WM_SETREDRAW, true, 0);
-			break;
 		}
+	}
+
+	void HandleNewLogEntries()
+	{
+		if (bSuspendAddingEntries)
+			return;
 
-		case WM_ACTIVITIESDIRTY:
 		{
-			FScopeLock _(&ActivitiesCs);
-			bActivitiesDirty = false;
-			bool bUpdatePositions = false;
-			for (int32 i = 0, e = Activities.Num(); i != e;)
-			{
-				Activity& A = Activities[i];
+			FScopeLock _(&NewLogEntriesCs);
+			Swap(NewLogEntries, TempLogEntries);
+		}
 
-				if (A.Name.IsEmpty())
-				{
-					DestroyWindow(A.NameHwnd);
-					DestroyWindow(A.StatusHwnd);
-					Activities.RemoveAt(i);
-					bUpdatePositions = true;
-					for (auto& KV : IdToActivityIndex)
-						if (KV.Value > i)
-							--KV.Value;
-					--e;
-					continue;
-				}
+		if (TempLogEntries.IsEmpty())
+			return;
 
-				if (!A.NameHwnd)
+		int32 OldSize = Log.Num();
+
+		for (NewLogEntry& E : TempLogEntries)
+		{
+			const TCHAR* SearchStr = *E.String;
+			while (true)
+			{
+				if (const TCHAR* LineBreak = FCString::Strchr(SearchStr, '\n'))
 				{
-					A.NameHwnd = CreateTextHwnd(*A.Name, Font);
-					A.StatusHwnd = CreateTextHwnd(*A.Status, Font);
-					bUpdatePositions = true;
+					int32 Len = LineBreak - SearchStr;
+					if (LineBreak > SearchStr && *(LineBreak - 1) == '\r')
+						--Len;
+					Log.Add({ FString(Len, SearchStr), E.Verbosity, E.Category, E.Time, E.TextAttribute, 1 });
+					SearchStr = LineBreak + 1;
 				}
-				else if (A.bStatusDirty)
+				else
 				{
-					A.bStatusDirty = false;
-					SetWindowTextW(A.StatusHwnd, *A.Status);
+					if (int32 Len = FCString::Strlen(SearchStr))
+						Log.Add({ SearchStr, E.Verbosity, E.Category, E.Time, E.TextAttribute, 1 });
+					break;
 				}
-				++i;
 			}
+		}
 
-			if (bActivitiesLightDirty)
-			{
-				InvalidateRect(MainHwnd, NULL, false);
-				bActivitiesLightDirty = false;
-			}
+		int32 NewSize = Log.Num();
+
+		int32 ToAddToHwnd = NewSize - OldSize;
+		int32 LogIndex = OldSize;
+		int32 NewTopIndex = -1;
+		int32 NewCaretIndex = -1;
+
+		SendMessageW(LogHwnd, WM_SETREDRAW, false, 0);
 
-			if (bUpdatePositions)
+		if (NewSize > LogHistoryMaxSize)
+		{
+			int32 ToRemove = NewSize - LogHistoryMaxSize;
+			Log.PopFront(ToRemove);
+			LogIndexOffset += ToRemove;
+
+			int32 ToChangeInHwnd = FMath::Min(LogHistoryMaxSize, ToRemove);
+
+			int32 ToRemoveFromHwnd = ToChangeInHwnd;
+			if (!bAutoScrollLog)
 			{
-				RECT Rect;
-				GetClientRect(hWnd, &Rect);
-				int NewActivitiesTotalHeight = TraverseActivityPositions(Rect.right, Rect.bottom, [](Activity& A, int32 X, int32 Y, int32 Width, int32 RowIndex) {});
-				if (NewActivitiesTotalHeight == ActivitiesTotalHeight)
-				{
-					UpdateActivityPositions(Rect.right, Rect.bottom, false);
-				}
-				else
-				{
-					ActivitiesTotalHeight = NewActivitiesTotalHeight;
-					UpdateSize(Rect.right, Rect.bottom, false);
-				}
-				RedrawWindow(MainHwnd, NULL, NULL, RDW_ERASE | RDW_INVALIDATE | RDW_ALLCHILDREN | RDW_UPDATENOW); // Important with UPDATENOW to make sure we avoid glitches
+				int32 TopIndex = SendMessageW(LogHwnd, LB_GETTOPINDEX, 0, 0);
+				NewTopIndex = FMath::Max(TopIndex - ToRemoveFromHwnd, 0);
+				int32 CaretIndex = SendMessageW(LogHwnd, LB_GETCARETINDEX, 0, 0);
+				NewCaretIndex = FMath::Max(CaretIndex - ToRemoveFromHwnd, 0);
 			}
-			break;
+
+			while (ToRemoveFromHwnd--)
+				SendMessageW(LogHwnd, LB_DELETESTRING, 0, 0);
+
+			ToAddToHwnd = ToChangeInHwnd;
+			LogIndex = LogHistoryMaxSize - ToAddToHwnd;
+		}
+
+		while (ToAddToHwnd--)
+		{
+			AddEntryToLogHwnd(Log[LogIndex], LogIndexOffset + LogIndex);
+			++LogIndex;
 		}
+		TempLogEntries.Reset(0);
 
+		if (bAutoScrollLog)
+		{
+			ScrollToBottom();
+		}
+		else if (NewTopIndex != -1)
+		{
+			SendMessageW(LogHwnd, LB_SETCARETINDEX, NewCaretIndex, 0);
+			SendMessageW(LogHwnd, LB_SETSEL, true, NewCaretIndex);
+			SetTopVisible(NewTopIndex, false);
+		}
+		RedrawLogScrollbar();
+
+		SendMessageW(LogHwnd, WM_SETREDRAW, true, 0);
+	}
+
+	LRESULT MainWinProc(HWND hWnd, UINT Msg, WPARAM wParam, LPARAM lParam)
+	{
+		switch (Msg)
+		{
 		case WM_HELP:
 			{
 				HHOOK Hook = SetWindowsHookEx(WH_CBT, StaticMessageBoxHookProc, NULL, GetCurrentThreadId());
@@ -2269,7 +2320,6 @@ public:
 	FCriticalSection NewLogEntriesCs;
 	TArray<NewLogEntry> NewLogEntries;
 	TArray<NewLogEntry> TempLogEntries;
-	TMap<FString, int> IdToActivityIndex;
 	TArray<FString> IncludeFilter;
 	TArray<FString> ExcludeFilter;
 	TRingBuffer<LogEntry> Log;
@@ -2324,11 +2374,13 @@ public:
 	FString SelectedCategory;
 
 	struct Activity { HWND NameHwnd = 0; HWND StatusHwnd = 0; FString Name; FString Status; int32 Light = 0; bool bStatusDirty = false; bool bAlignLeft = false; };
-	FCriticalSection ActivitiesCs;
+	struct ActivityModification { FString Name; FString Status; int32 Light = 0; bool bAlignLeft = false; bool bRemove = false; };
 	TArray<Activity> Activities;
-	bool bActivitiesDirty = false;
-	bool bActivitiesLightDirty = false;
+	TMap<FString, int> IdToActivityIndex;
+	FCriticalSection ActivityModificationsCs;
+	TSortedMap<int, ActivityModification> ActivityModifications;
 	int32 ActivitiesTotalHeight = 0;
+	HANDLE DirtyEvent;
 	
 	HANDLE Thread;
 	bool bAutoScrollLog = true;
@@ -2507,7 +2559,7 @@ void FWindowsConsoleOutputDevice2::SaveToINI()
 	bool bConsoleExpanded = Window->bConsoleExpanded;
 	WindowRWLock.WriteUnlock();
 
-	FString Filename = TEXT("DebugConsole.ini");
+	FString Filename = FPlatformProcess::GetModulesDirectory() + TEXT("\\DebugConsole.ini");
 	const TCHAR* Selection = TEXT("ConsoleWindows");
 	if (IsRunningDedicatedServer())
 		Selection = TEXT("ServerConsoleWindows");
@@ -2554,7 +2606,7 @@ void FWindowsConsoleOutputDevice2::Show( bool ShowWindow )
 		bool bHasX = false;
 		bool bHasY = false;
 
-		FString Filename = TEXT("DebugConsole.ini");
+		FString Filename = FPlatformProcess::GetModulesDirectory() + TEXT("\\DebugConsole.ini");
 		const TCHAR* Selection = TEXT("ConsoleWindows");
 		if (IsRunningDedicatedServer())
 			Selection = TEXT("ServerConsoleWindows");
@@ -2641,7 +2693,7 @@ void FWindowsConsoleOutputDevice2::Show( bool ShowWindow )
 		}
 
 	}
-	else
+	else if (Window)
 	{
 		SaveToINI();
 		Window->bIsVisible = false;
