Fix static analysis error in FWindowsConsoleOutputDevice2::FConsoleWindow::Start: CreateThread can return null, and SetThreadDescription should not be passed a null thread, so the result of CreateThread must be checked for null before calling SetThreadDescription.

diff --git a/Engine/Source/Runtime/ApplicationCore/Private/Windows/WindowsConsoleOutputDevice2.cpp b/Engine/Source/Runtime/ApplicationCore/Private/Windows/WindowsConsoleOutputDevice2.cpp
--- a/Engine/Source/Runtime/ApplicationCore/Private/Windows/WindowsConsoleOutputDevice2.cpp
+++ b/Engine/Source/Runtime/ApplicationCore/Private/Windows/WindowsConsoleOutputDevice2.cpp
@@ -161,7 +161,10 @@ public:
 	void Start()
 	{
 		Thread = CreateThread(NULL, 0, StaticThreadProc, this, 0, NULL);
-		SetThreadDescription(Thread, TEXT("LogConsoleHwnd"));
+		if (ensure(Thread))
+		{
+			SetThreadDescription(Thread, TEXT("LogConsoleHwnd"));
+		}
 	}
 
 	static DWORD WINAPI StaticThreadProc(LPVOID lpParameter)
