Fixed unsafe type casts in a few widely-used modules

#rnx

diff --git a/Engine/Source/Developer/DesktopPlatform/Private/DesktopPlatformBase.cpp b/Engine/Source/Developer/DesktopPlatform/Private/DesktopPlatformBase.cpp
--- a/Engine/Source/Developer/DesktopPlatform/Private/DesktopPlatformBase.cpp
+++ b/Engine/Source/Developer/DesktopPlatform/Private/DesktopPlatformBase.cpp
@@ -235,7 +235,7 @@ bool FDesktopPlatformBase::TryGetEngineVersion(const FString& RootDir, FEngineVe
 				}
 
 				int EncodedChangelist = (IsLicenseeVersion == 0)? Changelist : FEngineVersionBase::EncodeLicenseeChangelist(Changelist);
-				OutVersion = FEngineVersion(MajorVersion, MinorVersion, PatchVersion, EncodedChangelist, BranchName);
+				OutVersion = FEngineVersion(IntCastChecked<uint16>(MajorVersion), IntCastChecked<uint16>(MinorVersion), IntCastChecked<uint16>(PatchVersion), EncodedChangelist, BranchName);
 				return true;
 			}
 		}
@@ -276,14 +276,14 @@ bool FDesktopPlatformBase::TryGetEngineVersion(const FString& RootDir, FEngineVe
 					// Parse an identifier. Exact C rules for an identifier don't really matter; we just need alphanumeric sequences.
 					const TCHAR* TokenStart = TextPos++;
 					while(FChar::IsIdentifier(*TextPos)) TextPos++;
-					Tokens.Add(FString(TextPos - TokenStart, TokenStart));
+					Tokens.Add(FString(UE_PTRDIFF_TO_INT32(TextPos - TokenStart), TokenStart));
 				}
 				else if(*TextPos == '\"')
 				{
 					// Parse a string
 					const TCHAR* TokenStart = TextPos++;
 					while(*TextPos != 0 && (TextPos == TokenStart + 1 || *(TextPos - 1) != '\"')) TextPos++;
-					Tokens.Add(FString(TextPos - TokenStart, TokenStart));
+					Tokens.Add(FString(UE_PTRDIFF_TO_INT32(TextPos - TokenStart), TokenStart));
 				}
 				else if(*TextPos == '/' && *(TextPos + 1) == '/')
 				{
@@ -345,7 +345,7 @@ bool FDesktopPlatformBase::TryGetEngineVersion(const FString& RootDir, FEngineVe
 		if(MajorVersion != -1 && MinorVersion != -1 && PatchVersion != -1)
 		{
 			int EncodedChangelist = (IsLicenseeVersion == 0)? Changelist : FEngineVersionBase::EncodeLicenseeChangelist(Changelist);
-			OutVersion = FEngineVersion(MajorVersion, MinorVersion, PatchVersion, EncodedChangelist, BranchName);
+			OutVersion = FEngineVersion(IntCastChecked<uint16>(MajorVersion), IntCastChecked<uint16>(MinorVersion), IntCastChecked<uint16>(PatchVersion), EncodedChangelist, BranchName);
 			return true;
 		}
 	}
@@ -375,7 +375,7 @@ bool FDesktopPlatformBase::TryParseStockEngineVersion(const FString& Identifier,
 		return false;
 	}
 
-	OutVersion = FEngineVersion(Major, Minor, 0, 0, TEXT(""));
+	OutVersion = FEngineVersion(IntCastChecked<uint16>(Major), IntCastChecked<uint16>(Minor), 0, 0, TEXT(""));
 	return true;
 }
 
@@ -1231,7 +1231,7 @@ void FDesktopPlatformBase::CheckForLauncherEngineInstallation(const FString &App
 	}
 }
 
-int32 FDesktopPlatformBase::ParseReleaseVersion(const FString &Version)
+int32 FDesktopPlatformBase::ParseReleaseVersion(const FString& Version)
 {
 	TCHAR *End;
 
@@ -1247,7 +1247,7 @@ int32 FDesktopPlatformBase::ParseReleaseVersion(const FString &Version)
 		return INDEX_NONE;
 	}
 
-	return (Major << 16) + Minor;
+	return IntCastChecked<int32>((Major << 16) + Minor);
 }
 
 TSharedPtr<FJsonObject> FDesktopPlatformBase::LoadProjectFile(const FString &FileName)
diff --git a/Engine/Source/Developer/DesktopPlatform/Private/Linux/DesktopPlatformLinux.cpp b/Engine/Source/Developer/DesktopPlatform/Private/Linux/DesktopPlatformLinux.cpp
--- a/Engine/Source/Developer/DesktopPlatform/Private/Linux/DesktopPlatformLinux.cpp
+++ b/Engine/Source/Developer/DesktopPlatform/Private/Linux/DesktopPlatformLinux.cpp
@@ -469,7 +469,7 @@ bool FDesktopPlatformLinux::IsUnrealBuildToolRunning()
 				const TCHAR* Delim = FCString::Strchr(*Filename, '_');
 				if (Delim != nullptr)
 				{
-					FString Pid(*Filename, Delim - *Filename);
+					FString Pid(*Filename, UE_PTRDIFF_TO_INT32(Delim - *Filename));
 					int ProcessId = 0;
 					LexFromString(ProcessId, *Pid);
 					const FString EntryFullPath = FPlatformProcess::GetApplicationName(ProcessId);
diff --git a/Engine/Source/Developer/DesktopPlatform/Private/Mac/DesktopPlatformMac.cpp b/Engine/Source/Developer/DesktopPlatform/Private/Mac/DesktopPlatformMac.cpp
--- a/Engine/Source/Developer/DesktopPlatform/Private/Mac/DesktopPlatformMac.cpp
+++ b/Engine/Source/Developer/DesktopPlatform/Private/Mac/DesktopPlatformMac.cpp
@@ -646,7 +646,7 @@ bool FDesktopPlatformMac::IsUnrealBuildToolRunning()
 				const TCHAR* Delim = FCString::Strchr(*Filename, '_');
 				if (Delim != nullptr)
 				{
-					FString Pid(*Filename, Delim - *Filename);
+					FString Pid(*Filename, UE_PTRDIFF_TO_INT32(Delim - *Filename));
 					int ProcessId = 0;
 					LexFromString(ProcessId, *Pid);
 					const FString EntryFullPath = FPlatformProcess::GetApplicationName(ProcessId);
diff --git a/Engine/Source/Developer/DesktopPlatform/Private/Windows/DesktopPlatformWindows.cpp b/Engine/Source/Developer/DesktopPlatform/Private/Windows/DesktopPlatformWindows.cpp
--- a/Engine/Source/Developer/DesktopPlatform/Private/Windows/DesktopPlatformWindows.cpp
+++ b/Engine/Source/Developer/DesktopPlatform/Private/Windows/DesktopPlatformWindows.cpp
@@ -134,7 +134,7 @@ bool FDesktopPlatformWindows::OpenFontDialog(const void* ParentWindowHandle, FSt
 		}
 
 		OutFontName = (const TCHAR*)lf.lfFaceName;
-		OutHeight = PixelHeight;
+		OutHeight = (float)PixelHeight;
 		OutFlags = FontFlags;
 
 		::ReleaseDC( cf.hwndOwner, DC ); 
@@ -498,7 +498,7 @@ bool FDesktopPlatformWindows::IsUnrealBuildToolRunning()
 				const TCHAR* Delim = FCString::Strchr(*Filename, '_');
 				if (Delim != nullptr)
 				{
-					FString Pid(*Filename, Delim - *Filename);
+					FString Pid(*Filename, UE_PTRDIFF_TO_INT32(Delim - *Filename));
 					int ProcessId = 0;
 					LexFromString(ProcessId, *Pid);
 					const FString EntryFullPath = FPlatformProcess::GetApplicationName(ProcessId);
diff --git a/Engine/Source/Developer/Zen/Internal/ZenClient.cpp b/Engine/Source/Developer/Zen/Internal/ZenClient.cpp
--- a/Engine/Source/Developer/Zen/Internal/ZenClient.cpp
+++ b/Engine/Source/Developer/Zen/Internal/ZenClient.cpp
@@ -4,6 +4,7 @@
 #include "Async/Async.h"
 #include "Containers/StringView.h"
 #include "HAL/Event.h"
+#include "Math/UnrealMathUtility.h"
 #include "Memory/MemoryView.h"
 #include "Misc/StringBuilder.h"
 #include "ProfilingDebugging/CountersTrace.h"
@@ -26,8 +27,9 @@ bool SendData(FSocket& Socket, FMemoryView Data)
 {
 	while (Data.IsEmpty() == false)
 	{
+		const int32 BytesToSend = (int32)FMath::Clamp<uint64>(Data.GetSize(), 0, MAX_int32);
 		int32 BytesSent = 0;
-		if (Socket.Send(reinterpret_cast<const uint8*>(Data.GetData()), Data.GetSize(), BytesSent) == false)
+		if (Socket.Send((const uint8*)Data.GetData(), BytesToSend, BytesSent) == false)
 		{
 			return false;
 		}
diff --git a/Engine/Source/Developer/Zen/Internal/ZenServerHttp.cpp b/Engine/Source/Developer/Zen/Internal/ZenServerHttp.cpp
--- a/Engine/Source/Developer/Zen/Internal/ZenServerHttp.cpp
+++ b/Engine/Source/Developer/Zen/Internal/ZenServerHttp.cpp
@@ -129,7 +129,7 @@ namespace UE::Zen {
 
 	FZenHttpRequest::Result FZenHttpRequest::PerformBlockingPut(const TCHAR* Uri, const FCompositeBuffer& Buffer, EContentType ContentType)
 	{
-		uint32 ContentLength = 0u;
+		uint64 ContentLength = 0u;
 
 		ContentLength = Buffer.GetSize();
 		curl_easy_setopt(Curl, CURLOPT_UPLOAD, 1L);
@@ -289,7 +289,7 @@ namespace UE::Zen {
 
 	static std::atomic<int> RequestId{1};
 
-	FZenHttpRequest::Result FZenHttpRequest::PerformBlocking(FStringView Uri, RequestVerb Verb, uint32 ContentLength)
+	FZenHttpRequest::Result FZenHttpRequest::PerformBlocking(FStringView Uri, RequestVerb Verb, uint64 ContentLength)
 	{
 		LLM_SCOPE_BYTAG(ZenDDC);
 		// Strip any leading slashes because we compose the prefix and the suffix with a separating slash below
@@ -324,7 +324,7 @@ namespace UE::Zen {
 
 		if ((Verb != RequestVerb::Delete) && (Verb != RequestVerb::Get))
 		{
-			Headers.Add(FString::Printf(TEXT("Content-Length: %d"), ContentLength));
+			Headers.Add(FString::Printf(TEXT("Content-Length: %" UINT64_FMT), ContentLength));
 		}
 
 		// Build headers list
@@ -465,7 +465,7 @@ namespace UE::Zen {
 	FString FZenHttpRequest::GetAnsiBufferAsString(const TArray64<uint8>& Buffer)
 	{
 		// Content is NOT null-terminated; we need to specify lengths here
-		FUTF8ToTCHAR TCHARData(reinterpret_cast<const ANSICHAR*>(Buffer.GetData()), Buffer.Num());
+		FUTF8ToTCHAR TCHARData(reinterpret_cast<const ANSICHAR*>(Buffer.GetData()), IntCastChecked<int32>(Buffer.Num()));
 		return FString(TCHARData.Length(), TCHARData.Get());
 	}
 
diff --git a/Engine/Source/Developer/Zen/Internal/ZenServerHttp.h b/Engine/Source/Developer/Zen/Internal/ZenServerHttp.h
--- a/Engine/Source/Developer/Zen/Internal/ZenServerHttp.h
+++ b/Engine/Source/Developer/Zen/Internal/ZenServerHttp.h
@@ -245,7 +245,7 @@ private:
 		* @param Buffer Optional buffer to directly receive the result of the request.
 		* If unset the response body will be stored in the request.
 		*/
-	Result PerformBlocking(FStringView Uri, RequestVerb Verb, uint32 ContentLength);
+	Result PerformBlocking(FStringView Uri, RequestVerb Verb, uint64 ContentLength);
 
 	FZenHttpRequest::Result ParseRpcResponse(FZenHttpRequest::Result ResultFromPost, FCbPackage& OutResponse);
 
diff --git a/Engine/Source/Developer/Zen/Internal/ZenServerInterface.cpp b/Engine/Source/Developer/Zen/Internal/ZenServerInterface.cpp
--- a/Engine/Source/Developer/Zen/Internal/ZenServerInterface.cpp
+++ b/Engine/Source/Developer/Zen/Internal/ZenServerInterface.cpp
@@ -1099,10 +1099,10 @@ FZenServiceInstance::GetStats(FZenStats& Stats)
 	Stats = LastStats;
 
 	const uint64 CurrentTime = FPlatformTime::Cycles64();
-	const float MinTimeBetweenRequestsInSeconds = 0.5f;
-	const float DeltaTimeInSeconds = FPlatformTime::ToSeconds(CurrentTime - LastStatsTime);
+	constexpr double MinTimeBetweenRequestsInSeconds = 0.5;
+	const double DeltaTimeInSeconds = FPlatformTime::ToSeconds64(CurrentTime - LastStatsTime);
 
-	if ( StatsRequest.IsValid()==false && DeltaTimeInSeconds > MinTimeBetweenRequestsInSeconds)
+	if (!StatsRequest.IsValid() && DeltaTimeInSeconds > MinTimeBetweenRequestsInSeconds)
 	{
 #if WITH_EDITOR
 		EAsyncExecution ThreadPool = EAsyncExecution::LargeThreadPool;
diff --git a/Engine/Source/Runtime/ApplicationCore/Private/GenericPlatform/GenericApplication.cpp b/Engine/Source/Runtime/ApplicationCore/Private/GenericPlatform/GenericApplication.cpp
--- a/Engine/Source/Runtime/ApplicationCore/Private/GenericPlatform/GenericApplication.cpp
+++ b/Engine/Source/Runtime/ApplicationCore/Private/GenericPlatform/GenericApplication.cpp
@@ -154,7 +154,7 @@ void FDisplayMetrics::ApplyDefaultSafeZones()
 		if (SafeZoneRatio < 1.0f)
 		{
 			const float HalfUnsafeRatio = (1.0f - SafeZoneRatio) * 0.5f;
-			TitleSafePaddingSize = FVector4(PrimaryDisplayWidth * HalfUnsafeRatio, PrimaryDisplayHeight * HalfUnsafeRatio, PrimaryDisplayWidth * HalfUnsafeRatio, PrimaryDisplayHeight * HalfUnsafeRatio);
+			TitleSafePaddingSize = FVector4((float)PrimaryDisplayWidth * HalfUnsafeRatio, (float)PrimaryDisplayHeight * HalfUnsafeRatio, (float)PrimaryDisplayWidth * HalfUnsafeRatio, (float)PrimaryDisplayHeight * HalfUnsafeRatio);
 		}
 	}
 
@@ -162,7 +162,7 @@ void FDisplayMetrics::ApplyDefaultSafeZones()
 	if (ActionSafeZoneRatio < 1.0f)
 	{
 		const float HalfUnsafeRatio = (1.0f - ActionSafeZoneRatio) * 0.5f;
-		ActionSafePaddingSize = FVector4(PrimaryDisplayWidth * HalfUnsafeRatio, PrimaryDisplayHeight * HalfUnsafeRatio, PrimaryDisplayWidth * HalfUnsafeRatio, PrimaryDisplayHeight * HalfUnsafeRatio);
+		ActionSafePaddingSize = FVector4((float)PrimaryDisplayWidth * HalfUnsafeRatio, (float)PrimaryDisplayHeight * HalfUnsafeRatio, (float)PrimaryDisplayWidth * HalfUnsafeRatio, (float)PrimaryDisplayHeight * HalfUnsafeRatio);
 	}
 }
 
diff --git a/Engine/Source/Runtime/ApplicationCore/Private/GenericPlatform/GenericPlatformSoftwareCursor.cpp b/Engine/Source/Runtime/ApplicationCore/Private/GenericPlatform/GenericPlatformSoftwareCursor.cpp
--- a/Engine/Source/Runtime/ApplicationCore/Private/GenericPlatform/GenericPlatformSoftwareCursor.cpp
+++ b/Engine/Source/Runtime/ApplicationCore/Private/GenericPlatform/GenericPlatformSoftwareCursor.cpp
@@ -61,7 +61,7 @@ void FGenericPlatformSoftwareCursor::Lock(const RECT* const Bounds)
 	FVector2D Position = GetPosition();
 	if (UpdateCursorClipping(Position))
 	{
-		SetPosition(Position.X, Position.Y);
+		SetPosition((int32)Position.X, (int32)Position.Y);
 	}
 }
 
diff --git a/Engine/Source/Runtime/ApplicationCore/Private/HAL/PlatformFramePacer.cpp b/Engine/Source/Runtime/ApplicationCore/Private/HAL/PlatformFramePacer.cpp
--- a/Engine/Source/Runtime/ApplicationCore/Private/HAL/PlatformFramePacer.cpp
+++ b/Engine/Source/Runtime/ApplicationCore/Private/HAL/PlatformFramePacer.cpp
@@ -13,8 +13,8 @@ static FAutoConsoleCommand CFramePace(
 {
 	if (Args.Num() > 0)
 	{
-		uint32 RequestFramePace = FCString::Atod(*Args[0]);
-		uint32 ResultFramePace = FPlatformRHIFramePacer::SetFramePace(RequestFramePace);
+		int32 RequestFramePace = (int32)FCString::Atod(*Args[0]);
+		int32 ResultFramePace = FPlatformRHIFramePacer::SetFramePace(RequestFramePace);
 		UE_LOG(LogFramePacer, Display, TEXT("r.FramePace : requesting %d, set as %d"), RequestFramePace, ResultFramePace);
 	}
 	else
diff --git a/Engine/Source/Runtime/ApplicationCore/Private/Windows/Accessibility/WindowsUIAWidgetProvider.cpp b/Engine/Source/Runtime/ApplicationCore/Private/Windows/Accessibility/WindowsUIAWidgetProvider.cpp
--- a/Engine/Source/Runtime/ApplicationCore/Private/Windows/Accessibility/WindowsUIAWidgetProvider.cpp
+++ b/Engine/Source/Runtime/ApplicationCore/Private/Windows/Accessibility/WindowsUIAWidgetProvider.cpp
@@ -936,7 +936,7 @@ HRESULT STDCALL FWindowsUIAWindowProvider::ElementProviderFromPoint(double x, do
 		[this, &ReturnValue, &pRetVal, &x, &y]() {
 		if (IsValid())
 		{
-			TSharedPtr<IAccessibleWidget> Child = Widget->AsWindow()->GetChildAtPosition(x, y);
+			TSharedPtr<IAccessibleWidget> Child = Widget->AsWindow()->GetChildAtPosition((int32)x, (int32)y);
 			if (Child.IsValid())
 			{
 				*pRetVal = static_cast<IRawElementProviderFragment*>(&UIAManager->GetWidgetProvider(Child.ToSharedRef()));
diff --git a/Engine/Source/Runtime/ApplicationCore/Private/Windows/TextStoreACP.cpp b/Engine/Source/Runtime/ApplicationCore/Private/Windows/TextStoreACP.cpp
--- a/Engine/Source/Runtime/ApplicationCore/Private/Windows/TextStoreACP.cpp
+++ b/Engine/Source/Runtime/ApplicationCore/Private/Windows/TextStoreACP.cpp
@@ -473,10 +473,10 @@ STDAPI FTextStoreACP::GetTextExt(TsViewCookie vcView, LONG acpStart, LONG acpEnd
 	FVector2D Position;
 	FVector2D Size;
 	*pfClipped	=	TextContext->GetTextBounds(acpStart, acpEnd - acpStart, Position, Size);
-	prc->left	=	Position.X;
-	prc->top	=	Position.Y;
-	prc->right	=	Position.X + Size.X;
-	prc->bottom	=	Position.Y + Size.Y;
+	prc->left	=	(LONG)Position.X;
+	prc->top	=	(LONG)Position.Y;
+	prc->right	=	(LONG)(Position.X + Size.X);
+	prc->bottom	=	(LONG)(Position.Y + Size.Y);
 
 	return S_OK;
 }
@@ -494,10 +494,10 @@ STDAPI FTextStoreACP::GetScreenExt(TsViewCookie vcView, __RPC__out RECT *prc)
 	FVector2D Position;
 	FVector2D Size;
 	TextContext->GetScreenBounds(Position, Size);
-	prc->left	=	Position.X;
-	prc->top	=	Position.Y;
-	prc->right	=	Position.X + Size.X;
-	prc->bottom	=	Position.Y + Size.Y;
+	prc->left	=	(LONG)Position.X;
+	prc->top	=	(LONG)Position.Y;
+	prc->right	=	(LONG)(Position.X + Size.X);
+	prc->bottom	=	(LONG)(Position.Y + Size.Y);
 
 	return S_OK;
 }
diff --git a/Engine/Source/Runtime/ApplicationCore/Private/Windows/WindowsApplication.cpp b/Engine/Source/Runtime/ApplicationCore/Private/Windows/WindowsApplication.cpp
--- a/Engine/Source/Runtime/ApplicationCore/Private/Windows/WindowsApplication.cpp
+++ b/Engine/Source/Runtime/ApplicationCore/Private/Windows/WindowsApplication.cpp
@@ -818,7 +818,7 @@ static void GetMonitorsInfo(TArray<FMonitorInfo>& OutMonitorInfo)
 							const float CenterX = 0.5f * (Info.WorkArea.Right + Info.WorkArea.Left);
 							const float CenterY = 0.5f * (Info.WorkArea.Top + Info.WorkArea.Bottom);
 							const float DPIScaleFactor = FWindowsPlatformApplicationMisc::GetDPIScaleFactorAtPoint(CenterX, CenterY);
-							Info.DPI *= DPIScaleFactor;
+							Info.DPI = (int32)(Info.DPI * DPIScaleFactor);
 						}
 
 						// The editor shouldn't care about lower level display resolutions. This is only necessary for fullscreen exclusive situations.
@@ -888,7 +888,7 @@ EWindowTitleAlignment::Type FWindowsApplication::GetWindowTitleAlignment() const
 	VersionInfo.dwOSVersionInfoSize = sizeof(VersionInfo);
 
 	DWORDLONG LongConditionMask = 0;
-	int ConditionMask = VER_GREATER_EQUAL;
+	BYTE ConditionMask = VER_GREATER_EQUAL;
 	VER_SET_CONDITION(LongConditionMask, VER_MAJORVERSION, ConditionMask);
 	VER_SET_CONDITION(LongConditionMask, VER_MINORVERSION, ConditionMask);
 
@@ -1023,11 +1023,11 @@ int32 FWindowsApplication::ProcessMessage( HWND hwnd, uint32 msg, WPARAM wParam,
 			DeferMessage( CurrentNativeEventWindowPtr, hwnd, msg, wParam, lParam );
 			return 0;
 		case WM_IME_NOTIFY:
-			UE_LOG(LogWindowsDesktop, Verbose, TEXT("WM_IME_NOTIFY - %s"), IMNStrings.Find(wParam) ? *(IMNStrings[wParam]) : nullptr);
+			UE_LOG(LogWindowsDesktop, Verbose, TEXT("WM_IME_NOTIFY - %s"), IMNStrings.Find(IntCastChecked<uint32>(wParam)) ? *(IMNStrings[(uint32)wParam]) : nullptr);
 			DeferMessage( CurrentNativeEventWindowPtr, hwnd, msg, wParam, lParam );
 			return 0;
 		case WM_IME_REQUEST:
-			UE_LOG(LogWindowsDesktop, Verbose, TEXT("WM_IME_REQUEST - %s"), IMRStrings.Find(wParam) ? *(IMRStrings[wParam]) : nullptr);
+			UE_LOG(LogWindowsDesktop, Verbose, TEXT("WM_IME_REQUEST - %s"), IMRStrings.Find(IntCastChecked<uint32>(wParam)) ? *(IMRStrings[(uint32)wParam]) : nullptr);
 			DeferMessage( CurrentNativeEventWindowPtr, hwnd, msg, wParam, lParam );
 			return 0;
 			// Character
@@ -1360,14 +1360,14 @@ int32 FWindowsApplication::ProcessMessage( HWND hwnd, uint32 msg, WPARAM wParam,
 						HANDLE hTimer = CreateWaitableTimer(NULL, TRUE, NULL);
 						if (NULL != hTimer)
 						{
-							float WaitTimeMilliseconds = 0.0;
+							double WaitTimeMilliseconds = 0.0;
 							if (Freq.QuadPart > 0)
 							{
-								WaitTimeMilliseconds = 1000.0 * WaitTime / (float)Freq.QuadPart;
+								WaitTimeMilliseconds = 1000.0 * WaitTime / (double)Freq.QuadPart;
 							}
 
 							// Due time for WaitForSingleObject is in 100 nanosecond units.							
-							float WaitTime100NanoSeconds = (1000.0f * 10.0f * WaitTimeMilliseconds);
+							double WaitTime100NanoSeconds = (1000.0f * 10.0f * WaitTimeMilliseconds);
 
 							LARGE_INTEGER DueTime;
 
@@ -1451,10 +1451,10 @@ int32 FWindowsApplication::ProcessMessage( HWND hwnd, uint32 msg, WPARAM wParam,
 					case WMSZ_TOPLEFT:
 					case WMSZ_TOPRIGHT:
 					{
-						int32 MinWidth = SizeLimits.GetMinWidth().GetValue();
+						int32 MinWidth = (int32)SizeLimits.GetMinWidth().GetValue();
 						if (SizeLimits.GetMinHeight().GetValue() < SizeLimits.GetMinWidth().GetValue())
 						{
-							MinWidth = SizeLimits.GetMinHeight().GetValue() * AspectRatio;
+							MinWidth = (int32)(SizeLimits.GetMinHeight().GetValue() * AspectRatio);
 						}
 
 						if (NewWidth < MinWidth)
@@ -1476,10 +1476,10 @@ int32 FWindowsApplication::ProcessMessage( HWND hwnd, uint32 msg, WPARAM wParam,
 					case WMSZ_TOP:
 					case WMSZ_BOTTOM:
 					{
-						int32 MinHeight = SizeLimits.GetMinHeight().GetValue();
+						int32 MinHeight = (int32)SizeLimits.GetMinHeight().GetValue();
 						if (SizeLimits.GetMinWidth().GetValue() < SizeLimits.GetMinHeight().GetValue())
 						{
-							MinHeight = SizeLimits.GetMinWidth().GetValue() / AspectRatio;
+							MinHeight = (int32)(SizeLimits.GetMinWidth().GetValue() / AspectRatio);
 						}
 
 						if (NewHeight < MinHeight)
@@ -1504,7 +1504,7 @@ int32 FWindowsApplication::ProcessMessage( HWND hwnd, uint32 msg, WPARAM wParam,
 					case WMSZ_LEFT:
 					case WMSZ_RIGHT:
 						{
-							int32 AdjustedHeight = NewWidth / AspectRatio;
+							int32 AdjustedHeight = (int32)(NewWidth / AspectRatio);
 							Rect->top -= (AdjustedHeight - NewHeight) / 2;
 							Rect->bottom += (AdjustedHeight - NewHeight) / 2;
 							break;
@@ -1512,32 +1512,32 @@ int32 FWindowsApplication::ProcessMessage( HWND hwnd, uint32 msg, WPARAM wParam,
 					case WMSZ_TOP:
 					case WMSZ_BOTTOM:
 						{
-							int32 AdjustedWidth = NewHeight * AspectRatio;
+							int32 AdjustedWidth = (int32)(NewHeight * AspectRatio);
 							Rect->left -= (AdjustedWidth - NewWidth) / 2;
 							Rect->right += (AdjustedWidth - NewWidth) / 2;
 							break;
 						}
 					case WMSZ_TOPLEFT:
 						{
-							int32 AdjustedHeight = NewWidth / AspectRatio;
+							int32 AdjustedHeight = (int32)(NewWidth / AspectRatio);
 							Rect->top -= AdjustedHeight - NewHeight;
 							break;
 						}
 					case WMSZ_TOPRIGHT:
 						{
-							int32 AdjustedHeight = NewWidth / AspectRatio;
+							int32 AdjustedHeight = (int32)(NewWidth / AspectRatio);
 							Rect->top -= AdjustedHeight - NewHeight;
 							break;
 						}
 					case WMSZ_BOTTOMLEFT:
 						{
-							int32 AdjustedHeight = NewWidth / AspectRatio;
+							int32 AdjustedHeight = (int32)(NewWidth / AspectRatio);
 							Rect->bottom += AdjustedHeight - NewHeight;
 							break;
 						}
 					case WMSZ_BOTTOMRIGHT:
 						{
-							int32 AdjustedHeight = NewWidth / AspectRatio;
+							int32 AdjustedHeight = (int32)(NewWidth / AspectRatio);
 							Rect->bottom += AdjustedHeight - NewHeight;
 							break;
 						}
@@ -1610,7 +1610,7 @@ int32 FWindowsApplication::ProcessMessage( HWND hwnd, uint32 msg, WPARAM wParam,
 							HTRIGHT, HTBOTTOMLEFT, HTBOTTOM, HTBOTTOMRIGHT,
 							HTCAPTION, HTMINBUTTON, HTMAXBUTTON, HTCLOSE, HTSYSMENU};
 
-						return Results[Zone];
+						return IntCastChecked<int32>(Results[Zone]);
 					}
 				}
 			}
@@ -1721,7 +1721,7 @@ int32 FWindowsApplication::ProcessMessage( HWND hwnd, uint32 msg, WPARAM wParam,
 				{
 					FCoreDelegates::OnUserLoginChangedEvent.Broadcast(false, 0, 0);
 				}
-				return DefWindowProc(hwnd, msg, wParam, lParam);
+				return IntCastChecked<int32>(DefWindowProc(hwnd, msg, wParam, lParam));
 			}
 			break;
 
@@ -1797,10 +1797,10 @@ int32 FWindowsApplication::ProcessMessage( HWND hwnd, uint32 msg, WPARAM wParam,
 				const float DPIScaleFactor = CurrentNativeEventWindow->GetDPIScaleFactor();
 
 				// We always apply BorderWidth and BorderHeight since Slate always works with client area window sizes
-				MinMaxInfo->ptMinTrackSize.x = FMath::RoundToInt( SizeLimits.GetMinWidth().Get(MinMaxInfo->ptMinTrackSize.x) * DPIScaleFactor );
-				MinMaxInfo->ptMinTrackSize.y = FMath::RoundToInt( SizeLimits.GetMinHeight().Get(MinMaxInfo->ptMinTrackSize.y) * DPIScaleFactor );
-				MinMaxInfo->ptMaxTrackSize.x = FMath::RoundToInt( SizeLimits.GetMaxWidth().Get(MinMaxInfo->ptMaxTrackSize.x) * DPIScaleFactor ) + BorderWidth;
-				MinMaxInfo->ptMaxTrackSize.y = FMath::RoundToInt( SizeLimits.GetMaxHeight().Get(MinMaxInfo->ptMaxTrackSize.y) * DPIScaleFactor ) + BorderHeight;
+				MinMaxInfo->ptMinTrackSize.x = FMath::RoundToInt( SizeLimits.GetMinWidth().Get((float)MinMaxInfo->ptMinTrackSize.x) * DPIScaleFactor );
+				MinMaxInfo->ptMinTrackSize.y = FMath::RoundToInt( SizeLimits.GetMinHeight().Get((float)MinMaxInfo->ptMinTrackSize.y) * DPIScaleFactor );
+				MinMaxInfo->ptMaxTrackSize.x = FMath::RoundToInt( SizeLimits.GetMaxWidth().Get((float)MinMaxInfo->ptMaxTrackSize.x) * DPIScaleFactor ) + BorderWidth;
+				MinMaxInfo->ptMaxTrackSize.y = FMath::RoundToInt( SizeLimits.GetMaxHeight().Get((float)MinMaxInfo->ptMaxTrackSize.y) * DPIScaleFactor ) + BorderHeight;
 				return 0;
 			}
 			break;
@@ -1884,7 +1884,7 @@ int32 FWindowsApplication::ProcessMessage( HWND hwnd, uint32 msg, WPARAM wParam,
 			{
 				FScopedWidgetProvider Provider(UIAManager->GetWindowProvider(CurrentNativeEventWindow));
 				LRESULT Result = UiaReturnRawElementProvider(hwnd, wParam, lParam, &Provider.Provider);
-				return Result;
+				return IntCastChecked<int32>(Result);
 			}
 			break;
 		}
@@ -1898,7 +1898,7 @@ int32 FWindowsApplication::ProcessMessage( HWND hwnd, uint32 msg, WPARAM wParam,
 		}
 	}
 
-	return DefWindowProc(hwnd, msg, wParam, lParam);
+	return IntCastChecked<int32>(DefWindowProc(hwnd, msg, wParam, lParam));
 }
 
 void FWindowsApplication::CheckForShiftUpEvents(const int32 KeyCode)
@@ -1964,7 +1964,7 @@ int32 FWindowsApplication::ProcessDeferredMessage( const FDeferredWindowsMessage
 		case WM_CHAR:
 			{
 				// Character code is stored in WPARAM
-				const TCHAR Character = wParam;
+				const TCHAR Character = IntCastChecked<TCHAR>(wParam);
 
 				// LPARAM bit 30 will be ZERO for new presses, or ONE if this is a repeat
 				const bool bIsRepeat = ( lParam & 0x40000000 ) != 0;
@@ -1982,7 +1982,7 @@ int32 FWindowsApplication::ProcessDeferredMessage( const FDeferredWindowsMessage
 		case WM_KEYDOWN:
 			{
 				// Character code is stored in WPARAM
-				const int32 Win32Key = wParam;
+				const int32 Win32Key = IntCastChecked<int32>(wParam);
 
 				// The actual key to use.  Some keys will be translated into other keys. 
 				// I.E VK_CONTROL will be translated to either VK_LCONTROL or VK_RCONTROL as these
@@ -2066,7 +2066,7 @@ int32 FWindowsApplication::ProcessDeferredMessage( const FDeferredWindowsMessage
 		case WM_KEYUP:
 			{
 				// Character code is stored in WPARAM
-				int32 Win32Key = wParam;
+				int32 Win32Key = IntCastChecked<int32>(wParam);
 
 				// The actual key to use.  Some keys will be translated into other keys. 
 				// I.E VK_CONTROL will be translated to either VK_LCONTROL or VK_RCONTROL as these
diff --git a/Engine/Source/Runtime/ApplicationCore/Private/Windows/WindowsConsoleOutputDevice2.cpp b/Engine/Source/Runtime/ApplicationCore/Private/Windows/WindowsConsoleOutputDevice2.cpp
--- a/Engine/Source/Runtime/ApplicationCore/Private/Windows/WindowsConsoleOutputDevice2.cpp
+++ b/Engine/Source/Runtime/ApplicationCore/Private/Windows/WindowsConsoleOutputDevice2.cpp
@@ -683,7 +683,7 @@ public:
 
 	void UpdateTime(uint64 seconds)
 	{
-		FTimespan span = FTimespan::FromSeconds(seconds);
+		FTimespan span = FTimespan::FromSeconds((double)seconds);
 		FString Res;
 		if (int32 days = span.GetDays())
 			Res = FString::Printf(TEXT("%i."), FMath::Abs(days));
@@ -707,7 +707,7 @@ public:
 				return -1;
 		AddedEntryLogVirtualIndex = LogVirtualIndex;
 		const TCHAR* Empty = TEXT("");
-		int32 ItemIndex = SendMessageW(LogHwnd, LB_ADDSTRING, 0, (LPARAM)Empty);
+		int32 ItemIndex = IntCastChecked<int32>(SendMessageW(LogHwnd, LB_ADDSTRING, 0, (LPARAM)Empty));
 		SendMessageW(LogHwnd, LB_SETITEMDATA, ItemIndex, LogVirtualIndex);
 		AddedEntryLogVirtualIndex = -1;
 		return ItemIndex;
@@ -721,8 +721,8 @@ public:
 		{
 			int32 SelectedItemIndex = -1;
 			SendMessageW(LogHwnd, LB_GETSELITEMS, 1, (LPARAM)&SelectedItemIndex);
-			SelectedLogIndex = SendMessageW(LogHwnd, LB_GETITEMDATA, SelectedItemIndex, 0) - LogIndexOffset;
-			SelectedItemOffsetFromTop = SelectedItemIndex - SendMessageW(LogHwnd, LB_GETTOPINDEX, 0, 0);
+			SelectedLogIndex = IntCastChecked<int32>(SendMessageW(LogHwnd, LB_GETITEMDATA, SelectedItemIndex, 0) - LogIndexOffset);
+			SelectedItemOffsetFromTop = IntCastChecked<int32>(SelectedItemIndex - SendMessageW(LogHwnd, LB_GETTOPINDEX, 0, 0));
 		}
 
 		LogIndexOffset = 0;
@@ -831,7 +831,7 @@ public:
 		if (TotalScrollCount <= 0)
 			return;
 
-		int32 TopVisible = SendMessageW(LogHwnd, LB_GETTOPINDEX, 0, 0);
+		int32 TopVisible = IntCastChecked<int32>(SendMessageW(LogHwnd, LB_GETTOPINDEX, 0, 0));
 
 		int32 ThumbHeight = FMath::Max((PageSize * ScrollHeight) / TotalCount, int32(ButtonHeight));
 
@@ -874,7 +874,7 @@ public:
 	void MoveTopVisible(const RECT& Rect, int32 Offset)
 	{
 		bAutoScrollLog = false;
-		int32 TopVisible = SendMessageW(LogHwnd, LB_GETTOPINDEX, 0, 0);
+		int32 TopVisible = IntCastChecked<int32>(SendMessageW(LogHwnd, LB_GETTOPINDEX, 0, 0));
 		int32 PageSize = Rect.bottom / LogFontHeight;
 		TopVisible = FMath::Max(0, TopVisible + Offset);
 		SetTopVisible(TopVisible, false);
@@ -1148,7 +1148,7 @@ public:
 
 	int32 GetSelectedItems(TArray<int>& OutSelectedItems)
 	{
-		int32 SelectionCount = SendMessageW(LogHwnd, LB_GETSELCOUNT, 0, 0);
+		int32 SelectionCount = IntCastChecked<int32>(SendMessageW(LogHwnd, LB_GETSELCOUNT, 0, 0));
 		OutSelectedItems.SetNum(SelectionCount);
 		SendMessageW(LogHwnd, LB_GETSELITEMS, SelectionCount, (LPARAM)OutSelectedItems.GetData());
 		return SelectionCount;
@@ -1165,7 +1165,7 @@ public:
 		for (int32 I = 0; I != SelectionCount; ++I)
 		{
 			int32 Index = SelectedItems[I];
-			int32 LogIndex = SendMessageW(LogHwnd, LB_GETITEMDATA, Index, 0) - LogIndexOffset;
+			int32 LogIndex = IntCastChecked<int32>(SendMessageW(LogHwnd, LB_GETITEMDATA, Index, 0) - LogIndexOffset);
 			StringBuilder.Append(Log[LogIndex].String);
 			StringBuilder.Append(TEXT("\r\n"));
 		}
@@ -1222,7 +1222,7 @@ public:
 		{
 			bAutoScrollLog = false;
 			SuspendAddingEntries();
-			RightClickedItem = SendMessageW(hWnd, LB_ITEMFROMPOINT, 0, lParam);
+			RightClickedItem = IntCastChecked<int32>(SendMessageW(hWnd, LB_ITEMFROMPOINT, 0, lParam));
 
 			bool AlreadySelected = false;
 			TArray<int> SelectedItems;
@@ -1243,7 +1243,7 @@ public:
 			POINT MousePos{ ((int)(short)LOWORD(lParam)), ((int)(short)HIWORD(lParam)) };
 			TStringBuilder<1024> OutString;
 
-			int32 LogIndex = SendMessageW(LogHwnd, LB_GETITEMDATA, RightClickedItem, 0) - LogIndexOffset;
+			int32 LogIndex = IntCastChecked<int32>(SendMessageW(LogHwnd, LB_GETITEMDATA, RightClickedItem, 0) - LogIndexOffset);
 			RightClickedItem = -1;
 			if (LogIndex < 0)
 			{
@@ -1267,14 +1267,14 @@ public:
 				Info.fType = MFT_STRING;
 				Info.dwTypeData = (TCHAR*)Str;
 				Info.wID = id;
-				InsertMenuItem(Menu, Info.dwItemData++, 1, &Info);
+				InsertMenuItem(Menu, IntCastChecked<UINT>(Info.dwItemData++), 1, &Info);
 			};
 
 			auto AddSeparator = [&]()
 			{
 				Info.fMask = 0;
 				Info.fType = MFT_SEPARATOR;
-				InsertMenuItem(Menu, Info.dwItemData++, 1, &Info);
+				InsertMenuItem(Menu, IntCastChecked<UINT>(Info.dwItemData++), 1, &Info);
 			};
 
 			AddItem(TEXT("Activate auto scroll"), 101);
@@ -1316,7 +1316,7 @@ public:
 				}
 				if (Start != End)
 				{
-					SelectedWord = FStringView(Start, End - Start);
+					SelectedWord = FStringView(Start, UE_PTRDIFF_TO_INT32(End - Start));
 
 					AddSeparator();
 					AddItem(*TStringBuilder<64>().Append("Copy word '").Append(SelectedWord).Append("'"), 120);
@@ -1352,7 +1352,7 @@ public:
 				int32 SelectedItemIndex = -1;
 				if (SendMessageW(LogHwnd, LB_GETSELITEMS, 1, (LPARAM)&SelectedItemIndex) != 1)
 					break;
-				int32 LogIndex = SendMessageW(LogHwnd, LB_GETITEMDATA, SelectedItemIndex, 0) - LogIndexOffset;
+				int32 LogIndex = IntCastChecked<int32>(SendMessageW(LogHwnd, LB_GETITEMDATA, SelectedItemIndex, 0) - LogIndexOffset);
 				LogIndexOffset += LogIndex;
 				Log.PopFront(LogIndex);
 				RefreshLogHwnd();
@@ -1415,9 +1415,9 @@ public:
 			{
 				POINT MousePos{ ((int)(short)LOWORD(lParam)), ((int)(short)HIWORD(lParam)) };
 				TipPosition = MousePos;
-				int32 Res = SendMessageW(hWnd, LB_ITEMFROMPOINT, 0, lParam);
+				int32 Res = IntCastChecked<int32>(SendMessageW(hWnd, LB_ITEMFROMPOINT, 0, lParam));
 				int32 ItemIndex = LOWORD(Res);
-				int32 LogIndex = SendMessageW(hWnd, LB_GETITEMDATA, ItemIndex, 0) - LogIndexOffset;
+				int32 LogIndex = IntCastChecked<int32>(SendMessageW(hWnd, LB_GETITEMDATA, ItemIndex, 0) - LogIndexOffset);
 				if (LogIndex <= 0 || LogIndex >= Log.Num())
 					break;
 				TStringBuilder<1024> OutString;
@@ -1608,7 +1608,7 @@ public:
 			{
 				if (const TCHAR* LineBreak = FCString::Strchr(SearchStr, '\n'))
 				{
-					int32 Len = LineBreak - SearchStr;
+					int32 Len = UE_PTRDIFF_TO_INT32(LineBreak - SearchStr);
 					if (LineBreak > SearchStr && *(LineBreak - 1) == '\r')
 						--Len;
 					Log.Add({ FString(Len, SearchStr), E.Verbosity, E.Category, E.Time, E.TextAttribute, 1 });
@@ -1643,9 +1643,9 @@ public:
 			int32 ToRemoveFromHwnd = ToChangeInHwnd;
 			if (!bAutoScrollLog)
 			{
-				int32 TopIndex = SendMessageW(LogHwnd, LB_GETTOPINDEX, 0, 0);
+				int32 TopIndex = IntCastChecked<int32>(SendMessageW(LogHwnd, LB_GETTOPINDEX, 0, 0));
 				NewTopIndex = FMath::Max(TopIndex - ToRemoveFromHwnd, 0);
-				int32 CaretIndex = SendMessageW(LogHwnd, LB_GETCARETINDEX, 0, 0);
+				int32 CaretIndex = IntCastChecked<int32>(SendMessageW(LogHwnd, LB_GETCARETINDEX, 0, 0));
 				NewCaretIndex = FMath::Max(CaretIndex - ToRemoveFromHwnd, 0);
 			}
 
@@ -1771,7 +1771,7 @@ public:
 				RECT Rect;
 				GetClientRect(LogHwnd, &Rect);
 				int32 VisibleCount = FMath::CeilToInt(float(Rect.bottom) / LogFontHeight);
-				int32 TopIndex = SendMessageW(LogHwnd, LB_GETTOPINDEX, 0, 0);
+				int32 TopIndex = IntCastChecked<int32>(SendMessageW(LogHwnd, LB_GETTOPINDEX, 0, 0));
 				if (ItemCount - TopIndex < VisibleCount)
 				{
 					Rect.top = (ItemCount - TopIndex) * LogFontHeight;
@@ -1800,7 +1800,7 @@ public:
 			case ID_ADDCHECKPOINTBUTTON:
 			{
 				NMCUSTOMDRAW& nmcd = *(NMCUSTOMDRAW*)hdr;
-				uint32 CtlId = hdr->idFrom;
+				uint32 CtlId = IntCastChecked<uint32>(hdr->idFrom);
 				int32 ItemState = nmcd.uItemState;
 				TCHAR Str[32];
 				int32 StrLen = GetDlgItemText(hWnd, CtlId, Str, 32);
@@ -1895,7 +1895,7 @@ public:
 			case ODA_SELECT:
 			case ODA_DRAWENTIRE:
 			{
-				uint32 LogVirtualIndex = pdis->itemData;
+				uint32 LogVirtualIndex = IntCastChecked<uint32>(pdis->itemData);
 				LogEntry& Entry = Log[LogVirtualIndex - LogIndexOffset];
 
 
@@ -1970,7 +1970,7 @@ public:
 						for (int32 I = 0; I != SelectionCount; ++I)
 						{
 							int32 Index = SelectedItems[I];
-							int32 LogIndex = SendMessageW(LogHwnd, LB_GETITEMDATA, Index, 0) - LogIndexOffset;
+							int32 LogIndex = IntCastChecked<int32>(SendMessageW(LogHwnd, LB_GETITEMDATA, Index, 0) - LogIndexOffset);
 							StringBuilder.Append(Log[LogIndex].String);
 							StringBuilder.Append(TEXT("\r\n"));
 						}
diff --git a/Engine/Source/Runtime/ApplicationCore/Private/Windows/WindowsCursor.cpp b/Engine/Source/Runtime/ApplicationCore/Private/Windows/WindowsCursor.cpp
--- a/Engine/Source/Runtime/ApplicationCore/Private/Windows/WindowsCursor.cpp
+++ b/Engine/Source/Runtime/ApplicationCore/Private/Windows/WindowsCursor.cpp
@@ -207,8 +207,8 @@ void* FWindowsCursor::CreateCursorFromRGBABuffer(const FColor* Pixels, int32 Wid
 
 	ICONINFO IconInfo = { 0 };
 	IconInfo.fIcon = Windows::FALSE;
-	IconInfo.xHotspot = FMath::RoundToInt(InHotSpot.X * Width);
-	IconInfo.yHotspot = FMath::RoundToInt(InHotSpot.Y * Height);
+	IconInfo.xHotspot = IntCastChecked<DWORD>(FMath::RoundToInt(InHotSpot.X * Width));
+	IconInfo.yHotspot = IntCastChecked<DWORD>(FMath::RoundToInt(InHotSpot.Y * Height));
 	IconInfo.hbmColor = CursorColor;
 	IconInfo.hbmMask = CursorColor;
 
diff --git a/Engine/Source/Runtime/ApplicationCore/Private/Windows/WindowsPlatformApplicationMisc.cpp b/Engine/Source/Runtime/ApplicationCore/Private/Windows/WindowsPlatformApplicationMisc.cpp
--- a/Engine/Source/Runtime/ApplicationCore/Private/Windows/WindowsPlatformApplicationMisc.cpp
+++ b/Engine/Source/Runtime/ApplicationCore/Private/Windows/WindowsPlatformApplicationMisc.cpp
@@ -175,14 +175,14 @@ void FWindowsPlatformApplicationMisc::PreventScreenSaver()
 FLinearColor FWindowsPlatformApplicationMisc::GetScreenPixelColor(const FVector2D& InScreenPos, float /*InGamma*/)
 {
 	HDC TempDC = GetDC(HWND_DESKTOP);
-	COLORREF PixelColorRef = GetPixel(TempDC, InScreenPos.X, InScreenPos.Y);
+	COLORREF PixelColorRef = GetPixel(TempDC, (int)InScreenPos.X, (int)InScreenPos.Y);
 
 	ReleaseDC(HWND_DESKTOP, TempDC);
 
 	FColor sRGBScreenColor(
-		(PixelColorRef & 0xFF),
-		((PixelColorRef & 0xFF00) >> 8),
-		((PixelColorRef & 0xFF0000) >> 16),
+		(uint8)(PixelColorRef & 0xFF),
+		(uint8)((PixelColorRef & 0xFF00) >> 8),
+		(uint8)((PixelColorRef & 0xFF0000) >> 16),
 		255);
 
 	// Assume the screen color is coming in as sRGB space
diff --git a/Engine/Source/Runtime/ApplicationCore/Private/Windows/WindowsPlatformSplash.cpp b/Engine/Source/Runtime/ApplicationCore/Private/Windows/WindowsPlatformSplash.cpp
--- a/Engine/Source/Runtime/ApplicationCore/Private/Windows/WindowsPlatformSplash.cpp
+++ b/Engine/Source/Runtime/ApplicationCore/Private/Windows/WindowsPlatformSplash.cpp
@@ -564,11 +564,11 @@ uint32 WINAPI StartSplashScreenThread( LPVOID unused )
 					// Set window to fully transparent to start out
 					const float TimeSinceFadeStart = (float)( FPlatformTime::Seconds() - FadeStartTime );
 					const float FadeAmount = FMath::Clamp( TimeSinceFadeStart / FadeDuration, 0.0f, 1.0f );
-					const int32 NewOpacityByte = 255 * FadeAmount;
+					const int32 NewOpacityByte = (int32)(255 * FadeAmount);
 					if( NewOpacityByte != CurrentOpacityByte )
 					{
 						CurrentOpacityByte = NewOpacityByte;
-						SetLayeredWindowAttributes( GSplashScreenWnd, 0, CurrentOpacityByte, LWA_ALPHA );
+						SetLayeredWindowAttributes( GSplashScreenWnd, 0, (BYTE)CurrentOpacityByte, LWA_ALPHA );
 					}
 
 					// We're still fading, but still yield a timeslice
diff --git a/Engine/Source/Runtime/ApplicationCore/Private/Windows/WindowsTextInputMethodSystem.cpp b/Engine/Source/Runtime/ApplicationCore/Private/Windows/WindowsTextInputMethodSystem.cpp
--- a/Engine/Source/Runtime/ApplicationCore/Private/Windows/WindowsTextInputMethodSystem.cpp
+++ b/Engine/Source/Runtime/ApplicationCore/Private/Windows/WindowsTextInputMethodSystem.cpp
@@ -417,19 +417,19 @@ void FWindowsTextInputMethodSystem::UpdateIMMWindowPositions(HIMC IMMContext)
 		CANDIDATEFORM CandidateForm;
 		CandidateForm.dwIndex = 0;
 		CandidateForm.dwStyle = CFS_EXCLUDE;
-		CandidateForm.ptCurrentPos.x = Position.X;
-		CandidateForm.ptCurrentPos.y = Position.Y;
+		CandidateForm.ptCurrentPos.x = (LONG)Position.X;
+		CandidateForm.ptCurrentPos.y = (LONG)Position.Y;
 		CandidateForm.rcArea.left = CandidateForm.ptCurrentPos.x;
 		CandidateForm.rcArea.right = CandidateForm.ptCurrentPos.x;
 		CandidateForm.rcArea.top = CandidateForm.ptCurrentPos.y;
-		CandidateForm.rcArea.bottom = CandidateForm.ptCurrentPos.y + Size.Y;
+		CandidateForm.rcArea.bottom = CandidateForm.ptCurrentPos.y + (LONG)Size.Y;
 		::ImmSetCandidateWindow(IMMContext, &CandidateForm);
 
 		// Update composition window position.
 		COMPOSITIONFORM CompositionForm;
 		CompositionForm.dwStyle = CFS_POINT;
-		CompositionForm.ptCurrentPos.x = Position.X;
-		CompositionForm.ptCurrentPos.y = Position.Y + Size.Y;
+		CompositionForm.ptCurrentPos.x = (LONG)Position.X;
+		CompositionForm.ptCurrentPos.y = (LONG)(Position.Y + Size.Y);
 		::ImmSetCompositionWindow(IMMContext, &CompositionForm);
 	}
 }
@@ -962,7 +962,7 @@ int32 FWindowsTextInputMethodSystem::ProcessMessage(HWND hwnd, uint32 msg, WPARA
 	SCOPE_CYCLE_COUNTER(STAT_IMEWindowsProcessMessage);
 	if(CurrentAPI != EAPI::IMM)
 	{
-		return DefWindowProc(hwnd, msg, wParam, lParam);
+		return IntCastChecked<int32>(DefWindowProc(hwnd, msg, wParam, lParam));
 	}
 
 	switch(msg)
@@ -974,7 +974,7 @@ int32 FWindowsTextInputMethodSystem::ProcessMessage(HWND hwnd, uint32 msg, WPARA
 
 			UpdateIMMProperty(KeyboardLayoutHandle);
 
-			return DefWindowProc(hwnd, msg, wParam, lParam);
+			return IntCastChecked<int32>(DefWindowProc(hwnd, msg, wParam, lParam));
 		}
 		break;
 	case WM_IME_SETCONTEXT:
@@ -990,17 +990,17 @@ int32 FWindowsTextInputMethodSystem::ProcessMessage(HWND hwnd, uint32 msg, WPARA
 				UE_LOG(LogWindowsTextInputMethodSystem, Verbose, TEXT("Setting IMM context."));
 			}
 
-			return DefWindowProc(hwnd, msg, wParam, lParam);
+			return IntCastChecked<int32>(DefWindowProc(hwnd, msg, wParam, lParam));
 		}
 		break;
 	case WM_IME_NOTIFY:
 		{
-			return DefWindowProc(hwnd, msg, wParam, lParam);
+			return IntCastChecked<int32>(DefWindowProc(hwnd, msg, wParam, lParam));
 		}
 		break;
 	case WM_IME_REQUEST:
 		{
-			return DefWindowProc(hwnd, msg, wParam, lParam);
+			return IntCastChecked<int32>(DefWindowProc(hwnd, msg, wParam, lParam));
 		}
 		break;
 	case WM_IME_STARTCOMPOSITION:
@@ -1012,7 +1012,7 @@ int32 FWindowsTextInputMethodSystem::ProcessMessage(HWND hwnd, uint32 msg, WPARA
 				UE_LOG(LogWindowsTextInputMethodSystem, Verbose, TEXT("Beginning IMM composition."));
 			}
 
-			return DefWindowProc(hwnd, msg, wParam, lParam);
+			return IntCastChecked<int32>(DefWindowProc(hwnd, msg, wParam, lParam));
 		}
 		break;
 	case WM_IME_COMPOSITION:
@@ -1131,7 +1131,7 @@ int32 FWindowsTextInputMethodSystem::ProcessMessage(HWND hwnd, uint32 msg, WPARA
 				UE_LOG(LogWindowsTextInputMethodSystem, Verbose, TEXT("Updating IMM composition."));
 			}
 
-			return DefWindowProc(hwnd, msg, wParam, lParam);
+			return IntCastChecked<int32>(DefWindowProc(hwnd, msg, wParam, lParam));
 		}
 		break;
 	case WM_IME_ENDCOMPOSITION:
@@ -1144,7 +1144,7 @@ int32 FWindowsTextInputMethodSystem::ProcessMessage(HWND hwnd, uint32 msg, WPARA
 				UE_LOG(LogWindowsTextInputMethodSystem, Verbose, TEXT("Ending IMM composition."));
 			}
 
-			return DefWindowProc(hwnd, msg, wParam, lParam);
+			return IntCastChecked<int32>(DefWindowProc(hwnd, msg, wParam, lParam));
 		}
 		break;
 	case WM_IME_CHAR:
@@ -1158,7 +1158,7 @@ int32 FWindowsTextInputMethodSystem::ProcessMessage(HWND hwnd, uint32 msg, WPARA
 		{
 			UE_LOG(LogWindowsTextInputMethodSystem, Warning, TEXT("Unexpected windows message received for processing."));
 
-			return DefWindowProc(hwnd, msg, wParam, lParam);
+			return IntCastChecked<int32>(DefWindowProc(hwnd, msg, wParam, lParam));
 		}
 		break;
 	}
diff --git a/Engine/Source/Runtime/ApplicationCore/Private/Windows/WindowsWindow.cpp b/Engine/Source/Runtime/ApplicationCore/Private/Windows/WindowsWindow.cpp
--- a/Engine/Source/Runtime/ApplicationCore/Private/Windows/WindowsWindow.cpp
+++ b/Engine/Source/Runtime/ApplicationCore/Private/Windows/WindowsWindow.cpp
@@ -983,7 +983,7 @@ void FWindowsWindow::SetWindowFocus()
  */
 void FWindowsWindow::SetOpacity( const float InOpacity )
 {
-	SetLayeredWindowAttributes( HWnd, 0, FMath::TruncToInt( InOpacity * 255.0f ), LWA_ALPHA );
+	SetLayeredWindowAttributes( HWnd, 0, (BYTE)FMath::TruncToInt( InOpacity * 255.0f ), LWA_ALPHA );
 }
 
 /**
diff --git a/Engine/Source/Runtime/Online/SSL/Private/SslCertificateManager.cpp b/Engine/Source/Runtime/Online/SSL/Private/SslCertificateManager.cpp
--- a/Engine/Source/Runtime/Online/SSL/Private/SslCertificateManager.cpp
+++ b/Engine/Source/Runtime/Online/SSL/Private/SslCertificateManager.cpp
@@ -87,8 +87,8 @@ static bool DomainLessThan(const FString& DomainA, const FString& DomainB)
 		if (bDomainAIncludesSubdomains)
 		{
 			// both start with '.', sort by number of '.'s
-			const int32 DomainAPeriods = Algo::Count(DomainA, TEXT('.'));
-			const int32 DomainBPeriods = Algo::Count(DomainB, TEXT('.'));
+			const SIZE_T DomainAPeriods = Algo::Count(DomainA, TEXT('.'));
+			const SIZE_T DomainBPeriods = Algo::Count(DomainB, TEXT('.'));
 			if (DomainAPeriods == DomainBPeriods)
 			{
 				// sort alphabetically
diff --git a/Engine/Source/Runtime/Projects/Private/ModuleDescriptor.cpp b/Engine/Source/Runtime/Projects/Private/ModuleDescriptor.cpp
--- a/Engine/Source/Runtime/Projects/Private/ModuleDescriptor.cpp
+++ b/Engine/Source/Runtime/Projects/Private/ModuleDescriptor.cpp
@@ -670,7 +670,7 @@ bool FModuleDescriptor::IsLoadedInCurrentConfiguration() const
 
 void FModuleDescriptor::LoadModulesForPhase(ELoadingPhase::Type LoadingPhase, const TArray<FModuleDescriptor>& Modules, TMap<FName, EModuleLoadResult>& ModuleLoadErrors)
 {
-	FScopedSlowTask SlowTask(Modules.Num());
+	FScopedSlowTask SlowTask((float)Modules.Num());
 	for (int Idx = 0; Idx < Modules.Num(); Idx++)
 	{
 		SlowTask.EnterProgressFrame(1);
diff --git a/Engine/Source/Runtime/Projects/Private/PluginManager.cpp b/Engine/Source/Runtime/Projects/Private/PluginManager.cpp
--- a/Engine/Source/Runtime/Projects/Private/PluginManager.cpp
+++ b/Engine/Source/Runtime/Projects/Private/PluginManager.cpp
@@ -2071,7 +2071,7 @@ bool FPluginManager::LoadModulesForEnabledPlugins( const ELoadingPhase::Type Loa
 	}
 	else
 	{
-		FScopedSlowTask SlowTask(AllPlugins.Num());
+		FScopedSlowTask SlowTask((float)AllPlugins.Num());
 
 		// Load plugins!
 		for (const FDiscoveredPluginMap::ElementType& PluginPair : AllPlugins)
diff --git a/Engine/Source/Runtime/Sockets/Private/BSDSockets/IPAddressBSD.cpp b/Engine/Source/Runtime/Sockets/Private/BSDSockets/IPAddressBSD.cpp
--- a/Engine/Source/Runtime/Sockets/Private/BSDSockets/IPAddressBSD.cpp
+++ b/Engine/Source/Runtime/Sockets/Private/BSDSockets/IPAddressBSD.cpp
@@ -316,12 +316,12 @@ void FInternetAddrBSD::SetPort(int32 InPort)
 #if PLATFORM_HAS_BSD_IPV6_SOCKETS
 	if (GetProtocolType() == FNetworkProtocolTypes::IPv6)
 	{
-		((sockaddr_in6*)&Addr)->sin6_port = htons(InPort);
+		((sockaddr_in6*)&Addr)->sin6_port = htons(IntCastChecked<uint16>(InPort));
 		return;
 	}
 #endif
 
-	((sockaddr_in*)&Addr)->sin_port = htons(InPort);
+	((sockaddr_in*)&Addr)->sin_port = htons(IntCastChecked<uint16>(InPort));
 }
 
 int32 FInternetAddrBSD::GetPort() const
diff --git a/Engine/Source/Runtime/Sockets/Private/BSDSockets/SocketSubsystemBSD.cpp b/Engine/Source/Runtime/Sockets/Private/BSDSockets/SocketSubsystemBSD.cpp
--- a/Engine/Source/Runtime/Sockets/Private/BSDSockets/SocketSubsystemBSD.cpp
+++ b/Engine/Source/Runtime/Sockets/Private/BSDSockets/SocketSubsystemBSD.cpp
@@ -182,7 +182,8 @@ TSharedPtr<FInternetAddr> FSocketSubsystemBSD::GetAddressFromString(const FStrin
 		UniversalNetBufferPtr = &(((sockaddr_in*)&NetworkBuffer)->sin_addr);
 	}
 
-	NetworkBuffer.ss_family = AddrFamily;
+	// The type of ss_family varies by platform.
+	NetworkBuffer.ss_family = IntCastChecked<decltype(NetworkBuffer.ss_family)>(AddrFamily);
 
 	if (inet_pton(AddrFamily, TCHAR_TO_ANSI(*InAddress), UniversalNetBufferPtr))
 	{
diff --git a/Engine/Source/Runtime/Sockets/Private/BSDSockets/SocketsBSD.cpp b/Engine/Source/Runtime/Sockets/Private/BSDSockets/SocketsBSD.cpp
--- a/Engine/Source/Runtime/Sockets/Private/BSDSockets/SocketsBSD.cpp
+++ b/Engine/Source/Runtime/Sockets/Private/BSDSockets/SocketsBSD.cpp
@@ -604,7 +604,8 @@ bool FSocketBSD::SetLinger(bool bShouldLinger,int32 Timeout)
 	linger ling;
 
 	ling.l_onoff = bShouldLinger;
-	ling.l_linger = Timeout;
+	// The type of l_linger varies by platform.
+	ling.l_linger = IntCastChecked<decltype(ling.l_linger)>(Timeout);
 
 	return setsockopt(Socket,SOL_SOCKET,SO_LINGER,(char*)&ling,sizeof(ling)) == 0;
 }
@@ -765,15 +766,15 @@ ESocketBSDReturn FSocketBSD::HasState(ESocketBSDParam State, FTimespan WaitTime)
 	switch (State)
 	{
 	case ESocketBSDParam::CanRead:
-		SelectStatus = select(Socket + 1, &SocketSet, NULL, NULL, TimePointer);
+		SelectStatus = select(IntCastChecked<int>(Socket + 1), &SocketSet, NULL, NULL, TimePointer);
 		break;
 
 	case ESocketBSDParam::CanWrite:
-		SelectStatus = select(Socket + 1, NULL, &SocketSet, NULL, TimePointer);
+		SelectStatus = select(IntCastChecked<int>(Socket + 1), NULL, &SocketSet, NULL, TimePointer);
 		break;
 
 	case ESocketBSDParam::HasError:
-		SelectStatus = select(Socket + 1, NULL, NULL, &SocketSet, TimePointer);
+		SelectStatus = select(IntCastChecked<int>(Socket + 1), NULL, NULL, &SocketSet, TimePointer);
 		break;
 	}
 
diff --git a/Engine/Source/Runtime/Sockets/Private/SocketSubsystem.cpp b/Engine/Source/Runtime/Sockets/Private/SocketSubsystem.cpp
--- a/Engine/Source/Runtime/Sockets/Private/SocketSubsystem.cpp
+++ b/Engine/Source/Runtime/Sockets/Private/SocketSubsystem.cpp
@@ -90,7 +90,8 @@ static bool GetLocalHostAddrViaConnect(TSharedRef<FInternetAddr>& IO)
 	Addr.sin_port = 0x0100;
 	socklen_t AddrSize = sizeof(Addr);
 
-	int Socket = socket(AF_INET, SOCK_DGRAM, 0);
+	// The type returned by socket varies by platform.
+	auto Socket = socket(AF_INET, SOCK_DGRAM, 0);
 
 	if ( connect(Socket, (struct sockaddr*)&Addr, sizeof(Addr)) != 0 )
 	{
diff --git a/Engine/Source/Runtime/TraceLog/Private/Trace/Control.cpp b/Engine/Source/Runtime/TraceLog/Private/Trace/Control.cpp
--- a/Engine/Source/Runtime/TraceLog/Private/Trace/Control.cpp
+++ b/Engine/Source/Runtime/TraceLog/Private/Trace/Control.cpp
@@ -54,7 +54,7 @@ static FControlCommands	GControlCommands;
 static UPTRINT			GControlListen		= 0;
 static UPTRINT			GControlSocket		= 0;
 static EControlState	GControlState;		// = EControlState::Closed;
-static uint32			GControlPort		= 1985;
+static uint16			GControlPort		= 1985;
 
 ////////////////////////////////////////////////////////////////////////////////
 static uint32 Writer_ControlHash(const ANSICHAR* Word)
@@ -114,7 +114,7 @@ static bool Writer_ControlListen()
 		uint32 Seed = uint32(TimeGetTimestamp());
 		for (uint32 i = 0; i < 10 && !GControlListen; Seed *= 13, ++i)
 		{
-			uint32 Port = (Seed & 0x1fff) + 0x8000;
+			uint16 Port((Seed & 0x1fff) + 0x8000);
 			GControlListen = TcpSocketListen(Port);
 			if (GControlListen)
 			{
diff --git a/Engine/Source/Runtime/TraceLog/Private/Trace/Writer.cpp b/Engine/Source/Runtime/TraceLog/Private/Trace/Writer.cpp
--- a/Engine/Source/Runtime/TraceLog/Private/Trace/Writer.cpp
+++ b/Engine/Source/Runtime/TraceLog/Private/Trace/Writer.cpp
@@ -841,7 +841,7 @@ struct WorkerUpdateLock
 
 	bool TimedOut()
 	{
-		uint64 WaitTime = GetTime() - StartSeconds;
+		const double WaitTime = GetTime() - StartSeconds;
 		return WaitTime > MaxWaitSeconds;
 	}
 
