[Runtime]
* TrackedActivity - Added sort value parameter when creating tracked activities
* WindowsConsoleOutputDevice2 - Added sorting of activities based on sort value.
* WindowsConsoleOutputDevice2 - Added show/hide timedate options to log context menu

diff --git a/Engine/Source/Runtime/ApplicationCore/Private/Windows/WindowsConsoleOutputDevice2.cpp b/Engine/Source/Runtime/ApplicationCore/Private/Windows/WindowsConsoleOutputDevice2.cpp
--- a/Engine/Source/Runtime/ApplicationCore/Private/Windows/WindowsConsoleOutputDevice2.cpp
+++ b/Engine/Source/Runtime/ApplicationCore/Private/Windows/WindowsConsoleOutputDevice2.cpp
@@ -115,7 +115,7 @@ public:
 	}
 
 	// Light 0 = none, 1 = Red, 2 = Yellow, 3 = Green
-	void SetActivity(const TCHAR* Id, const TCHAR* Name, const TCHAR* Status, int32 Light = 0, bool bAlignLeft = true)
+	void SetActivity(const TCHAR* Id, const TCHAR* Name, const TCHAR* Status, int32 Light, int32 SortValue, bool bAlignLeft = true)
 	{
 		if (!*Name)
 			return;
@@ -141,6 +141,7 @@ public:
 		if (IsNew)
 		{
 			Modification.Name = Name;
+			Modification.SortValue = SortValue;
 			Modification.bAlignLeft = bAlignLeft;
 		}
 		Modification.Status = Status;
@@ -256,7 +257,7 @@ public:
 				FCString::Sprintf(Id, TEXT("%u%s"), Info.Id, Str);
 
 				if (Event != FTrackedActivity::EEvent::Removed)
-					SetActivity(Id, Str, Info.Status, (int32)Info.Light, Info.Type == FTrackedActivity::EType::Activity);
+					SetActivity(Id, Str, Info.Status, (int32)Info.Light, Info.SortValue, Info.Type == FTrackedActivity::EType::Activity);
 				else
 					RemoveStatus(Id);
 			});
@@ -266,7 +267,7 @@ public:
 				TCHAR Id[256];
 				const TCHAR* Str = Info.Name;
 				FCString::Sprintf(Id, TEXT("%u%s"), Info.Id, Str);
-				SetActivity(Id, Str, Info.Status, (int32)Info.Light);
+				SetActivity(Id, Str, Info.Status, (int32)Info.Light, Info.SortValue);
 			});
 
 		UpdateWindow(MainHwnd);
@@ -526,7 +527,8 @@ public:
 	{
 		Width -= 6; // Margins
 
-		int32 LeftCount = 0;
+		TArray<Activity*, TInlineAllocator<32>> LeftActivities;
+		TArray<Activity*, TInlineAllocator<32>> RightActivities;
 		uint32 ActivityCount = 0;
 		for (int32 i = 0, e = Activities.Num(); i != e; ++i)
 		{
@@ -535,8 +537,22 @@ public:
 				continue;
 			++ActivityCount;
 			if (A.bAlignLeft)
-				++LeftCount;
+				LeftActivities.Add(&A);
+			else
+				RightActivities.Add(&A);
 		}
+
+		auto SortFunc = [](const Activity* A, const Activity* B)
+		{
+			if (A->SortValue != B->SortValue)
+				return A->SortValue < B->SortValue;
+			return uintptr_t(A) < uintptr_t(B);
+		};
+
+		Algo::Sort(LeftActivities, SortFunc);
+		Algo::Sort(RightActivities, SortFunc);
+
+		int32 LeftCount = LeftActivities.Num();
 		int32 RightCount = ActivityCount - LeftCount;
 
 		int32 LeftColumnMinWidth = 600;
@@ -568,12 +584,12 @@ public:
 		int32 ColWidth;
 		int32 ColOffset;
 
-		auto IterateActivities = [&](bool bAlignLeft)
+		auto IterateActivities = [&](TArray<Activity*, TInlineAllocator<32>>& SortedActivities)
 		{
-			for (int32 i = 0, e = Activities.Num(); i != e; ++i)
+			for (int32 i = 0, e = SortedActivities.Num(); i != e; ++i)
 			{
-				Activity& A = Activities[i];
-				if (A.Name.IsEmpty() || A.bAlignLeft != bAlignLeft)
+				Activity& A = *SortedActivities[i];
+				if (A.Name.IsEmpty())
 					continue;
 
 				InLamda(A, X, Y, ColWidth, RowIndex);
@@ -593,7 +609,7 @@ public:
 		int TotalLeftWidth = Width - RightColumnCount * RightColumnWidth - X;
 		ColWidth = TotalLeftWidth / LeftColumnCount;
 		ColOffset = ColWidth - 2;
-		IterateActivities(true);
+		IterateActivities(LeftActivities);
 
 		X = Width - RightColumnWidth + 8;
 		Y = StartY;
@@ -601,7 +617,7 @@ public:
 		ColWidth = RightColumnWidth - 8;
 		ColOffset = -RightColumnWidth;
 
-		IterateActivities(false);
+		IterateActivities(RightActivities);
 		return TotalHeight;
 	}
 
@@ -688,7 +704,7 @@ public:
 		if (int32 days = span.GetDays())
 			Res = FString::Printf(TEXT("%i."), FMath::Abs(days));
 		Res = FString::Printf(TEXT("%s%02i:%02i:%02i"), *Res, FMath::Abs(span.GetHours()), FMath::Abs(span.GetMinutes()), FMath::Abs(span.GetSeconds()));
-		SetActivity(TEXT("Time"), TEXT("Time"), *Res, 0, false);
+		SetActivity(TEXT("Time"), TEXT("Time"), *Res, 0, 1, false);
 	}
 
 	struct LogEntry;
@@ -1089,18 +1105,17 @@ public:
 
 		if (ConsoleShowDateTime == 1)
 		{
-			FTimespan Ts = FTimespan::FromSeconds(FPlatformTime::Seconds() - GStartTime);
-			OutString << '[';
+			FTimespan Ts = FTimespan::FromSeconds(E.Time);
 			OutString.Appendf(TEXT("%02i:"), Ts.GetHours() + Ts.GetDays()*24);
 			OutString.Appendf(TEXT("%02i:"), Ts.GetMinutes());
 			OutString.Appendf(TEXT("%02i."), Ts.GetSeconds());
-			OutString.Appendf(TEXT("%03i]"), Ts.GetFractionMilli());
-			OutString.Appendf(TEXT("[%3llu] "), GFrameCounter % 1000);
+			OutString.Appendf(TEXT("%03i  "), Ts.GetFractionMilli());
+			//OutString.Appendf(TEXT("[%3llu] "), GFrameCounter % 1000);
 		}
 		else if (ConsoleShowDateTime == 2)
 		{
 			FDateTime::Now().ToString(TEXT("[%Y.%m.%d-%H.%M.%S:%s]"), OutString);
-			OutString.Appendf(TEXT("[%3llu] "), GFrameCounter % 1000);
+			//OutString.Appendf(TEXT("[%3llu] "), GFrameCounter % 1000);
 		}
 
 		if (ForFilter)
@@ -1108,8 +1123,6 @@ public:
 			OutString << '[' << CategoryBuilder << TEXT("][") << ToString(E.Verbosity) << ']';
 		}
 
-
-
 		const TCHAR* Category = *CategoryBuilder;
 		if (FCString::Strstr(Category, TEXT("Log")) == Category)
 			Category += 3;
@@ -1326,6 +1339,14 @@ public:
 				}
 			}
 
+			AddSeparator();
+			if (ConsoleShowDateTime != 1)
+				AddItem(TEXT("Show Time"), 131);
+			if (ConsoleShowDateTime != 2)
+				AddItem(TEXT("Show Date and Time"), 132);
+			if (ConsoleShowDateTime != 0)
+				AddItem(TEXT("Hide Date and Time"), 133);
+
 			MapWindowPoints(hWnd, HWND_DESKTOP, &MousePos, 1);
 			TrackPopupMenu(Menu, 0, MousePos.x, MousePos.y, 0, hWnd, NULL);
 			break;
@@ -1393,6 +1414,20 @@ public:
 				RefreshLogHwnd();
 				break;
 			}
+			case 131:
+				ConsoleShowDateTime = 1;
+				RefreshLogHwnd();
+				break;
+
+			case 132:
+				ConsoleShowDateTime = 2;
+				RefreshLogHwnd();
+				break;
+
+			case 133:
+				ConsoleShowDateTime = 0;
+				RefreshLogHwnd();
+				break;
 			}
 			break;
 		}
@@ -1473,7 +1508,7 @@ public:
 
 		// Transfer modifications to Activities and update IdToActivity
 		{
-			TArray<int> IndicesToRemove;
+			TArray<int, TInlineAllocator<32>> IndicesToRemove;
 
 			FScopeLock _(&ActivityModificationsCs);
 
@@ -1492,6 +1527,7 @@ public:
 					Activity& A = Activities[Index];
 					A.Name = Mod.Name;
 					A.Light = Mod.Light;
+					A.SortValue = Mod.SortValue;
 					A.Status = Mod.Status;
 					A.bAlignLeft = Mod.bAlignLeft;
 					continue;
@@ -2376,8 +2412,8 @@ public:
 	FString SelectedWord;
 	FString SelectedCategory;
 
-	struct Activity { HWND NameHwnd = 0; HWND StatusHwnd = 0; FString Name; FString Status; int32 Light = 0; bool bStatusDirty = false; bool bAlignLeft = false; };
-	struct ActivityModification { FString Name; FString Status; int32 Light = 0; bool bAlignLeft = false; bool bRemove = false; };
+	struct Activity { HWND NameHwnd = 0; HWND StatusHwnd = 0; FString Name; FString Status; int32 Light = 0; int32 SortValue; bool bStatusDirty = false; bool bAlignLeft = false; };
+	struct ActivityModification { FString Name; FString Status; int32 Light = 0; int32 SortValue; bool bAlignLeft = false; bool bRemove = false; };
 	TArray<Activity> Activities;
 	TMap<FString, int> IdToActivityIndex;
 	FCriticalSection ActivityModificationsCs;
diff --git a/Engine/Source/Runtime/Core/Private/Misc/TrackedActivity.cpp b/Engine/Source/Runtime/Core/Private/Misc/TrackedActivity.cpp
--- a/Engine/Source/Runtime/Core/Private/Misc/TrackedActivity.cpp
+++ b/Engine/Source/Runtime/Core/Private/Misc/TrackedActivity.cpp
@@ -32,6 +32,7 @@ public:
 		FString Name;
 		TArray<FInfo> Stack;
 		EType Type;
+		int32 SortValue;
 		uint32 Id;
 	};
 
@@ -41,11 +42,12 @@ public:
 	uint32 EventFuncMaxDepth = ~0u;
 	uint32 IdCounter = 1u;
 
-	FActivity* Create(const TCHAR* Name, const TCHAR* Status, ELight Light, EType Type)
+	FActivity* Create(const TCHAR* Name, const TCHAR* Status, ELight Light, EType Type, int32 SortValue)
 	{
 		FActivity* Activity = new FActivity();
 		Activity->Name = Name;
 		Activity->Type = Type;
+		Activity->SortValue = SortValue;
 		Activity->Stack.Add({ Status, Light, false });
 
 		FScopeLock _(&ActivitiesCs);
@@ -132,7 +134,7 @@ public:
 			Light = Activity.Stack[--StackIndex].Light;
 		}
 
-		return { *Activity.Name, Status, Light, Activity.Type, Activity.Id };
+		return { *Activity.Name, Status, Light, Activity.Type, Activity.SortValue, Activity.Id };
 	}
 
 	void TraverseActivities(const TFunction<void(const FTrackedActivity::FInfo& Info)>& Func)
@@ -165,9 +167,9 @@ public:
 };
 
 
-FTrackedActivity::FTrackedActivity(const TCHAR* Name, const TCHAR* Status, ELight Light, EType Type)
+FTrackedActivity::FTrackedActivity(const TCHAR* Name, const TCHAR* Status, ELight Light, EType Type, int32 SortValue)
 {
-	Internal = FTrackedActivityManager::Get().Create(Name, Status, Light, Type);
+	Internal = FTrackedActivityManager::Get().Create(Name, Status, Light, Type, SortValue);
 }
 
 FTrackedActivity::~FTrackedActivity()
@@ -202,15 +204,13 @@ void FTrackedActivity::Update(ELight Light, uint32 Index)
 
 FTrackedActivity& FTrackedActivity::GetEngineActivity()
 {
-	static TSharedPtr<FTrackedActivity> A(MakeShared<FTrackedActivity>(TEXT("Status"), TEXT("Unknown")));
+	static TSharedPtr<FTrackedActivity> A(MakeShared<FTrackedActivity>(TEXT("Status"), TEXT("Unknown"), ELight::None, EType::Activity, 0));
 	return *A;
 }
 
 FTrackedActivity& FTrackedActivity::GetIOActivity()
 {
-	GetEngineActivity(); // Ugly, but just to make sure Status show first
-
-	static TSharedPtr<FTrackedActivity> A(MakeShared<FTrackedActivity>(TEXT("I/O"), TEXT("Idle")));
+	static TSharedPtr<FTrackedActivity> A(MakeShared<FTrackedActivity>(TEXT("I/O"), TEXT("Idle"), ELight::None, EType::Activity, 1));
 	return *A;
 }
 
diff --git a/Engine/Source/Runtime/Core/Public/Misc/TrackedActivity.h b/Engine/Source/Runtime/Core/Public/Misc/TrackedActivity.h
--- a/Engine/Source/Runtime/Core/Public/Misc/TrackedActivity.h
+++ b/Engine/Source/Runtime/Core/Public/Misc/TrackedActivity.h
@@ -40,8 +40,16 @@ public:
 		Debug,
 	};
 
-	/** Ctor/Dtor */
-	FTrackedActivity(const TCHAR* Name, const TCHAR* Status, ELight Light = ELight::Inherit, EType Type = EType::Activity);
+	/** Ctor
+	* @param Name Name of tracked activity
+	* @param Status Initial status of tracked activity
+	* @param Light for the activity. Will show as a dot in front of activity in console
+	* @param Type Decides where activity information show in console. Activity is to the left, Info to the right
+	* @param SortValue Decides in what order within type the activity will show in console. Lower value means earlier
+	*/
+	FTrackedActivity(const TCHAR* Name, const TCHAR* Status = TEXT(""), ELight Light = ELight::None, EType Type = EType::Activity, int32 SortValue = 100);
+
+	/** Dtor */
 	~FTrackedActivity();
 
 
@@ -75,6 +83,7 @@ public:
 		const TCHAR* Status;
 		ELight Light; // Can not be "Inherit"
 		EType Type;
+		int32 SortValue;
 		uint32 Id; // Unique Id for Activity (Is thread safe counter starting at 1)
 	};
 
