Adding a new option to only select users who are not already added to the game (not available on all platforms)

 - also fix bitfield assignment

#rnx

diff --git a/Engine/Source/Runtime/ApplicationCore/Public/GenericPlatform/GenericPlatformApplicationMisc.h b/Engine/Source/Runtime/ApplicationCore/Public/GenericPlatform/GenericPlatformApplicationMisc.h
--- a/Engine/Source/Runtime/ApplicationCore/Public/GenericPlatform/GenericPlatformApplicationMisc.h
+++ b/Engine/Source/Runtime/ApplicationCore/Public/GenericPlatform/GenericPlatformApplicationMisc.h
@@ -57,8 +57,9 @@ enum class EPlatformUserSelectorFlags : uint8
 {
 	None                         = 0,
 	RequiresOnlineEnabledProfile = (1 << 1), // whether to only show profiles that have an online account
-	ShowSkipButton               = (1 << 2), // display a 'skip' button
-	AllowGuests                  = (1 << 4), // temporary guest accounts, throwaways etc.
+	ShowSkipButton               = (1 << 2), // display a 'skip' button UI to be closed & return the initiating user
+	AllowGuests                  = (1 << 3), // temporary guest accounts, throwaways etc.
+	ShowNewUsersOnly             = (1 << 4), // do not show profiles that are already in-use by the title
 
 	Default                      = ShowSkipButton, // most commonly used options
 };
