Adding an option to prevent input devices being automatically unpaired when a user signs out (only settable via specialization.)

#rnx

diff --git a/Engine/Source/Runtime/ApplicationCore/Private/GenericPlatform/GenericPlatformInputDeviceMapper.cpp b/Engine/Source/Runtime/ApplicationCore/Private/GenericPlatform/GenericPlatformInputDeviceMapper.cpp
--- a/Engine/Source/Runtime/ApplicationCore/Private/GenericPlatform/GenericPlatformInputDeviceMapper.cpp
+++ b/Engine/Source/Runtime/ApplicationCore/Private/GenericPlatform/GenericPlatformInputDeviceMapper.cpp
@@ -410,7 +410,7 @@ void FGenericPlatformInputDeviceMapper::OnUserLoginChangedEvent(bool bLoggedIn,
 		// TODO: As of right now there is no logic that needs to run when a new platform user logs in, but there may be in the future
 	}
 	// A platform user has logged out
-	else
+	else if (bUnpairInputDevicesWhenLoggingOut)
 	{
 		// Remap any input devices that the logged out platform user had to the "Unpaired" user
 		FPlatformUserId UnkownUserId = GetUserForUnpairedInputDevices();
diff --git a/Engine/Source/Runtime/ApplicationCore/Public/GenericPlatform/GenericPlatformInputDeviceMapper.h b/Engine/Source/Runtime/ApplicationCore/Public/GenericPlatform/GenericPlatformInputDeviceMapper.h
--- a/Engine/Source/Runtime/ApplicationCore/Public/GenericPlatform/GenericPlatformInputDeviceMapper.h
+++ b/Engine/Source/Runtime/ApplicationCore/Public/GenericPlatform/GenericPlatformInputDeviceMapper.h
@@ -297,7 +297,7 @@ protected:
 
 	/** 
 	* Callback for when FCoreDelegates::FOnUserLoginChangedEvent is broadcasted.
-	* If the user has logged out, then remap any input devices that the user had to the "Unpaired" user on this platform.
+	* If the user has logged out, then remap any input devices that the user had to the "Unpaired" user on this platform (if bUnpairInputDevicesWhenLoggingOut is true)
 	*/
 	virtual void OnUserLoginChangedEvent(bool bLoggedIn, int32 RawPlatformUserId, int32 UserIndex) override;
 	
@@ -310,4 +310,7 @@ protected:
 	/** Flags for backwards compatibility with the older "int32 ControllerId" implementation */
 	const bool bUsingControllerIdAsUserId = true;
 	const bool bShouldBroadcastLegacyDelegates = true;
+
+	/** Additional functionality that specializations can customize */
+	bool bUnpairInputDevicesWhenLoggingOut = true;
 };
\ No newline at end of file
