[Core Runtime]
* WindowsConsoleOutputDevice2.cpp - Fixed bug where activities added before output device itself didn't align Info types to the right

diff --git a/Engine/Source/Runtime/ApplicationCore/Private/Windows/WindowsConsoleOutputDevice2.cpp b/Engine/Source/Runtime/ApplicationCore/Private/Windows/WindowsConsoleOutputDevice2.cpp
--- a/Engine/Source/Runtime/ApplicationCore/Private/Windows/WindowsConsoleOutputDevice2.cpp
+++ b/Engine/Source/Runtime/ApplicationCore/Private/Windows/WindowsConsoleOutputDevice2.cpp
@@ -267,7 +267,7 @@ public:
 				TCHAR Id[256];
 				const TCHAR* Str = Info.Name;
 				FCString::Sprintf(Id, TEXT("%u%s"), Info.Id, Str);
-				SetActivity(Id, Str, Info.Status, (int32)Info.Light, Info.SortValue);
+				SetActivity(Id, Str, Info.Status, (int32)Info.Light, Info.SortValue, Info.Type == FTrackedActivity::EType::Activity);
 			});
 
 		UpdateWindow(MainHwnd);
