Fixed unsafe type casts in Mac platform code

#rnx

diff --git a/Engine/Source/Developer/DesktopPlatform/Private/Mac/DesktopPlatformMac.cpp b/Engine/Source/Developer/DesktopPlatform/Private/Mac/DesktopPlatformMac.cpp
--- a/Engine/Source/Developer/DesktopPlatform/Private/Mac/DesktopPlatformMac.cpp
+++ b/Engine/Source/Developer/DesktopPlatform/Private/Mac/DesktopPlatformMac.cpp
@@ -337,7 +337,7 @@ bool FDesktopPlatformMac::OpenFontDialog(const void* ParentWindowHandle, FString
 				FPlatformString::CFStringToTCHAR((CFStringRef)[Font fontName], FontName);
 
 				OutFontName = FontName;
-				OutHeight = [Font pointSize];
+				OutHeight = (float)[Font pointSize];
 
 				auto FontFlags = EFontImportFlags::None;
 
@@ -376,7 +376,7 @@ bool FDesktopPlatformMac::FileDialogShared(bool bSave, const void* ParentWindowH
 				NSOpenPanel* OpenPanel = (NSOpenPanel*)Panel;
 				[OpenPanel setCanChooseFiles: true];
 				[OpenPanel setCanChooseDirectories: false];
-				[OpenPanel setAllowsMultipleSelection: Flags & EFileDialogFlags::Multiple];
+				[OpenPanel setAllowsMultipleSelection: (Flags & EFileDialogFlags::Multiple) != 0];
 			}
 
 			[Panel setCanCreateDirectories: bSave];
diff --git a/Engine/Source/Runtime/ApplicationCore/Private/Apple/AppleControllerInterface.cpp b/Engine/Source/Runtime/ApplicationCore/Private/Apple/AppleControllerInterface.cpp
--- a/Engine/Source/Runtime/ApplicationCore/Private/Apple/AppleControllerInterface.cpp
+++ b/Engine/Source/Runtime/ApplicationCore/Private/Apple/AppleControllerInterface.cpp
@@ -352,7 +352,7 @@ void FAppleControllerInterface::HandleInputInternal(const FGamepadKeyNames::Type
 {
     const double CurrentTime = FPlatformTime::Seconds();
     const float InitialRepeatDelay = 0.2f;
-    const float RepeatDelay = 0.1;
+    const float RepeatDelay = 0.1f;
     GCController* Cont = Controllers[ControllerIndex].Controller;
     
 	IPlatformInputDeviceMapper& DeviceMapper = IPlatformInputDeviceMapper::Get();
@@ -390,7 +390,7 @@ void FAppleControllerInterface::HandleVirtualButtonGamepad(const FGamepadKeyName
     GCExtendedGamepad *ExtendedPreviousGamepad = Controllers[ControllerIndex].PreviousExtendedGamepad;;
 
     // Send controller events any time we are passed the given input threshold similarly to PC/Console (see: XInputInterface.cpp)
-    const float RepeatDeadzone = 0.24;
+    const float RepeatDeadzone = 0.24f;
     
     bool bWasNegativePressed = false;
     bool bNegativePressed = false;
@@ -512,7 +512,7 @@ void FAppleControllerInterface::HandleAnalogGamepad(const FGamepadKeyNames::Type
 	DeviceMapper.RemapControllerIdToPlatformUserAndDevice(Controllers[ControllerIndex].PlayerIndex, OUT UserId, OUT DeviceId);
     
     // Send controller events any time we are passed the given input threshold similarly to PC/Console (see: XInputInterface.cpp)
-    const float RepeatDeadzone = 0.24;
+    const float RepeatDeadzone = 0.24f;
     bool bWasPositivePressed = false;
     bool bPositivePressed = false;
     bool bWasNegativePressed = false;
diff --git a/Engine/Source/Runtime/ApplicationCore/Private/Mac/CocoaTextView.cpp b/Engine/Source/Runtime/ApplicationCore/Private/Mac/CocoaTextView.cpp
--- a/Engine/Source/Runtime/ApplicationCore/Private/Mac/CocoaTextView.cpp
+++ b/Engine/Source/Runtime/ApplicationCore/Private/Mac/CocoaTextView.cpp
@@ -487,7 +487,7 @@
 		}
 
 		NSScreen* PrimaryScreen = [[self window] screen];
-		const float PrimaryScreenHeight = [PrimaryScreen visibleFrame].size.height;
+		const CGFloat PrimaryScreenHeight = [PrimaryScreen visibleFrame].size.height;
 		Position.Y = -(Position.Y - PrimaryScreenHeight + 1);
 
 		NSRect GlyphBox = {{Position.X,Position.Y},{Size.X,Size.Y}};
diff --git a/Engine/Source/Runtime/ApplicationCore/Private/Mac/HIDInputInterface.h b/Engine/Source/Runtime/ApplicationCore/Private/Mac/HIDInputInterface.h
--- a/Engine/Source/Runtime/ApplicationCore/Private/Mac/HIDInputInterface.h
+++ b/Engine/Source/Runtime/ApplicationCore/Private/Mac/HIDInputInterface.h
@@ -45,8 +45,8 @@ private:
 	{
 		IOHIDElementRef ElementRef;
 		IOHIDElementType Type;
-		uint16 UsagePage;
-		uint16 Usage;
+		uint32 UsagePage;
+		uint32 Usage;
 		int32 MinValue;
 		int32 MaxValue;
 	};
diff --git a/Engine/Source/Runtime/ApplicationCore/Private/Mac/MacApplication.cpp b/Engine/Source/Runtime/ApplicationCore/Private/Mac/MacApplication.cpp
--- a/Engine/Source/Runtime/ApplicationCore/Private/Mac/MacApplication.cpp
+++ b/Engine/Source/Runtime/ApplicationCore/Private/Mac/MacApplication.cpp
@@ -422,10 +422,10 @@ FPlatformRect FMacApplication::GetWorkArea(const FPlatformRect& CurrentWindow) c
 	const NSRect VisibleFrame = Screen->VisibleFramePixels;
 
 	FPlatformRect WorkArea;
-	WorkArea.Left = VisibleFrame.origin.x;
-	WorkArea.Top = VisibleFrame.origin.y;
-	WorkArea.Right = WorkArea.Left + VisibleFrame.size.width;
-	WorkArea.Bottom = WorkArea.Top + VisibleFrame.size.height;
+	WorkArea.Left = FMath::TruncToInt(VisibleFrame.origin.x);
+	WorkArea.Top = FMath::TruncToInt(VisibleFrame.origin.y);
+	WorkArea.Right = WorkArea.Left + FMath::TruncToInt(VisibleFrame.size.width);
+	WorkArea.Bottom = WorkArea.Top + FMath::TruncToInt(VisibleFrame.size.height);
 
 	return WorkArea;
 }
@@ -814,9 +814,9 @@ void FMacApplication::ProcessEvent(const FDeferredMacEvent& Event)
 		{
 			// MouseMoved events are suspended during drag and drop operations, so we need to update the cursor position here
 			NSPoint CursorPos = [NSEvent mouseLocation];
-			FVector2D NewPosition = ConvertCocoaPositionToSlate(CursorPos.x, CursorPos.y);
+			FVector2D FloatPosition = ConvertCocoaPositionToSlate(CursorPos.x, CursorPos.y);
+			FIntVector2 NewPosition(FMath::TruncToInt(FloatPosition.X), FMath::TruncToInt(FloatPosition.Y));
 			FMacCursor* MacCursor = (FMacCursor*)Cursor.Get();
-			const FVector2D MouseDelta = NewPosition - MacCursor->GetPosition();
 			if (MacCursor->UpdateCursorClipping(NewPosition))
 			{
 				MacCursor->SetPosition(NewPosition.X, NewPosition.Y);
@@ -907,7 +907,7 @@ void FMacApplication::ProcessMouseMovedEvent(const FDeferredMacEvent& Event, TSh
 	if (bUsingHighPrecisionMouseInput)
 	{
 		// Get the mouse position
-		FVector2D HighPrecisionMousePos = MacCursor->GetPosition();
+		FIntVector2 HighPrecisionMousePos = MacCursor->GetIntPosition();
 
 		// Find the visible frame of the screen the cursor is currently on.
 		FMacScreenRef Screen = FindScreenBySlatePosition(HighPrecisionMousePos.X, HighPrecisionMousePos.Y);
@@ -916,15 +916,17 @@ void FMacApplication::ProcessMouseMovedEvent(const FDeferredMacEvent& Event, TSh
 		// Under OS X we disassociate the cursor and mouse position during hi-precision mouse input.
 		// The game snaps the mouse cursor back to the starting point when this is disabled, which
 		// accumulates mouse delta that we want to ignore.
-		const FVector2D AccumDelta = MacCursor->GetMouseWarpDelta();
+		const FIntVector2 AccumDelta = MacCursor->GetMouseWarpDelta();
 
 		// Account for warping delta's
-		FVector2D Delta = FVector2D(Event.Delta.X, Event.Delta.Y);
-		const FVector2D WarpDelta(FMath::Abs(AccumDelta.X)<FMath::Abs(Delta.X) ? AccumDelta.X : Delta.X, FMath::Abs(AccumDelta.Y)<FMath::Abs(Delta.Y) ? AccumDelta.Y : Delta.Y);
-		Delta -= WarpDelta;
+		FIntVector2 Delta(FMath::TruncToInt(Event.Delta.X), FMath::TruncToInt(Event.Delta.Y));
+		const FIntVector2 WarpDelta(FMath::Abs(AccumDelta.X)<FMath::Abs(Delta.X) ? AccumDelta.X : Delta.X, FMath::Abs(AccumDelta.Y)<FMath::Abs(Delta.Y) ? AccumDelta.Y : Delta.Y);
+		Delta.X -= WarpDelta.X;
+		Delta.Y -= WarpDelta.Y;
 
 		// Update to latest position
-		HighPrecisionMousePos += Delta;
+		HighPrecisionMousePos.X += Delta.X;
+		HighPrecisionMousePos.Y += Delta.Y;
 
 		// Clip to lock rect
 		MacCursor->UpdateCursorClipping(HighPrecisionMousePos);
@@ -939,8 +941,8 @@ void FMacApplication::ProcessMouseMovedEvent(const FDeferredMacEvent& Event, TSh
 				VisibleFrame.origin.y += 5;
 				VisibleFrame.size.height -= 10;
 			}
-			int32 ClampedPosX = FMath::Clamp((int32)HighPrecisionMousePos.X, (int32)VisibleFrame.origin.x, (int32)(VisibleFrame.origin.x + VisibleFrame.size.width)-1);
-			int32 ClampedPosY = FMath::Clamp((int32)HighPrecisionMousePos.Y, (int32)VisibleFrame.origin.y, (int32)(VisibleFrame.origin.y + VisibleFrame.size.height)-1);
+			int32 ClampedPosX = FMath::Clamp(HighPrecisionMousePos.X, (int32)VisibleFrame.origin.x, (int32)(VisibleFrame.origin.x + VisibleFrame.size.width)-1);
+			int32 ClampedPosY = FMath::Clamp(HighPrecisionMousePos.Y, (int32)VisibleFrame.origin.y, (int32)(VisibleFrame.origin.y + VisibleFrame.size.height)-1);
 			MacCursor->SetPosition(ClampedPosX, ClampedPosY);
 		}
 		else
@@ -954,8 +956,10 @@ void FMacApplication::ProcessMouseMovedEvent(const FDeferredMacEvent& Event, TSh
 	else
 	{
 		NSPoint CursorPos = [NSEvent mouseLocation];
-		FVector2D NewPosition = ConvertCocoaPositionToSlate(CursorPos.x, CursorPos.y);
-		const FVector2D MouseDelta = NewPosition - MacCursor->GetPosition();
+		FVector2D FloatPosition = ConvertCocoaPositionToSlate(CursorPos.x, CursorPos.y);
+		FIntVector2 NewPosition(FMath::TruncToInt(FloatPosition.X), FMath::TruncToInt(FloatPosition.Y));
+		const FIntVector2 OldPosition = MacCursor->GetIntPosition();
+		const FIntVector2 MouseDelta(NewPosition.X - OldPosition.X, NewPosition.Y - OldPosition.Y);
 		if (MacCursor->UpdateCursorClipping(NewPosition))
 		{
 			MacCursor->SetPosition(NewPosition.X, NewPosition.Y);
@@ -971,8 +975,8 @@ void FMacApplication::ProcessMouseMovedEvent(const FDeferredMacEvent& Event, TSh
 			// its position based on mouse move delta
 			if (DraggedWindow && DraggedWindow == EventWindow->GetWindowHandle())
 			{
-				const int32 X = FMath::TruncToInt(EventWindow->PositionX + MouseDelta.X);
-				const int32 Y = FMath::TruncToInt(EventWindow->PositionY + MouseDelta.Y);
+				const int32 X = EventWindow->PositionX + MouseDelta.X;
+				const int32 Y = EventWindow->PositionY + MouseDelta.Y;
 				MessageHandler->OnMovedWindow(EventWindow.ToSharedRef(), X, Y);
 				EventWindow->PositionX = X;
 				EventWindow->PositionY = Y;
@@ -1094,8 +1098,7 @@ void FMacApplication::ProcessMouseUpEvent(const FDeferredMacEvent& Event, TShare
 
 void FMacApplication::ProcessScrollWheelEvent(const FDeferredMacEvent& Event, TSharedPtr<FMacWindow> EventWindow)
 {
-	const float DeltaX = (Event.ModifierFlags & NSEventModifierFlagShift) ? Event.Delta.Y : Event.Delta.X;
-	const float DeltaY = (Event.ModifierFlags & NSEventModifierFlagShift) ? Event.Delta.X : Event.Delta.Y;
+	const float DeltaY = (Event.ModifierFlags & NSEventModifierFlagShift) ? (float)Event.Delta.X : (float)Event.Delta.Y;
 
 	NSEventPhase Phase = Event.Phase;
 
@@ -1189,14 +1192,14 @@ void FMacApplication::OnWindowDidMove(TSharedRef<FMacWindow> Window)
 	NSRect WindowFrame = [Window->GetWindowHandle() frame];
 	NSRect OpenGLFrame = [Window->GetWindowHandle() openGLFrame];
 
-	const float X = WindowFrame.origin.x;
-	const float Y = WindowFrame.origin.y + ([Window->GetWindowHandle() windowMode] == EWindowMode::Fullscreen ? WindowFrame.size.height : OpenGLFrame.size.height);
+	const double X = WindowFrame.origin.x;
+	const double Y = WindowFrame.origin.y + ([Window->GetWindowHandle() windowMode] == EWindowMode::Fullscreen ? WindowFrame.size.height : OpenGLFrame.size.height);
 
 	FVector2D SlatePosition = ConvertCocoaPositionToSlate(X, Y);
 
-	MessageHandler->OnMovedWindow(Window, SlatePosition.X, SlatePosition.Y);
-	Window->PositionX = SlatePosition.X;
-	Window->PositionY = SlatePosition.Y;
+	MessageHandler->OnMovedWindow(Window, FMath::TruncToInt(SlatePosition.X), FMath::TruncToInt(SlatePosition.Y));
+	Window->PositionX = FMath::TruncToInt(SlatePosition.X);
+	Window->PositionY = FMath::TruncToInt(SlatePosition.Y);
 }
 
 void FMacApplication::OnWindowWillResize(TSharedRef<FMacWindow> Window)
@@ -1215,8 +1218,8 @@ void FMacApplication::OnWindowDidResize(TSharedRef<FMacWindow> Window, bool bRes
 	OnWindowDidMove(Window);
 
 	// default is no override
-	uint32 Width = [Window->GetWindowHandle() openGLFrame].size.width * Window->GetDPIScaleFactor();
-	uint32 Height = [Window->GetWindowHandle() openGLFrame].size.height * Window->GetDPIScaleFactor();
+	uint32 Width = FMath::TruncToInt([Window->GetWindowHandle() openGLFrame].size.width * Window->GetDPIScaleFactor());
+	uint32 Height = FMath::TruncToInt([Window->GetWindowHandle() openGLFrame].size.height * Window->GetDPIScaleFactor());
 
 	if (Window->GetWindowHandle().TargetWindowMode == EWindowMode::WindowedFullscreen)
 	{
@@ -1801,8 +1804,8 @@ void FMacApplication::UpdateScreensArray()
 	for (int32 Index = 0; Index < SortedScreens.Num(); ++Index)
 	{
 		FMacScreenRef& CurScreen = SortedScreens[Index];
-		const float DPIScaleFactor = bUseHighDPIMode ? CurScreen->Screen.backingScaleFactor : 1.0f;
-		if (DPIScaleFactor != 1.0f)
+		const double DPIScaleFactor = bUseHighDPIMode ? CurScreen->Screen.backingScaleFactor : 1.0;
+		if (DPIScaleFactor != 1.0)
 		{
 			CurScreen->FramePixels.size.width = CurScreen->Frame.size.width * DPIScaleFactor;
 			CurScreen->FramePixels.size.height = CurScreen->Frame.size.height * DPIScaleFactor;
@@ -1812,8 +1815,8 @@ void FMacApplication::UpdateScreensArray()
 			for (int32 OtherIndex = Index + 1; OtherIndex < SortedScreens.Num(); ++OtherIndex)
 			{
 				FMacScreenRef& OtherScreen = SortedScreens[OtherIndex];
-				const float DiffFrame = (OtherScreen->Frame.origin.x - CurScreen->Frame.origin.x) * DPIScaleFactor;
-				const float DiffVisibleFrame = (OtherScreen->VisibleFrame.origin.x - CurScreen->VisibleFrame.origin.x) * DPIScaleFactor;
+				const double DiffFrame = (OtherScreen->Frame.origin.x - CurScreen->Frame.origin.x) * DPIScaleFactor;
+				const double DiffVisibleFrame = (OtherScreen->VisibleFrame.origin.x - CurScreen->VisibleFrame.origin.x) * DPIScaleFactor;
 				OtherScreen->FramePixels.origin.x = CurScreen->FramePixels.origin.x + DiffFrame;
 				OtherScreen->VisibleFramePixels.origin.x = CurScreen->VisibleFramePixels.origin.x + DiffVisibleFrame;
 			}
@@ -1825,14 +1828,14 @@ void FMacApplication::UpdateScreensArray()
 	for (int32 Index = 0; Index < SortedScreens.Num(); ++Index)
 	{
 		FMacScreenRef& CurScreen = SortedScreens[Index];
-		const float DPIScaleFactor = bUseHighDPIMode ? CurScreen->Screen.backingScaleFactor : 1.0f;
-		if (DPIScaleFactor != 1.0f)
+		const double DPIScaleFactor = bUseHighDPIMode ? CurScreen->Screen.backingScaleFactor : 1.0;
+		if (DPIScaleFactor != 1.0)
 		{
 			for (int32 OtherIndex = Index + 1; OtherIndex < SortedScreens.Num(); ++OtherIndex)
 			{
 				FMacScreenRef& OtherScreen = SortedScreens[OtherIndex];
-				const float DiffFrame = (OtherScreen->Frame.origin.y - CurScreen->Frame.origin.y) * DPIScaleFactor;
-				const float DiffVisibleFrame = (OtherScreen->VisibleFrame.origin.y - CurScreen->VisibleFrame.origin.y) * DPIScaleFactor;
+				const double DiffFrame = (OtherScreen->Frame.origin.y - CurScreen->Frame.origin.y) * DPIScaleFactor;
+				const double DiffVisibleFrame = (OtherScreen->VisibleFrame.origin.y - CurScreen->VisibleFrame.origin.y) * DPIScaleFactor;
 				OtherScreen->FramePixels.origin.y = CurScreen->FramePixels.origin.y + DiffFrame;
 				OtherScreen->VisibleFramePixels.origin.y = CurScreen->VisibleFramePixels.origin.y + DiffVisibleFrame;
 			}
@@ -1876,10 +1879,10 @@ float FMacApplication::GetPrimaryScreenBackingScaleFactor()
 {
 	FScopeLock Lock(&GAllScreensMutex);
 	const bool bUseHighDPIMode = FPlatformApplicationMisc::IsHighDPIModeEnabled();
-	return bUseHighDPIMode ? GAllScreens[0]->Screen.backingScaleFactor : 1.0f;
+	return bUseHighDPIMode ? (float)GAllScreens[0]->Screen.backingScaleFactor : 1.0f;
 }
 
-FMacScreenRef FMacApplication::FindScreenBySlatePosition(float X, float Y)
+FMacScreenRef FMacApplication::FindScreenBySlatePosition(double X, double Y)
 {
 	NSPoint Point = NSMakePoint(X, Y);
 
@@ -1898,7 +1901,7 @@ FMacScreenRef FMacApplication::FindScreenBySlatePosition(float X, float Y)
 	return TargetScreen;
 }
 
-FMacScreenRef FMacApplication::FindScreenByCocoaPosition(float X, float Y)
+FMacScreenRef FMacApplication::FindScreenByCocoaPosition(double X, double Y)
 {
 	NSPoint Point = NSMakePoint(X, Y);
 
@@ -1917,29 +1920,29 @@ FMacScreenRef FMacApplication::FindScreenByCocoaPosition(float X, float Y)
 	return TargetScreen;
 }
 
-FVector2D FMacApplication::ConvertSlatePositionToCocoa(float X, float Y)
+FVector2D FMacApplication::ConvertSlatePositionToCocoa(double X, double Y)
 {
 	FMacScreenRef Screen = FindScreenBySlatePosition(X, Y);
 	const bool bUseHighDPIMode = FPlatformApplicationMisc::IsHighDPIModeEnabled();
-	const float DPIScaleFactor = bUseHighDPIMode ? Screen->Screen.backingScaleFactor : 1.0f;
+	const double DPIScaleFactor = bUseHighDPIMode ? Screen->Screen.backingScaleFactor : 1.0;
 	const FVector2D OffsetOnScreen = FVector2D(X - Screen->FramePixels.origin.x, Screen->FramePixels.origin.y + Screen->FramePixels.size.height - Y) / DPIScaleFactor;
 	return FVector2D(Screen->Screen.frame.origin.x + OffsetOnScreen.X, Screen->Screen.frame.origin.y + OffsetOnScreen.Y);
 }
 
-FVector2D FMacApplication::ConvertCocoaPositionToSlate(float X, float Y)
+FVector2D FMacApplication::ConvertCocoaPositionToSlate(double X, double Y)
 {
 	FMacScreenRef Screen = FindScreenByCocoaPosition(X, Y);
 	const bool bUseHighDPIMode = FPlatformApplicationMisc::IsHighDPIModeEnabled();
-	const float DPIScaleFactor = bUseHighDPIMode ? Screen->Screen.backingScaleFactor : 1.0f;
+	const double DPIScaleFactor = bUseHighDPIMode ? Screen->Screen.backingScaleFactor : 1.0;
 	const FVector2D OffsetOnScreen = FVector2D(X - Screen->Screen.frame.origin.x, Screen->Screen.frame.origin.y + Screen->Screen.frame.size.height - Y) * DPIScaleFactor;
 	return FVector2D(Screen->FramePixels.origin.x + OffsetOnScreen.X, Screen->FramePixels.origin.y + OffsetOnScreen.Y);
 }
 
-CGPoint FMacApplication::ConvertSlatePositionToCGPoint(float X, float Y)
+CGPoint FMacApplication::ConvertSlatePositionToCGPoint(double X, double Y)
 {
 	FMacScreenRef Screen = FindScreenBySlatePosition(X, Y);
 	const bool bUseHighDPIMode = FPlatformApplicationMisc::IsHighDPIModeEnabled();
-	const float DPIScaleFactor = bUseHighDPIMode ? Screen->Screen.backingScaleFactor : 1.0f;
+	const double DPIScaleFactor = bUseHighDPIMode ? Screen->Screen.backingScaleFactor : 1.0;
 	const FVector2D OffsetOnScreen = FVector2D(X - Screen->FramePixels.origin.x, Y - Screen->FramePixels.origin.y) / DPIScaleFactor;
 	return CGPointMake(Screen->Frame.origin.x + OffsetOnScreen.X, Screen->Frame.origin.y + OffsetOnScreen.Y);
 }
@@ -1947,8 +1950,8 @@ CGPoint FMacApplication::ConvertSlatePositionToCGPoint(float X, float Y)
 EWindowZone::Type FMacApplication::GetCurrentWindowZone(const TSharedRef<FMacWindow>& Window) const
 {
 	const FVector2D CursorPos = ((FMacCursor*)Cursor.Get())->GetPosition();
-	const int32 LocalMouseX = CursorPos.X - Window->PositionX;
-	const int32 LocalMouseY = CursorPos.Y - Window->PositionY;
+	const int32 LocalMouseX = FMath::TruncToInt(CursorPos.X - Window->PositionX);
+	const int32 LocalMouseY = FMath::TruncToInt(CursorPos.Y - Window->PositionY);
 	return MessageHandler->GetWindowZoneForPoint(Window, LocalMouseX, LocalMouseY);
 }
 
@@ -2087,7 +2090,7 @@ unichar FMacApplication::TranslateKeyCodeToUniCode(uint32 KeyCode, uint32 Modifi
 					UniCharCount BufferLength = 256;
 					uint32 DeadKeyState = 0;
 
-					OSStatus Status = UCKeyTranslate(KeyboardLayout, KeyCode, kUCKeyActionDown, ((Modifier) >> 8) & 0xFF, LMGetKbdType(), kUCKeyTranslateNoDeadKeysMask, &DeadKeyState, BufferLength, &BufferLength, Buffer);
+					OSStatus Status = UCKeyTranslate(KeyboardLayout, (uint16)KeyCode, kUCKeyActionDown, (uint16)(((Modifier) >> 8) & 0xFF), LMGetKbdType(), kUCKeyTranslateNoDeadKeysMask, &DeadKeyState, BufferLength, &BufferLength, Buffer);
 					if (Status == noErr)
 					{
 						return Buffer[0];
@@ -2123,7 +2126,7 @@ TCHAR FMacApplication::TranslateCharCode(TCHAR CharCode, uint32 KeyCode) const
 					uint32 DeadKeyState = 0;
 
 					// To ensure we get a latin character, we pretend that command modifier key is pressed
-					OSStatus Status = UCKeyTranslate(KeyboardLayout, KeyCode, kUCKeyActionDown, cmdKey >> 8, LMGetKbdType(), kUCKeyTranslateNoDeadKeysMask, &DeadKeyState, BufferLength, &BufferLength, Buffer);
+					OSStatus Status = UCKeyTranslate(KeyboardLayout, (uint16)KeyCode, kUCKeyActionDown, cmdKey >> 8, LMGetKbdType(), kUCKeyTranslateNoDeadKeysMask, &DeadKeyState, BufferLength, &BufferLength, Buffer);
 					if (Status == noErr)
 					{
 						CharCode = Buffer[0];
@@ -2221,8 +2224,8 @@ void FDisplayMetrics::RebuildDisplayMetrics(FDisplayMetrics& OutDisplayMetrics)
 	const NSRect VisibleFrame = PrimaryScreen->VisibleFramePixels;
 
 	// Total screen size of the primary monitor
-	OutDisplayMetrics.PrimaryDisplayWidth = ScreenFrame.size.width;
-	OutDisplayMetrics.PrimaryDisplayHeight = ScreenFrame.size.height;
+	OutDisplayMetrics.PrimaryDisplayWidth = FMath::TruncToInt(ScreenFrame.size.width);
+	OutDisplayMetrics.PrimaryDisplayHeight = FMath::TruncToInt(ScreenFrame.size.height);
 
 	OutDisplayMetrics.MonitorInfo.Empty();
 
@@ -2266,8 +2269,20 @@ void FDisplayMetrics::RebuildDisplayMetrics(FDisplayMetrics& OutDisplayMetrics)
 		
 			CFRelease(ArrDisplay);
 
-			Info.DisplayRect = FPlatformRect(Screen->FramePixels.origin.x, Screen->FramePixels.origin.y, Screen->FramePixels.origin.x + Screen->FramePixels.size.width, Screen->FramePixels.origin.y + Screen->FramePixels.size.height);
-			Info.WorkArea = FPlatformRect(Screen->VisibleFramePixels.origin.x, Screen->VisibleFramePixels.origin.y, Screen->VisibleFramePixels.origin.x + Screen->VisibleFramePixels.size.width, Screen->VisibleFramePixels.origin.y + Screen->VisibleFramePixels.size.height);
+			Info.DisplayRect = FPlatformRect
+			(
+				FMath::TruncToInt(Screen->FramePixels.origin.x),
+				FMath::TruncToInt(Screen->FramePixels.origin.y),
+				FMath::TruncToInt(Screen->FramePixels.origin.x + Screen->FramePixels.size.width),
+				FMath::TruncToInt(Screen->FramePixels.origin.y + Screen->FramePixels.size.height)
+			);
+			Info.WorkArea = FPlatformRect
+			(
+				FMath::TruncToInt(Screen->VisibleFramePixels.origin.x),
+				FMath::TruncToInt(Screen->VisibleFramePixels.origin.y),
+				FMath::TruncToInt(Screen->VisibleFramePixels.origin.x + Screen->VisibleFramePixels.size.width),
+				FMath::TruncToInt(Screen->VisibleFramePixels.origin.y + Screen->VisibleFramePixels.size.height)
+			);
 			Info.bIsPrimary = Screen->Screen == [NSScreen mainScreen];
 
 			// dpi computations
@@ -2318,16 +2333,16 @@ void FDisplayMetrics::RebuildDisplayMetrics(FDisplayMetrics& OutDisplayMetrics)
 	}
 
 	// Virtual desktop area
-	OutDisplayMetrics.VirtualDisplayRect.Left = WholeWorkspace.origin.x;
-	OutDisplayMetrics.VirtualDisplayRect.Top = FMath::Min(WholeWorkspace.origin.y, 0.0);
-	OutDisplayMetrics.VirtualDisplayRect.Right = WholeWorkspace.origin.x + WholeWorkspace.size.width;
-	OutDisplayMetrics.VirtualDisplayRect.Bottom = WholeWorkspace.size.height + OutDisplayMetrics.VirtualDisplayRect.Top;
+	OutDisplayMetrics.VirtualDisplayRect.Left = FMath::TruncToInt(WholeWorkspace.origin.x);
+	OutDisplayMetrics.VirtualDisplayRect.Top = FMath::TruncToInt(FMath::Min(WholeWorkspace.origin.y, 0.0));
+	OutDisplayMetrics.VirtualDisplayRect.Right = FMath::TruncToInt(WholeWorkspace.origin.x + WholeWorkspace.size.width);
+	OutDisplayMetrics.VirtualDisplayRect.Bottom = FMath::TruncToInt(WholeWorkspace.size.height + OutDisplayMetrics.VirtualDisplayRect.Top);
 
 	// Get the screen rect of the primary monitor, excluding taskbar etc.
-	OutDisplayMetrics.PrimaryDisplayWorkAreaRect.Left = VisibleFrame.origin.x;
-	OutDisplayMetrics.PrimaryDisplayWorkAreaRect.Top = VisibleFrame.origin.y;
-	OutDisplayMetrics.PrimaryDisplayWorkAreaRect.Right = VisibleFrame.origin.x + VisibleFrame.size.width;
-	OutDisplayMetrics.PrimaryDisplayWorkAreaRect.Bottom = OutDisplayMetrics.PrimaryDisplayWorkAreaRect.Top + VisibleFrame.size.height;
+	OutDisplayMetrics.PrimaryDisplayWorkAreaRect.Left = FMath::TruncToInt(VisibleFrame.origin.x);
+	OutDisplayMetrics.PrimaryDisplayWorkAreaRect.Top = FMath::TruncToInt(VisibleFrame.origin.y);
+	OutDisplayMetrics.PrimaryDisplayWorkAreaRect.Right = FMath::TruncToInt(VisibleFrame.origin.x + VisibleFrame.size.width);
+	OutDisplayMetrics.PrimaryDisplayWorkAreaRect.Bottom = FMath::TruncToInt(OutDisplayMetrics.PrimaryDisplayWorkAreaRect.Top + VisibleFrame.size.height);
 
 	// Apply the debug safe zones
 	OutDisplayMetrics.ApplyDefaultSafeZones();
diff --git a/Engine/Source/Runtime/ApplicationCore/Private/Mac/MacConsoleOutputDevice.cpp b/Engine/Source/Runtime/ApplicationCore/Private/Mac/MacConsoleOutputDevice.cpp
--- a/Engine/Source/Runtime/ApplicationCore/Private/Mac/MacConsoleOutputDevice.cpp
+++ b/Engine/Source/Runtime/ApplicationCore/Private/Mac/MacConsoleOutputDevice.cpp
@@ -33,10 +33,10 @@ void FMacConsoleOutputDevice::SaveToINI()
 	{
 		NSRect Frame = [ConsoleHandle frame];
 
-		GConfig->SetInt(TEXT("DebugMac"), TEXT("ConsoleWidth"), Frame.size.width, IniFilename);
-		GConfig->SetInt(TEXT("DebugMac"), TEXT("ConsoleHeight"), Frame.size.height, IniFilename);
-		GConfig->SetInt(TEXT("DebugMac"), TEXT("ConsoleX"), Frame.origin.x, IniFilename);
-		GConfig->SetInt(TEXT("DebugMac"), TEXT("ConsoleY"), Frame.origin.y, IniFilename);
+		GConfig->SetInt(TEXT("DebugMac"), TEXT("ConsoleWidth"), FMath::TruncToInt(Frame.size.width), IniFilename);
+		GConfig->SetInt(TEXT("DebugMac"), TEXT("ConsoleHeight"), FMath::TruncToInt(Frame.size.height), IniFilename);
+		GConfig->SetInt(TEXT("DebugMac"), TEXT("ConsoleX"), FMath::TruncToInt(Frame.origin.x), IniFilename);
+		GConfig->SetInt(TEXT("DebugMac"), TEXT("ConsoleY"), FMath::TruncToInt(Frame.origin.y), IniFilename);
 	}
 }
 
diff --git a/Engine/Source/Runtime/ApplicationCore/Private/Mac/MacCursor.cpp b/Engine/Source/Runtime/ApplicationCore/Private/Mac/MacCursor.cpp
--- a/Engine/Source/Runtime/ApplicationCore/Private/Mac/MacCursor.cpp
+++ b/Engine/Source/Runtime/ApplicationCore/Private/Mac/MacCursor.cpp
@@ -28,8 +28,8 @@ FMacCursor::FMacCursor()
 :	bIsVisible(true)
 ,	bUseHighPrecisionMode(false)
 ,	CursorTypeOverride(-1)
-,	CurrentPosition(FVector2D::ZeroVector)
-,	MouseWarpDelta(FVector2D::ZeroVector)
+,	CurrentPosition(0, 0)
+,	MouseWarpDelta(0, 0)
 ,	bIsPositionInitialised(false)
 ,	bShouldIgnoreLocking(false)
 ,	HIDInterface(0)
@@ -287,27 +287,34 @@ void* FMacCursor::CreateCursorFromRGBABuffer(const FColor* Pixels, int32 Width,
 	return CursorHandle;
 }
 
-FVector2D FMacCursor::GetPosition() const
+FIntVector2 FMacCursor::GetIntPosition() const
 {
-	FVector2D CurrentPos = CurrentPosition;
-	if (!bIsPositionInitialised)
+	if (bIsPositionInitialised)
 	{
-		SCOPED_AUTORELEASE_POOL;
-		NSPoint CursorPos = [NSEvent mouseLocation];
-		CurrentPos = FMacApplication::ConvertCocoaPositionToSlate(CursorPos.x, CursorPos.y);
+		return CurrentPosition;
 	}
 
-	return CurrentPos;
+	SCOPED_AUTORELEASE_POOL;
+	NSPoint CursorPos = [NSEvent mouseLocation];
+	FVector2D CurrentPos = FMacApplication::ConvertCocoaPositionToSlate(CursorPos.x, CursorPos.y);
+	return FIntVector2(FMath::TruncToInt(CurrentPos.X), FMath::TruncToInt(CurrentPos.Y));
+}
+
+FVector2D FMacCursor::GetPosition() const
+{
+	FIntVector2 CurrentPos = GetIntPosition();
+	return FVector2D(CurrentPos.X, CurrentPos.Y);
 }
 
 void FMacCursor::SetPosition(const int32 X, const int32 Y)
 {
-	FVector2D NewPos(X, Y);
+	FIntVector2 NewPos(X, Y);
 	UpdateCursorClipping(NewPos);
 
-	MouseWarpDelta += (NewPos - CurrentPosition);
+	MouseWarpDelta.X += NewPos.X - CurrentPosition.X;
+	MouseWarpDelta.Y += NewPos.Y - CurrentPosition.Y;
 
-	if (!bIsPositionInitialised || FIntVector(NewPos.X, NewPos.Y, 0) != FIntVector(CurrentPosition.X, CurrentPosition.Y, 0))
+	if (!bIsPositionInitialised || NewPos != CurrentPosition)
 	{
 		if (!bUseHighPrecisionMode || (CurrentCursor && bIsVisible) || !bIsPositionInitialised)
 		{
@@ -371,8 +378,8 @@ void FMacCursor::Lock(const RECT* const Bounds)
 
 	MacApplication->OnCursorLock();
 
-	bIsPositionInitialised = false; // Force GetPosition() to update its cached position in case the cursor was warped by another app while we were in high precision mode
-	FVector2D Position = GetPosition();
+	bIsPositionInitialised = false; // Force GetIntPosition() to update its cached position in case the cursor was warped by another app while we were in high precision mode
+	FIntVector2 Position = GetIntPosition();
 	if (UpdateCursorClipping(Position))
 	{
 		SetPosition(Position.X, Position.Y);
@@ -380,15 +387,15 @@ void FMacCursor::Lock(const RECT* const Bounds)
 	}
 }
 
-bool FMacCursor::UpdateCursorClipping(FVector2D& CursorPosition)
+bool FMacCursor::UpdateCursorClipping(FIntVector2& CursorPosition)
 {
 	bool bAdjusted = false;
 
 	if (CursorClipRect.Area() > 0)
 	{
-		FVector2D PositionOnScreen(CursorPosition);
+		FIntVector2 PositionOnScreen(CursorPosition);
 		FIntRect ClipRect(CursorClipRect);
-		FVector2D ScreenOrigin(FVector2D::ZeroVector);
+		FIntVector2 ScreenOrigin(0, 0);
 
 		if (PositionOnScreen.X < ClipRect.Min.X)
 		{
@@ -414,7 +421,8 @@ bool FMacCursor::UpdateCursorClipping(FVector2D& CursorPosition)
 
 		if (bAdjusted)
 		{
-			CursorPosition = PositionOnScreen + ScreenOrigin;
+			CursorPosition.X = PositionOnScreen.X + ScreenOrigin.X;
+			CursorPosition.Y = PositionOnScreen.Y + ScreenOrigin.Y;
 		}
 	}
 
@@ -460,7 +468,7 @@ PRAGMA_DISABLE_DEPRECATION_WARNINGS
 PRAGMA_ENABLE_DEPRECATION_WARNINGS
 }
 
-void FMacCursor::UpdateCurrentPosition(const FVector2D &Position)
+void FMacCursor::UpdateCurrentPosition(const FIntVector2& Position)
 {
 	CurrentPosition = Position;
 	bIsPositionInitialised = true;
@@ -486,15 +494,15 @@ void FMacCursor::WarpCursor(const int32 X, const int32 Y)
 		CGAssociateMouseAndMouseCursorPosition(true);
 	}
 
-	UpdateCurrentPosition(FVector2D(X, Y));
+	UpdateCurrentPosition(FIntVector2(X, Y));
 
 	MacApplication->IgnoreMouseMoveDelta();
 }
 
-FVector2D FMacCursor::GetMouseWarpDelta()
+FIntVector2 FMacCursor::GetMouseWarpDelta()
 {
-	FVector2D Result = (!bUseHighPrecisionMode || (CurrentCursor && bIsVisible)) ? MouseWarpDelta : FVector2D::ZeroVector;
-	MouseWarpDelta = FVector2D::ZeroVector;
+	FIntVector2 Result = (!bUseHighPrecisionMode || (CurrentCursor && bIsVisible)) ? MouseWarpDelta : FIntVector2(0, 0);
+	MouseWarpDelta = FIntVector2(0, 0);
 	return Result;
 }
 
@@ -541,7 +549,7 @@ PRAGMA_ENABLE_DEPRECATION_WARNINGS
 		// On disable put the cursor where the user would expect it
 		if (!bEnable && (!CurrentCursor || !bIsVisible))
 		{
-			FVector2D Position = GetPosition();
+			FIntVector2 Position = GetIntPosition();
 			UpdateCursorClipping(Position);
 			WarpCursor(Position.X, Position.Y);
 		}
diff --git a/Engine/Source/Runtime/ApplicationCore/Private/Mac/MacPlatformApplicationMisc.cpp b/Engine/Source/Runtime/ApplicationCore/Private/Mac/MacPlatformApplicationMisc.cpp
--- a/Engine/Source/Runtime/ApplicationCore/Private/Mac/MacPlatformApplicationMisc.cpp
+++ b/Engine/Source/Runtime/ApplicationCore/Private/Mac/MacPlatformApplicationMisc.cpp
@@ -442,7 +442,7 @@ float FMacPlatformApplicationMisc::GetDPIScaleFactorAtPoint(float X, float Y)
 	if (MacApplication && FPlatformApplicationMisc::IsHighDPIModeEnabled())
 	{
 		FMacScreenRef Screen = FMacApplication::FindScreenBySlatePosition(X, Y);
-		return Screen->Screen.backingScaleFactor;
+		return (float)Screen->Screen.backingScaleFactor;
 	}
 	return 1.0f;
 }
diff --git a/Engine/Source/Runtime/ApplicationCore/Private/Mac/MacPlatformSurvey.cpp b/Engine/Source/Runtime/ApplicationCore/Private/Mac/MacPlatformSurvey.cpp
--- a/Engine/Source/Runtime/ApplicationCore/Private/Mac/MacPlatformSurvey.cpp
+++ b/Engine/Source/Runtime/ApplicationCore/Private/Mac/MacPlatformSurvey.cpp
@@ -70,7 +70,7 @@ bool FMacPlatformSurvey::GetSurveyResults( FHardwareSurveyResults& OutResults, b
 		Size = sizeof(int64);
 		if (sysctlbyname("hw.cpufrequency", &CPUSpeed, &Size, NULL, 0) == 0)
 		{
-			OutResults.CPUClockGHz = 0.000000001 * CPUSpeed;
+			OutResults.CPUClockGHz = (float)(0.000000001 * CPUSpeed);
 		}
 		else
 		{
diff --git a/Engine/Source/Runtime/ApplicationCore/Private/Mac/MacWindow.cpp b/Engine/Source/Runtime/ApplicationCore/Private/Mac/MacWindow.cpp
--- a/Engine/Source/Runtime/ApplicationCore/Private/Mac/MacWindow.cpp
+++ b/Engine/Source/Runtime/ApplicationCore/Private/Mac/MacWindow.cpp
@@ -44,10 +44,10 @@ void FMacWindow::Initialize( FMacApplication* const Application, const TSharedRe
 	const int32 SizeX = FMath::Max(FMath::CeilToInt( Definition->WidthDesiredOnScreen ), 1);
 	const int32 SizeY = FMath::Max(FMath::CeilToInt( Definition->HeightDesiredOnScreen ), 1);
 
-	PositionX = Definition->XDesiredPositionOnScreen;
-	PositionY = Definition->YDesiredPositionOnScreen >= TargetScreen->VisibleFramePixels.origin.y ? Definition->YDesiredPositionOnScreen : TargetScreen->VisibleFramePixels.origin.y;
+	PositionX = FMath::TruncToInt(Definition->XDesiredPositionOnScreen);
+	PositionY = FMath::TruncToInt(Definition->YDesiredPositionOnScreen >= TargetScreen->VisibleFramePixels.origin.y ? Definition->YDesiredPositionOnScreen : TargetScreen->VisibleFramePixels.origin.y);
 
-	const float ScreenDPIScaleFactor = FPlatformApplicationMisc::IsHighDPIModeEnabled() ? TargetScreen->Screen.backingScaleFactor : 1.0f;
+	const double ScreenDPIScaleFactor = FPlatformApplicationMisc::IsHighDPIModeEnabled() ? TargetScreen->Screen.backingScaleFactor : 1.0;
 	const FVector2D CocoaPosition = FMacApplication::ConvertSlatePositionToCocoa(PositionX, PositionY);
 	const NSRect ViewRect = NSMakeRect(CocoaPosition.X, CocoaPosition.Y - (SizeY / ScreenDPIScaleFactor) + 1, SizeX / ScreenDPIScaleFactor, SizeY / ScreenDPIScaleFactor);
 
@@ -236,11 +236,11 @@ bool FMacWindow::GetFullScreenInfo( int32& X, int32& Y, int32& Width, int32& Hei
 	bool const bIsFullscreen = (GetWindowMode() == EWindowMode::Fullscreen);
 	const NSRect Frame = WindowHandle.screen.frame;
 	const FVector2D SlatePosition = FMacApplication::ConvertCocoaPositionToSlate(Frame.origin.x, Frame.origin.y - Frame.size.height + 1.0f);
-	X = SlatePosition.X;
-	Y = SlatePosition.Y;
-	const float DPIScaleFactor = FPlatformApplicationMisc::IsHighDPIModeEnabled() ? WindowHandle.screen.backingScaleFactor : 1.0f;
-	Width = Frame.size.width * DPIScaleFactor;
-	Height = Frame.size.height * DPIScaleFactor;
+	X = FMath::TruncToInt(SlatePosition.X);
+	Y = FMath::TruncToInt(SlatePosition.Y);
+	const double DPIScaleFactor = FPlatformApplicationMisc::IsHighDPIModeEnabled() ? WindowHandle.screen.backingScaleFactor : 1.0;
+	Width = FMath::TruncToInt(Frame.size.width * DPIScaleFactor);
+	Height = FMath::TruncToInt(Frame.size.height * DPIScaleFactor);
 	return true;
 }
 
@@ -384,7 +384,7 @@ void FMacWindow::SetWindowMode(EWindowMode::Type NewWindowMode)
 {
 	if(WindowHandle)
 	{
-		ApplySizeAndModeChanges(PositionX, PositionY, WindowHandle.contentView.frame.size.width, WindowHandle.contentView.frame.size.height, NewWindowMode);
+		ApplySizeAndModeChanges(PositionX, PositionY, FMath::TruncToInt(WindowHandle.contentView.frame.size.width), FMath::TruncToInt(WindowHandle.contentView.frame.size.height), NewWindowMode);
 	}
 }
 
@@ -420,12 +420,12 @@ bool FMacWindow::GetRestoredDimensions(int32& X, int32& Y, int32& Width, int32&
 
 		const FVector2D SlatePosition = FMacApplication::ConvertCocoaPositionToSlate(Frame.origin.x, Frame.origin.y);
 
-		const float DPIScaleFactor = FPlatformApplicationMisc::IsHighDPIModeEnabled() ? WindowHandle.backingScaleFactor : 1.0f;
-		Width = Frame.size.width * DPIScaleFactor;
-		Height = Frame.size.height * DPIScaleFactor;
+		const double DPIScaleFactor = FPlatformApplicationMisc::IsHighDPIModeEnabled() ? WindowHandle.backingScaleFactor : 1.0;
+		Width = FMath::TruncToInt(Frame.size.width * DPIScaleFactor);
+		Height = FMath::TruncToInt(Frame.size.height * DPIScaleFactor);
 
-		X = SlatePosition.X;
-		Y = SlatePosition.Y - Height + 1;
+		X = FMath::TruncToInt(SlatePosition.X);
+		Y = FMath::TruncToInt(SlatePosition.Y - Height + 1);
 
 		return true;
 	}
@@ -533,7 +533,7 @@ bool FMacWindow::IsRegularWindow() const
 
 float FMacWindow::GetDPIScaleFactor() const
 {
-	return FPlatformApplicationMisc::IsHighDPIModeEnabled() ? WindowHandle.backingScaleFactor : 1.0f;
+	return FPlatformApplicationMisc::IsHighDPIModeEnabled() ? (float)WindowHandle.backingScaleFactor : 1.0f;
 }
 
 void FMacWindow::SetNativeWindowButtonsVisibility(bool bVisible)
@@ -610,7 +610,7 @@ void FMacWindow::ApplySizeAndModeChanges(int32 X, int32 Y, int32 Width, int32 He
 			CGError Error = CGAcquireDisplayFadeReservation(kCGMaxDisplayReservationInterval, &FadeReservationToken);
 			if (Error == kCGErrorSuccess)
 			{
-				CGDisplayFade(FadeReservationToken, 0.3, kCGDisplayBlendNormal, kCGDisplayBlendSolidColor, 0, 0, 0, true);
+				CGDisplayFade(FadeReservationToken, 0.3f, kCGDisplayBlendNormal, kCGDisplayBlendSolidColor, 0, 0, 0, true);
 			}
 		}, UnrealResizeEventMode, true);
 	}
diff --git a/Engine/Source/Runtime/ApplicationCore/Public/Mac/CocoaWindow.h b/Engine/Source/Runtime/ApplicationCore/Public/Mac/CocoaWindow.h
--- a/Engine/Source/Runtime/ApplicationCore/Public/Mac/CocoaWindow.h
+++ b/Engine/Source/Runtime/ApplicationCore/Public/Mac/CocoaWindow.h
@@ -19,7 +19,7 @@
 	bool bDisplayReconfiguring;
 	bool bRenderInitialized;
 	bool bIsBeingOrderedFront;
-	float Opacity;
+	CGFloat Opacity;
 @public
 	bool bAcceptsInput;
 	bool bZoomed;
diff --git a/Engine/Source/Runtime/ApplicationCore/Public/Mac/MacApplication.h b/Engine/Source/Runtime/ApplicationCore/Public/Mac/MacApplication.h
--- a/Engine/Source/Runtime/ApplicationCore/Public/Mac/MacApplication.h
+++ b/Engine/Source/Runtime/ApplicationCore/Public/Mac/MacApplication.h
@@ -249,15 +249,15 @@ public:
 
 	static void UpdateScreensArray();
 
-	static FMacScreenRef FindScreenBySlatePosition(float X, float Y);
+	static FMacScreenRef FindScreenBySlatePosition(double X, double Y);
 
-	static FMacScreenRef FindScreenByCocoaPosition(float X, float Y);
+	static FMacScreenRef FindScreenByCocoaPosition(double X, double Y);
 
-	static FVector2D ConvertSlatePositionToCocoa(float X, float Y);
+	static FVector2D ConvertSlatePositionToCocoa(double X, double Y);
 
-	static FVector2D ConvertCocoaPositionToSlate(float X, float Y);
+	static FVector2D ConvertCocoaPositionToSlate(double X, double Y);
 
-	static CGPoint ConvertSlatePositionToCGPoint(float X, float Y);
+	static CGPoint ConvertSlatePositionToCGPoint(double X, double Y);
 
 	static FVector2D CalculateScreenOrigin(NSScreen* Screen);
 
diff --git a/Engine/Source/Runtime/ApplicationCore/Public/Mac/MacCursor.h b/Engine/Source/Runtime/ApplicationCore/Public/Mac/MacCursor.h
--- a/Engine/Source/Runtime/ApplicationCore/Public/Mac/MacCursor.h
+++ b/Engine/Source/Runtime/ApplicationCore/Public/Mac/MacCursor.h
@@ -4,6 +4,7 @@
 
 #include "GenericPlatform/ICursor.h"
 #include "Math/IntRect.h"
+#include "Math/IntVector.h"
 
 @class FCocoaWindow;
 
@@ -39,15 +40,17 @@ public:
 
 public:
 
-	bool UpdateCursorClipping(FVector2D& CursorPosition);
+	bool UpdateCursorClipping(FIntVector2& CursorPosition);
 
 	void WarpCursor(const int32 X, const int32 Y);
 
-	FVector2D GetMouseWarpDelta();
+	FIntVector2 GetIntPosition() const;
+
+	FIntVector2 GetMouseWarpDelta();
 
 	void SetHighPrecisionMouseMode(const bool bEnable);
 
-	void UpdateCurrentPosition(const FVector2D& Position);
+	void UpdateCurrentPosition(const FIntVector2& Position);
 
 	void UpdateVisibility();
 
@@ -69,8 +72,8 @@ private:
 	NSCursor* CurrentCursor;
 	int32 CursorTypeOverride;
 
-	FVector2D CurrentPosition;
-	FVector2D MouseWarpDelta;
+	FIntVector2 CurrentPosition;
+	FIntVector2 MouseWarpDelta;
 	bool bIsPositionInitialised;
 	bool bShouldIgnoreLocking;
 
