Fixed unsafe type cast warnings in ApplicationCore on iOS

#rnx

diff --git a/Engine/Source/Runtime/ApplicationCore/Private/IOS/Accessibility/IOSAccessibilityCache.cpp b/Engine/Source/Runtime/ApplicationCore/Private/IOS/Accessibility/IOSAccessibilityCache.cpp
--- a/Engine/Source/Runtime/ApplicationCore/Private/IOS/Accessibility/IOSAccessibilityCache.cpp
+++ b/Engine/Source/Runtime/ApplicationCore/Private/IOS/Accessibility/IOSAccessibilityCache.cpp
@@ -116,7 +116,7 @@ static void DumpAccessibilityStatsForwarder()
 
 					// Bounding rect
 					FBox2D Bounds = Widget->GetBounds();
-					const float Scale = [IOSAppDelegate GetDelegate].IOSView.contentScaleFactor;
+					const double Scale = [IOSAppDelegate GetDelegate].IOSView.contentScaleFactor;
 					Bounds.Min /= Scale;
 					Bounds.Max /= Scale;
 
diff --git a/Engine/Source/Runtime/ApplicationCore/Private/IOS/Accessibility/IOSAccessibilityElement.cpp b/Engine/Source/Runtime/ApplicationCore/Private/IOS/Accessibility/IOSAccessibilityElement.cpp
--- a/Engine/Source/Runtime/ApplicationCore/Private/IOS/Accessibility/IOSAccessibilityElement.cpp
+++ b/Engine/Source/Runtime/ApplicationCore/Private/IOS/Accessibility/IOSAccessibilityElement.cpp
@@ -230,9 +230,9 @@ if (!_AccessibilityChildren)
 -(id)accessibilityHitTest:(CGPoint)point
 {
 	AccessibleWidgetId FoundId = IAccessibleWidget::InvalidAccessibleWidgetId;
-	const float Scale = [IOSAppDelegate GetDelegate].IOSView.contentScaleFactor;
-	const int32 X = point.x * Scale;
-	const int32 Y = point.y * Scale;
+	const double Scale = [IOSAppDelegate GetDelegate].IOSView.contentScaleFactor;
+	const int32 X = FMath::TruncToInt(point.x * Scale);
+	const int32 Y = FMath::TruncToInt(point.y * Scale);
 	// Update the labels while we're on the game thread, since IOS is going to request them immediately.
 	FString TempLabel, TempHint, TempValue;
 	const AccessibleWidgetId TempId = self.Leaf.Id;
diff --git a/Engine/Source/Runtime/ApplicationCore/Private/IOS/IOSAppDelegate.cpp b/Engine/Source/Runtime/ApplicationCore/Private/IOS/IOSAppDelegate.cpp
--- a/Engine/Source/Runtime/ApplicationCore/Private/IOS/IOSAppDelegate.cpp
+++ b/Engine/Source/Runtime/ApplicationCore/Private/IOS/IOSAppDelegate.cpp
@@ -828,8 +828,8 @@ static IOSAppDelegate* CachedDelegate = nil;
 - (void)LoadScreenResolutionModifiers
 {
 	// cache these UI thread sensitive vars for later use
-	self.ScreenScale = [[UIScreen mainScreen] scale];
-	self.NativeScale = [[UIScreen mainScreen] nativeScale];
+	self.ScreenScale = (float)[[UIScreen mainScreen] scale];
+	self.NativeScale = (float)[[UIScreen mainScreen] nativeScale];
 
 	// need to cache the MobileContentScaleFactor for framebuffer creation.
 	static IConsoleVariable* CVarScale = IConsoleManager::Get().FindConsoleVariable(TEXT("r.MobileContentScaleFactor"));
@@ -991,10 +991,7 @@ static FAutoConsoleVariableRef CVarGEnableThermalsReport(
     
     CGRect MainFrame = [[UIScreen mainScreen] bounds];
     self.Window = [[UIWindow alloc] initWithFrame:MainFrame];
-    
-    // get the native scale
-    const float NativeScale = [[UIScreen mainScreen] scale];
-    
+
     [self.Window makeKeyAndVisible];
 
     FAppEntry::PreInit(self, application);
@@ -1609,7 +1606,7 @@ void HandleReceivedNotification(UNNotification* notification)
 			if(activationEvent != nullptr)
 			{
 				FString	activationEventFString(activationEvent);
-				int32	fireDate = [notification.date timeIntervalSince1970];
+				int32	fireDate = FMath::TruncToInt([notification.date timeIntervalSince1970]);
 				
 				FFunctionGraphTask::CreateAndDispatchWhenReady([activationEventFString, fireDate, AppState]()
 															   {
@@ -1666,7 +1663,7 @@ didReceiveNotificationResponse:(UNNotificationResponse *)response
 		{
 			FAppEntry::gAppLaunchedWithLocalNotification = true;
 			FAppEntry::gLaunchLocalNotificationActivationEvent = FString(activationEvent);
-			FAppEntry::gLaunchLocalNotificationFireDate = [response.notification.date timeIntervalSince1970];
+			FAppEntry::gLaunchLocalNotificationFireDate = FMath::TruncToInt([response.notification.date timeIntervalSince1970]);
 		}
 	}
 	
diff --git a/Engine/Source/Runtime/ApplicationCore/Private/IOS/IOSApplication.cpp b/Engine/Source/Runtime/ApplicationCore/Private/IOS/IOSApplication.cpp
--- a/Engine/Source/Runtime/ApplicationCore/Private/IOS/IOSApplication.cpp
+++ b/Engine/Source/Runtime/ApplicationCore/Private/IOS/IOSApplication.cpp
@@ -155,7 +155,7 @@ void FDisplayMetrics::RebuildDisplayMetrics(FDisplayMetrics& OutDisplayMetrics)
     OutDisplayMetrics.PrimaryDisplayHeight = OutDisplayMetrics.PrimaryDisplayWorkAreaRect.Bottom - OutDisplayMetrics.PrimaryDisplayWorkAreaRect.Top;
     
 #if !PLATFORM_TVOS
-    const float RequestedContentScaleFactor = View.contentScaleFactor;
+    const double RequestedContentScaleFactor = View.contentScaleFactor;
     
     //we need to set these according to the orientation
     TAutoConsoleVariable<float>* CVar_Left = nullptr;
@@ -180,10 +180,10 @@ void FDisplayMetrics::RebuildDisplayMetrics(FDisplayMetrics& OutDisplayMetrics)
     }
     
     // of the CVars are set, use their values. If not, use what comes from iOS
-    const float Inset_Left = (!CVar_Left || CVar_Left->AsVariable()->GetFloat() < 0.0f) ? CachedInsets.left : CVar_Left->AsVariable()->GetFloat();
-    const float Inset_Top = (!CVar_Top || CVar_Top->AsVariable()->GetFloat() < 0.0f) ? CachedInsets.top : CVar_Top->AsVariable()->GetFloat();
-    const float Inset_Right = (!CVar_Right || CVar_Right->AsVariable()->GetFloat() < 0.0f) ? CachedInsets.right : CVar_Right->AsVariable()->GetFloat();
-    const float Inset_Bottom = (!CVar_Bottom || CVar_Bottom->AsVariable()->GetFloat() < 0.0f) ? CachedInsets.bottom : CVar_Bottom->AsVariable()->GetFloat();
+    const double Inset_Left = (!CVar_Left || CVar_Left->AsVariable()->GetFloat() < 0.0f) ? CachedInsets.left : CVar_Left->AsVariable()->GetFloat();
+    const double Inset_Top = (!CVar_Top || CVar_Top->AsVariable()->GetFloat() < 0.0f) ? CachedInsets.top : CVar_Top->AsVariable()->GetFloat();
+    const double Inset_Right = (!CVar_Right || CVar_Right->AsVariable()->GetFloat() < 0.0f) ? CachedInsets.right : CVar_Right->AsVariable()->GetFloat();
+    const double Inset_Bottom = (!CVar_Bottom || CVar_Bottom->AsVariable()->GetFloat() < 0.0f) ? CachedInsets.bottom : CVar_Bottom->AsVariable()->GetFloat();
     
     //setup the asymmetrical padding
     OutDisplayMetrics.TitleSafePaddingSize.X = Inset_Left;
diff --git a/Engine/Source/Runtime/ApplicationCore/Private/IOS/IOSCursor.cpp b/Engine/Source/Runtime/ApplicationCore/Private/IOS/IOSCursor.cpp
--- a/Engine/Source/Runtime/ApplicationCore/Private/IOS/IOSCursor.cpp
+++ b/Engine/Source/Runtime/ApplicationCore/Private/IOS/IOSCursor.cpp
@@ -58,7 +58,7 @@ void FIOSCursor::Lock(const RECT* const Bounds)
 	FVector2D Position = GetPosition();
 	if (UpdateCursorClipping(Position))
 	{
-		SetPosition(Position.X, Position.Y);
+		SetPosition(FMath::TruncToInt(Position.X), FMath::TruncToInt(Position.Y));
 	}
 }
 
diff --git a/Engine/Source/Runtime/ApplicationCore/Private/IOS/IOSInputInterface.cpp b/Engine/Source/Runtime/ApplicationCore/Private/IOS/IOSInputInterface.cpp
--- a/Engine/Source/Runtime/ApplicationCore/Private/IOS/IOSInputInterface.cpp
+++ b/Engine/Source/Runtime/ApplicationCore/Private/IOS/IOSInputInterface.cpp
@@ -305,7 +305,7 @@ void ModifyVectorByOrientation(FVector& Vec, bool bIsRotation)
             if (bIsRotation)
             {
                 // swap and negate (as needed) roll and pitch
-                float Temp = Vec.X;
+                double Temp = Vec.X;
                 Vec.X = -Vec.Z;
                 Vec.Z = Temp;
                 Vec.Y *= -1.0f;
@@ -313,7 +313,7 @@ void ModifyVectorByOrientation(FVector& Vec, bool bIsRotation)
             else
             {
                 // swap and negate (as needed) x and y
-                float Temp = Vec.X;
+                double Temp = Vec.X;
                 Vec.X = -Vec.Y;
                 Vec.Y = Temp;
             }
@@ -323,14 +323,14 @@ void ModifyVectorByOrientation(FVector& Vec, bool bIsRotation)
             if (bIsRotation)
             {
                 // swap and negate (as needed) roll and pitch
-                float Temp = Vec.X;
+                double Temp = Vec.X;
                 Vec.X = -Vec.Z;
                 Vec.Z = -Temp;
             }
             else
             {
                 // swap and negate (as needed) x and y
-                float Temp = Vec.X;
+                double Temp = Vec.X;
                 Vec.X = Vec.Y;
                 Vec.Y = -Temp;
             }
@@ -374,7 +374,7 @@ void FIOSInputInterface::ProcessTouchesAndKeys(uint32 ControllerId, const TArray
         int32 KeyCode = InKeyInputStack[KeyIndex];
         int32 CharCode = InKeyInputStack[KeyIndex + 1];
         MessageHandler->OnKeyDown(KeyCode, CharCode, false);
-        MessageHandler->OnKeyChar(CharCode,  false);
+        MessageHandler->OnKeyChar((TCHAR)CharCode,  false);
         MessageHandler->OnKeyUp  (KeyCode, CharCode, false);
     }
 }
@@ -508,7 +508,7 @@ void FIOSInputInterface::SendControllerEvents()
     
     if ( bHaveMouse )
     {
-        MessageHandler->OnRawMouseMove(MouseDeltaX, MouseDeltaY);
+        MessageHandler->OnRawMouseMove(FMath::TruncToInt(MouseDeltaX), FMath::TruncToInt(MouseDeltaY));
         MouseDeltaX = 0.f;
         MouseDeltaY = 0.f;
         
@@ -657,8 +657,8 @@ void FIOSInputInterface::GetMovementData(FVector& Attitude, FVector& RotationRat
         FVector FinalAcceleration = -FilteredAccelerometer.GetSafeNormal();
         
         // calculate Roll/Pitch
-        float CurrentPitch = FMath::Atan2(FinalAcceleration.Y, FinalAcceleration.Z);
-        float CurrentRoll = -FMath::Atan2(FinalAcceleration.X, FinalAcceleration.Z);
+        float CurrentPitch = (float)FMath::Atan2(FinalAcceleration.Y, FinalAcceleration.Z);
+        float CurrentRoll = -(float)FMath::Atan2(FinalAcceleration.X, FinalAcceleration.Z);
         
         // if we want to calibrate, use the current values as center
         if (bIsCalibrationRequested)
diff --git a/Engine/Source/Runtime/ApplicationCore/Private/IOS/IOSPlatformApplicationMisc.cpp b/Engine/Source/Runtime/ApplicationCore/Private/IOS/IOSPlatformApplicationMisc.cpp
--- a/Engine/Source/Runtime/ApplicationCore/Private/IOS/IOSPlatformApplicationMisc.cpp
+++ b/Engine/Source/Runtime/ApplicationCore/Private/IOS/IOSPlatformApplicationMisc.cpp
@@ -89,7 +89,7 @@ EAppReturnType::Type MessageBoxExtImpl( EAppMsgType::Type MsgType, const TCHAR*
 
 	while (AppDelegate.AlertResponse == -1)
 	{
-		FPlatformProcess::Sleep(.1);
+		FPlatformProcess::Sleep(0.1f);
 	}
 
 	EAppReturnType::Type Result = (EAppReturnType::Type)AppDelegate.AlertResponse;
@@ -284,16 +284,16 @@ EScreenPhysicalAccuracy FIOSPlatformApplicationMisc::ComputePhysicalScreenDensit
 
 	// If it hasn't been set, assume that the density is a multiple of the 
 	// native Content Scaling Factor.  Won't be exact, but should be close enough.
-	const float NativeScale =[[UIScreen mainScreen] scale];
-	ScreenDensity = 163 * NativeScale;
+	const double NativeScale =[[UIScreen mainScreen] scale];
+	ScreenDensity = FMath::TruncToInt(163 * NativeScale);
 
 	// look up the current scale factor
 	UIView* View = [IOSAppDelegate GetDelegate].IOSView;
-	const float ContentScaleFactor = View.contentScaleFactor;
+	const double ContentScaleFactor = View.contentScaleFactor;
 
 	if ( ContentScaleFactor != 0 )
 	{
-		ScreenDensity = ScreenDensity * ( ContentScaleFactor / NativeScale );
+		ScreenDensity = FMath::TruncToInt(ScreenDensity * (ContentScaleFactor / NativeScale));
 	}
 
 	return EScreenPhysicalAccuracy::Approximation;
diff --git a/Engine/Source/Runtime/ApplicationCore/Private/IOS/IOSPlatformSurvey.cpp b/Engine/Source/Runtime/ApplicationCore/Private/IOS/IOSPlatformSurvey.cpp
--- a/Engine/Source/Runtime/ApplicationCore/Private/IOS/IOSPlatformSurvey.cpp
+++ b/Engine/Source/Runtime/ApplicationCore/Private/IOS/IOSPlatformSurvey.cpp
@@ -23,9 +23,9 @@ bool FIOSPlatformSurvey::GetSurveyResults(FHardwareSurveyResults& OutResults, bo
 
 	// display 0 is max size
 	CGRect MainFrame = [[UIScreen mainScreen] bounds];
-	float Scale = [[UIScreen mainScreen] scale];
-	OutResults.Displays[0].CurrentModeWidth = MainFrame.size.width * Scale;
-	OutResults.Displays[0].CurrentModeHeight = MainFrame.size.height * Scale;
+	double Scale = [[UIScreen mainScreen] scale];
+	OutResults.Displays[0].CurrentModeWidth = FMath::TruncToInt(MainFrame.size.width * Scale);
+	OutResults.Displays[0].CurrentModeHeight = FMath::TruncToInt(MainFrame.size.height * Scale);
 
 	// display 1 is current size
 	FPlatformRect ScreenSize = FIOSWindow::GetScreenRect();
diff --git a/Engine/Source/Runtime/ApplicationCore/Private/IOS/IOSView.cpp b/Engine/Source/Runtime/ApplicationCore/Private/IOS/IOSView.cpp
--- a/Engine/Source/Runtime/ApplicationCore/Private/IOS/IOSView.cpp
+++ b/Engine/Source/Runtime/ApplicationCore/Private/IOS/IOSView.cpp
@@ -272,12 +272,12 @@ id<MTLDevice> GMetalDevice = nil;
 			if (ScreenHeight > ScreenWidth)
 			{
 				Height = RequestedResX;
-				Width = (Height * AspectRatio + 0.5f);
+				Width = FMath::TruncToInt(Height * AspectRatio + 0.5f);
 			}
 			else
 			{
 				Width = RequestedResX;
-				Height = (Width * AspectRatio + 0.5f);
+				Height = FMath::TruncToInt(Width * AspectRatio + 0.5f);
 			}
 		}
 		else if (RequestedResY > 0)
@@ -286,12 +286,12 @@ id<MTLDevice> GMetalDevice = nil;
 			if (ScreenHeight > ScreenWidth)
 			{
 				Width = RequestedResY;
-				Height = (Width * AspectRatio + 0.5f);
+				Height = FMath::TruncToInt(Width * AspectRatio + 0.5f);
 			}
 			else
 			{
 				Height = RequestedResY;
-				Width = (Height * AspectRatio + 0.5f);
+				Width = FMath::TruncToInt(Height * AspectRatio + 0.5f);
 			}
 		}
 		else
@@ -308,7 +308,7 @@ id<MTLDevice> GMetalDevice = nil;
 {
 	if (!bIsInitialized)
 	{
-		[self CalculateContentScaleFactor:self.frame.size.width ScreenHeight:self.frame.size.height];
+		[self CalculateContentScaleFactor:FMath::TruncToInt(self.frame.size.width) ScreenHeight:FMath::TruncToInt(self.frame.size.height)];
 		bIsInitialized = true;
 	}
 	return true;
@@ -427,8 +427,8 @@ self.accessibilityElements = @[Window.accessibilityContainer];
 	TouchInput TouchMessage;
 	TouchMessage.Handle = TouchIndex;
 	TouchMessage.Type = Type;
-	TouchMessage.Position = FVector2D(FMath::Min<float>(_ViewSize.width - 1, Loc.x), FMath::Min<float>(_ViewSize.height - 1, Loc.y)) * Scale;
-	TouchMessage.LastPosition = FVector2D(FMath::Min<float>(_ViewSize.width - 1, PrevLoc.x), FMath::Min<float>(_ViewSize.height - 1, PrevLoc.y)) * Scale;
+	TouchMessage.Position = FVector2D(FMath::Min<double>(_ViewSize.width - 1, Loc.x), FMath::Min<double>(_ViewSize.height - 1, Loc.y)) * Scale;
+	TouchMessage.LastPosition = FVector2D(FMath::Min<double>(_ViewSize.width - 1, PrevLoc.x), FMath::Min<double>(_ViewSize.height - 1, PrevLoc.y)) * Scale;
 	TouchMessage.Force = Type != TouchEnded ? Force : 0.0f;
 	
 	// skip moves that didn't actually move - this will help input handling to skip over the first
@@ -535,21 +535,21 @@ self.accessibilityElements = @[Window.accessibilityContainer];
 			continue;
 		}
 
-		float Force = Touch.force;
+		double Force = Touch.force;
 		
 		// map larger values to 1..10, so 10 is a max across platforms
-		if (Force > 1.0f)
+		if (Force > 1.0)
 		{
-			Force = 10.0f * Force / Touch.maximumPossibleForce;
+			Force = 10.0 * Force / Touch.maximumPossibleForce;
 		}
 		
 		// Handle devices without force touch
-		if ((Type == TouchBegan || Type == TouchMoved) && Force == 0.f)
+		if ((Type == TouchBegan || Type == TouchMoved) && Force == 0.0)
 		{
-			Force = 1.f;
+			Force = 1.0;
 		}
 
-		[self  HandleTouchAtLoc:Loc PrevLoc:PrevLoc TouchIndex:TouchIndex Force:Force Type:Type TouchesArray:TouchesArray];
+		[self  HandleTouchAtLoc:Loc PrevLoc:PrevLoc TouchIndex:TouchIndex Force:(float)Force Type:Type TouchesArray:TouchesArray];
 	}
 
 	FIOSInputInterface::QueueTouchInput(TouchesArray);
@@ -954,10 +954,10 @@ self.accessibilityElements = @[Window.accessibilityContainer];
 	CGRect Frame = [[info objectForKey:UIKeyboardFrameEndUserInfoKey] CGRectValue];
 	
 	FPlatformRect ScreenRect;
-	ScreenRect.Top = Frame.origin.y;
-	ScreenRect.Bottom = (Frame.origin.y + Frame.size.height);
-	ScreenRect.Left = Frame.origin.x;
-	ScreenRect.Right = (Frame.origin.x + Frame.size.width);
+	ScreenRect.Top = FMath::TruncToInt(Frame.origin.y);
+	ScreenRect.Bottom = FMath::TruncToInt(Frame.origin.y + Frame.size.height);
+	ScreenRect.Left = FMath::TruncToInt(Frame.origin.x);
+	ScreenRect.Right = FMath::TruncToInt(Frame.origin.x + Frame.size.width);
 	
 	[FIOSAsyncTask CreateTaskWithBlock:^bool(void){
 		[IOSAppDelegate GetDelegate].IOSApplication->OnVirtualKeyboardShown().Broadcast(ScreenRect);
diff --git a/Engine/Source/Runtime/ApplicationCore/Private/IOS/IOSWindow.cpp b/Engine/Source/Runtime/ApplicationCore/Private/IOS/IOSWindow.cpp
--- a/Engine/Source/Runtime/ApplicationCore/Private/IOS/IOSWindow.cpp
+++ b/Engine/Source/Runtime/ApplicationCore/Private/IOS/IOSWindow.cpp
@@ -95,7 +95,7 @@ void FIOSWindow::OnScaleFactorChanged(IConsoleVariable* CVar)
 		
 	// Force a re-layout of our views as the size has probably changed
 	CGRect Frame = [View frame];
-	[View CalculateContentScaleFactor:Frame.size.width ScreenHeight:Frame.size.height];
+	[View CalculateContentScaleFactor:FMath::TruncToInt(Frame.size.width) ScreenHeight:FMath::TruncToInt(Frame.size.height)];
 	[View layoutSubviews];
 }
 
@@ -109,7 +109,7 @@ void FIOSWindow::OnConsoleResolutionChanged(IConsoleVariable* CVar)
 		
 	// Force a re-layout of our views as the size has probably changed
 	CGRect Frame = [View frame];
-	[View CalculateContentScaleFactor:Frame.size.width ScreenHeight:Frame.size.height];
+	[View CalculateContentScaleFactor:FMath::TruncToInt(Frame.size.width) ScreenHeight:FMath::TruncToInt(Frame.size.height)];
 	[View layoutSubviews];
 }
 
@@ -125,10 +125,10 @@ FPlatformRect FIOSWindow::GetScreenRect()
 		CGRect Frame = [View frame];
 		CGFloat Scale = View.contentScaleFactor;
 		
-		ScreenRect.Top = Frame.origin.y * Scale;
-		ScreenRect.Bottom = (Frame.origin.y + View.ViewSize.height) * Scale;
-		ScreenRect.Left = Frame.origin.x * Scale;
-		ScreenRect.Right = (Frame.origin.x + View.ViewSize.width) * Scale;
+		ScreenRect.Top = FMath::TruncToInt(Frame.origin.y * Scale);
+		ScreenRect.Bottom = FMath::TruncToInt((Frame.origin.y + View.ViewSize.height) * Scale);
+		ScreenRect.Left = FMath::TruncToInt(Frame.origin.x * Scale);
+		ScreenRect.Right = FMath::TruncToInt((Frame.origin.x + View.ViewSize.width) * Scale);
 	}
 
 	return ScreenRect;
@@ -142,10 +142,10 @@ FPlatformRect FIOSWindow::GetUIWindowRect()
 	CGRect Bounds = [Window bounds];
 
 	FPlatformRect WindowRect;
-	WindowRect.Top = Bounds.origin.y;
-	WindowRect.Bottom = Bounds.origin.y + Bounds.size.height;
-	WindowRect.Left = Bounds.origin.x;
-	WindowRect.Right = Bounds.origin.x + Bounds.size.width;
+	WindowRect.Top = FMath::TruncToInt(Bounds.origin.y);
+	WindowRect.Bottom = FMath::TruncToInt(Bounds.origin.y + Bounds.size.height);
+	WindowRect.Left = FMath::TruncToInt(Bounds.origin.x);
+	WindowRect.Right = FMath::TruncToInt(Bounds.origin.x + Bounds.size.width);
 
 	return WindowRect;
 }
