Fix implicit conversion warnings
#rnx

diff --git a/Engine/Source/Developer/DesktopPlatform/Private/Windows/DesktopPlatformWindows.cpp b/Engine/Source/Developer/DesktopPlatform/Private/Windows/DesktopPlatformWindows.cpp
--- a/Engine/Source/Developer/DesktopPlatform/Private/Windows/DesktopPlatformWindows.cpp
+++ b/Engine/Source/Developer/DesktopPlatform/Private/Windows/DesktopPlatformWindows.cpp
@@ -117,7 +117,7 @@ bool FDesktopPlatformWindows::OpenFontDialog(const void* ParentWindowHandle, FSt
 	{
 		HDC DC = ::GetDC( cf.hwndOwner ); 
 		const float LogicalPixelsY = static_cast<float>(GetDeviceCaps(DC, LOGPIXELSY));
-		const int32 PixelHeight = static_cast<int32>(-lf.lfHeight * ( 72.0f / LogicalPixelsY ));	// Always target 72 DPI
+		const int32 PixelHeight = static_cast<int32>((float) - lf.lfHeight * (72.0f / LogicalPixelsY));	// Always target 72 DPI
 		auto FontFlags = EFontImportFlags::None;
 		if ( lf.lfUnderline )
 		{
diff --git a/Engine/Source/Runtime/ApplicationCore/Private/Windows/WindowsApplication.cpp b/Engine/Source/Runtime/ApplicationCore/Private/Windows/WindowsApplication.cpp
--- a/Engine/Source/Runtime/ApplicationCore/Private/Windows/WindowsApplication.cpp
+++ b/Engine/Source/Runtime/ApplicationCore/Private/Windows/WindowsApplication.cpp
@@ -815,10 +815,10 @@ static void GetMonitorsInfo(TArray<FMonitorInfo>& OutMonitorInfo)
 						else
 						{
 							// we also need to include the OS scaling value
-							const float CenterX = 0.5f * (Info.WorkArea.Right + Info.WorkArea.Left);
-							const float CenterY = 0.5f * (Info.WorkArea.Top + Info.WorkArea.Bottom);
+							const float CenterX = 0.5f * float(Info.WorkArea.Right + Info.WorkArea.Left);
+							const float CenterY = 0.5f * float(Info.WorkArea.Top + Info.WorkArea.Bottom);
 							const float DPIScaleFactor = FWindowsPlatformApplicationMisc::GetDPIScaleFactorAtPoint(CenterX, CenterY);
-							Info.DPI = (int32)(Info.DPI * DPIScaleFactor);
+							Info.DPI = (int32)((float)Info.DPI * DPIScaleFactor);
 						}
 
 						// The editor shouldn't care about lower level display resolutions. This is only necessary for fullscreen exclusive situations.
@@ -1134,8 +1134,8 @@ int32 FWindowsApplication::ProcessMessage( HWND hwnd, uint32 msg, WPARAM wParam,
 								const int32 Width = GetSystemMetrics(IsVirtualScreen ? SM_CXVIRTUALSCREEN : SM_CXSCREEN);
 								const int32 Height = GetSystemMetrics(IsVirtualScreen ? SM_CYVIRTUALSCREEN : SM_CYSCREEN);
 
-								CursorPoint.x = static_cast<int>((float(Raw->data.mouse.lLastX) / 65535.0f) * Width) + Left;
-								CursorPoint.y = static_cast<int>((float(Raw->data.mouse.lLastY) / 65535.0f) * Height) + Top;
+								CursorPoint.x = static_cast<int>((float(Raw->data.mouse.lLastX) / 65535.0f) * (float)Width) + Left;
+								CursorPoint.y = static_cast<int>((float(Raw->data.mouse.lLastY) / 65535.0f) * (float)Height) + Top;
 
 								const int32 ClipWidth = ClipCursorRect.right - ClipCursorRect.left;
 								const int32 ClipHeight = ClipCursorRect.bottom - ClipCursorRect.top;
@@ -1363,7 +1363,7 @@ int32 FWindowsApplication::ProcessMessage( HWND hwnd, uint32 msg, WPARAM wParam,
 							double WaitTimeMilliseconds = 0.0;
 							if (Freq.QuadPart > 0)
 							{
-								WaitTimeMilliseconds = 1000.0 * WaitTime / (double)Freq.QuadPart;
+								WaitTimeMilliseconds = 1000.0 * (double)WaitTime / (double)Freq.QuadPart;
 							}
 
 							// Due time for WaitForSingleObject is in 100 nanosecond units.							
@@ -1504,7 +1504,7 @@ int32 FWindowsApplication::ProcessMessage( HWND hwnd, uint32 msg, WPARAM wParam,
 					case WMSZ_LEFT:
 					case WMSZ_RIGHT:
 						{
-							int32 AdjustedHeight = (int32)(NewWidth / AspectRatio);
+							int32 AdjustedHeight = (int32)((float)NewWidth / AspectRatio);
 							Rect->top -= (AdjustedHeight - NewHeight) / 2;
 							Rect->bottom += (AdjustedHeight - NewHeight) / 2;
 							break;
@@ -1512,32 +1512,32 @@ int32 FWindowsApplication::ProcessMessage( HWND hwnd, uint32 msg, WPARAM wParam,
 					case WMSZ_TOP:
 					case WMSZ_BOTTOM:
 						{
-							int32 AdjustedWidth = (int32)(NewHeight * AspectRatio);
+							int32 AdjustedWidth = (int32)((float)NewHeight * AspectRatio);
 							Rect->left -= (AdjustedWidth - NewWidth) / 2;
 							Rect->right += (AdjustedWidth - NewWidth) / 2;
 							break;
 						}
 					case WMSZ_TOPLEFT:
 						{
-							int32 AdjustedHeight = (int32)(NewWidth / AspectRatio);
+							int32 AdjustedHeight = (int32)((float)NewWidth / AspectRatio);
 							Rect->top -= AdjustedHeight - NewHeight;
 							break;
 						}
 					case WMSZ_TOPRIGHT:
 						{
-							int32 AdjustedHeight = (int32)(NewWidth / AspectRatio);
+							int32 AdjustedHeight = (int32)((float)NewWidth / AspectRatio);
 							Rect->top -= AdjustedHeight - NewHeight;
 							break;
 						}
 					case WMSZ_BOTTOMLEFT:
 						{
-							int32 AdjustedHeight = (int32)(NewWidth / AspectRatio);
+							int32 AdjustedHeight = (int32)((float)NewWidth / AspectRatio);
 							Rect->bottom += AdjustedHeight - NewHeight;
 							break;
 						}
 					case WMSZ_BOTTOMRIGHT:
 						{
-							int32 AdjustedHeight = (int32)(NewWidth / AspectRatio);
+							int32 AdjustedHeight = (int32)((float)NewWidth / AspectRatio);
 							Rect->bottom += AdjustedHeight - NewHeight;
 							break;
 						}
@@ -2309,7 +2309,7 @@ int32 FWindowsApplication::ProcessDeferredMessage( const FDeferredWindowsMessage
 						for ( uint32 i = 0; i < InputCount; i++ )
 						{
 							TOUCHINPUT Input = Inputs[i];
-							FVector2D Location( Input.x / 100.0f, Input.y / 100.0f );
+							FVector2D Location((float)Input.x / 100.0f, (float)Input.y / 100.0f );
 							if ( Input.dwFlags & TOUCHEVENTF_DOWN )
 							{
 								int32 TouchIndex = GetTouchIndexForID( Input.dwID );
diff --git a/Engine/Source/Runtime/ApplicationCore/Private/Windows/WindowsConsoleOutputDevice2.cpp b/Engine/Source/Runtime/ApplicationCore/Private/Windows/WindowsConsoleOutputDevice2.cpp
--- a/Engine/Source/Runtime/ApplicationCore/Private/Windows/WindowsConsoleOutputDevice2.cpp
+++ b/Engine/Source/Runtime/ApplicationCore/Private/Windows/WindowsConsoleOutputDevice2.cpp
@@ -855,7 +855,7 @@ public:
 		if (MoveSpace <= 0)
 			return;
 
-		int32 ThumbPos = ButtonHeight + FMath::CeilToInt(float(TopVisible * MoveSpace) / TotalScrollCount);
+		int32 ThumbPos = ButtonHeight + FMath::CeilToInt(float(TopVisible * MoveSpace) / (float)TotalScrollCount);
 
 		OutPos = ThumbPos;
 		OutHeight = ThumbHeight;
@@ -875,7 +875,7 @@ public:
 
 		int32 MoveSpace = ScrollHeight - ThumbHeight;
 
-		return FMath::FloorToInt(float(ThumbPos - ButtonHeight) * TotalScrollCount / MoveSpace);
+		return FMath::FloorToInt((float)(ThumbPos - ButtonHeight) * (float)TotalScrollCount / (float)MoveSpace);
 	}
 
 	void SetTopVisible(int32 TopVisible, bool ShouldPost)
@@ -1806,7 +1806,7 @@ public:
 				int32 ItemCount = (int)SendMessageW(LogHwnd, LB_GETCOUNT, 0, 0);
 				RECT Rect;
 				GetClientRect(LogHwnd, &Rect);
-				int32 VisibleCount = FMath::CeilToInt(float(Rect.bottom) / LogFontHeight);
+				int32 VisibleCount = FMath::CeilToInt(float(Rect.bottom) / (float)LogFontHeight);
 				int32 TopIndex = IntCastChecked<int32>(SendMessageW(LogHwnd, LB_GETTOPINDEX, 0, 0));
 				if (ItemCount - TopIndex < VisibleCount)
 				{
diff --git a/Engine/Source/Runtime/ApplicationCore/Private/Windows/WindowsPlatformSurvey.cpp b/Engine/Source/Runtime/ApplicationCore/Private/Windows/WindowsPlatformSurvey.cpp
--- a/Engine/Source/Runtime/ApplicationCore/Private/Windows/WindowsPlatformSurvey.cpp
+++ b/Engine/Source/Runtime/ApplicationCore/Private/Windows/WindowsPlatformSurvey.cpp
@@ -104,7 +104,7 @@ bool FWindowsPlatformSurvey::GetSurveyResults( FHardwareSurveyResults& OutResult
 			NTSTATUS NTStatus = CallNtPowerInformation(ProcessorInformation, NULL, 0L, (PVOID)PowerInfo, PowerInfoSize);
 			if (NT_SUCCESS(NTStatus))
 			{
-				OutResults.CPUClockGHz = 0.001f * PowerInfo[0].MaxMhz;
+				OutResults.CPUClockGHz = 0.001f * (float)PowerInfo[0].MaxMhz;
 			}
 			else
 			{
