Add bluetooth keyboard on tvOS

diff --git a/Engine/Source/Runtime/ApplicationCore/Private/IOS/IOSPlatformApplicationMisc.cpp b/Engine/Source/Runtime/ApplicationCore/Private/IOS/IOSPlatformApplicationMisc.cpp
--- a/Engine/Source/Runtime/ApplicationCore/Private/IOS/IOSPlatformApplicationMisc.cpp
+++ b/Engine/Source/Runtime/ApplicationCore/Private/IOS/IOSPlatformApplicationMisc.cpp
@@ -301,5 +301,9 @@ EScreenPhysicalAccuracy FIOSPlatformApplicationMisc::ComputePhysicalScreenDensit
 
 bool FIOSPlatformApplicationMisc::RequiresVirtualKeyboard()
 {
-    return PLATFORM_HAS_TOUCH_MAIN_SCREEN || PLATFORM_TVOS;
+#if !PLATFORM_TVOS
+    return PLATFORM_HAS_TOUCH_MAIN_SCREEN
+#else
+    return true;
+#endif
 }
