Fix build iOS

diff --git a/Engine/Source/Runtime/ApplicationCore/Private/IOS/IOSPlatformApplicationMisc.cpp b/Engine/Source/Runtime/ApplicationCore/Private/IOS/IOSPlatformApplicationMisc.cpp
--- a/Engine/Source/Runtime/ApplicationCore/Private/IOS/IOSPlatformApplicationMisc.cpp
+++ b/Engine/Source/Runtime/ApplicationCore/Private/IOS/IOSPlatformApplicationMisc.cpp
@@ -302,7 +302,7 @@ EScreenPhysicalAccuracy FIOSPlatformApplicationMisc::ComputePhysicalScreenDensit
 bool FIOSPlatformApplicationMisc::RequiresVirtualKeyboard()
 {
 #if !PLATFORM_TVOS
-    return PLATFORM_HAS_TOUCH_MAIN_SCREEN
+    return PLATFORM_HAS_TOUCH_MAIN_SCREEN;
 #else
     return true;
 #endif
