Remove/replace uses of master/slave found in ApplicationCore with base folder Runtime

diff --git a/Engine/Source/Runtime/ApplicationCore/Private/Mac/HIDInputInterface.cpp b/Engine/Source/Runtime/ApplicationCore/Private/Mac/HIDInputInterface.cpp
--- a/Engine/Source/Runtime/ApplicationCore/Private/Mac/HIDInputInterface.cpp
+++ b/Engine/Source/Runtime/ApplicationCore/Private/Mac/HIDInputInterface.cpp
@@ -257,7 +257,7 @@ void HIDInputInterface::FHIDDeviceInfo::SetupMappings()
 	}
 	else
 	{
-		// Array generated from [REDACTED_PROFILE_URL]/360Controller/blob/master/360Controller/Info.plist (with additions)
+		// Array generated from [REDACTED_PROFILE_URL]/360Controller-n/blob/c83340e55be337e28a149a2930657006d3a3fef8/360Controller/Info.plist (with additions)
 		static const FXBox360ControllerID XBox360Controllers[] =
 		{
 			{ 0x0E6F, 0x0113 }, 	// AfterglowGamepadforXbox360
diff --git a/Engine/Source/Runtime/ApplicationCore/Private/Mac/MacApplication.cpp b/Engine/Source/Runtime/ApplicationCore/Private/Mac/MacApplication.cpp
--- a/Engine/Source/Runtime/ApplicationCore/Private/Mac/MacApplication.cpp
+++ b/Engine/Source/Runtime/ApplicationCore/Private/Mac/MacApplication.cpp
@@ -2294,7 +2294,16 @@ void FDisplayMetrics::RebuildDisplayMetrics(FDisplayMetrics& OutDisplayMetrics)
 
 			// Monitor's name can only be obtained from IOKit
 			io_iterator_t IOIterator;
-			kern_return_t Result = IOServiceGetMatchingServices(kIOMasterPortDefault, IOServiceMatching("IODisplayConnect"), &IOIterator);
+            mach_port_t mainPortDefault;
+            if (@available(macOS 12.0, iOS 15.0, *))
+            {
+                mainPortDefault = kIOMainPortDefault;
+            }
+            else
+            {
+                mainPortDefault = kIOMasterPortDefault;
+            }
+			kern_return_t Result = IOServiceGetMatchingServices(kIOMainPortDefault, IOServiceMatching("IODisplayConnect"), &IOIterator);
 			if (Result == kIOReturnSuccess)
 			{
 				io_object_t Device;
diff --git a/Engine/Source/Runtime/ApplicationCore/Private/Mac/MacCursor.cpp b/Engine/Source/Runtime/ApplicationCore/Private/Mac/MacCursor.cpp
--- a/Engine/Source/Runtime/ApplicationCore/Private/Mac/MacCursor.cpp
+++ b/Engine/Source/Runtime/ApplicationCore/Private/Mac/MacCursor.cpp
@@ -144,15 +144,23 @@ FMacCursor::FMacCursor()
 
 PRAGMA_DISABLE_DEPRECATION_WARNINGS
 	// Get the IOHIDSystem so we can disable mouse acceleration
-	mach_port_t MasterPort;
-	kern_return_t KernResult = IOMasterPort(MACH_PORT_NULL, &MasterPort);
+	mach_port_t MainPort;
+    kern_return_t KernResult;
+    if (@available(macOS 12.0, iOS 15.0, *))
+    {
+        KernResult = IOMainPort(MACH_PORT_NULL, &MainPort);
+    }
+    else
+    {
+        KernResult = IOMasterPort(MACH_PORT_NULL, &MainPort);
+    }
 	if (KERN_SUCCESS == KernResult)
 	{
 		CFMutableDictionaryRef ClassesToMatch = IOServiceMatching("IOHIDSystem");
 		if (ClassesToMatch)
 		{
 			io_iterator_t MatchingServices;
-			KernResult = IOServiceGetMatchingServices(MasterPort, ClassesToMatch, &MatchingServices);
+			KernResult = IOServiceGetMatchingServices(MainPort, ClassesToMatch, &MatchingServices);
 			if (KERN_SUCCESS == KernResult)
 			{
 				io_object_t IntfService;
