[Core Runtime - WindowsConsoleOutputDevice2]
* Fixed buggy datetime output.
* Added so shell execute can be done from activity texts. Just add [<command>] to the activity and the command(s) will execute the command when clicking on the activity. Example [https://someurl] will trigger the browser to open the linik

diff --git a/Engine/Source/Runtime/ApplicationCore/Private/Windows/WindowsConsoleOutputDevice2.cpp b/Engine/Source/Runtime/ApplicationCore/Private/Windows/WindowsConsoleOutputDevice2.cpp
--- a/Engine/Source/Runtime/ApplicationCore/Private/Windows/WindowsConsoleOutputDevice2.cpp
+++ b/Engine/Source/Runtime/ApplicationCore/Private/Windows/WindowsConsoleOutputDevice2.cpp
@@ -67,6 +67,7 @@ public:
 		ID_COMMANDBUTTON,
 		ID_CLEARLOGBUTTON,
 		ID_ADDCHECKPOINTBUTTON,
+		ID_TRACKEDACTIVITY,
 	};
 	enum : int
 	{
@@ -76,6 +77,9 @@ public:
 	FConsoleWindow(FWindowsConsoleOutputDevice2* InOwner)
 	:	Owner(InOwner)
 	{
+		StartDateTime = FDateTime::Now();
+		StartTime = FPlatformTime::Seconds() - GStartTime;
+
 		// Set up status light brushes
 		COLORREF Colors[] = {RGB(237, 28, 36), RGB(255, 255, 0), RGB(128, 255, 128)};
 		for (int32 I = 0; I != 3; ++I)
@@ -507,17 +511,17 @@ public:
 		DeleteObject(ThumbBrush);
 	}
 
-	HWND CreateWindow2(const TCHAR* ClassName, const TCHAR* Str, DWORD Style, HFONT hFont, int32 X = 0, int32 Y = 0, int32 Width = 1, int32 Height = 1, ENotificationId InNid = (ENotificationId)0, DWORD exStyle = 0)
+	HWND CreateWindow2(const TCHAR* ClassName, const TCHAR* Str, DWORD Style, HFONT hFont, int32 X = 0, int32 Y = 0, int32 Width = 1, int32 Height = 1, ENotificationId InNid = (ENotificationId)0, DWORD ExStyle  = 0)
 	{
-		HWND h = CreateWindowEx(exStyle, ClassName, Str, Style | WS_VISIBLE | WS_CHILD, X, Y, Width, Height, MainHwnd, (HMENU)InNid, NULL, NULL);
+		HWND h = CreateWindowEx(ExStyle , ClassName, Str, Style | WS_VISIBLE | WS_CHILD, X, Y, Width, Height, MainHwnd, (HMENU)InNid, NULL, NULL);
 		if (hFont)
 			SendMessageW(h, WM_SETFONT, (WPARAM)hFont, 0);
 		return h;
 	}
 
-	HWND CreateTextHwnd(const TCHAR* Str, HFONT hFont, int32 X = 0, int32 Y = 0, int32 Width = 1, int32 Height = 1, DWORD exStyle = 0)
+	HWND CreateTextHwnd(const TCHAR* Str, HFONT hFont, ENotificationId InNid = (ENotificationId)0, DWORD Style = 0, DWORD ExStyle = 0)
 	{
-		HWND h = CreateWindow2(WC_STATIC, Str, SS_OWNERDRAW, hFont, X, Y, Width, Height, (ENotificationId)0, exStyle);
+		HWND h = CreateWindow2(WC_STATIC, Str, Style | SS_OWNERDRAW, hFont, 0, 0, 1, 1, InNid, ExStyle );
 		SetWindowLong(h, GWLP_USERDATA, -1);
 		return h;
 	}
@@ -1114,7 +1118,7 @@ public:
 		}
 		else if (ConsoleShowDateTime == 2)
 		{
-			FDateTime::Now().ToString(TEXT("[%Y.%m.%d-%H.%M.%S:%s]"), OutString);
+			(StartDateTime + FTimespan::FromSeconds(E.Time - StartTime)).ToString(TEXT("[%Y.%m.%d-%H.%M.%S:%s] "), OutString);
 			//OutString.Appendf(TEXT("[%3llu] "), GFrameCounter % 1000);
 		}
 
@@ -1502,6 +1506,64 @@ public:
 		return CallNextHookEx(0, nCode, wParam, lParam);
 	}
 
+	bool ExtractHyperLinks(TStringBuilder<256>& OutStatus, TArray<FString>* OutHyperLinks, const FString& String)
+	{
+		const TCHAR* Str = *String;
+		while (true)
+		{
+			const TCHAR* LinkScopeStart = FCString::Strchr(Str, '[');
+			if (LinkScopeStart == nullptr)
+			{
+				break;
+			}
+
+			++LinkScopeStart;
+
+			const TCHAR* LinkScopeEnd = FCString::Strchr(LinkScopeStart, ']');
+			if (LinkScopeEnd == nullptr)
+			{
+				break;
+			}
+
+			OutStatus.Append(Str, int32(LinkScopeStart - Str - 1));
+			Str = LinkScopeEnd + 1;
+
+			if (!OutHyperLinks)
+			{
+				continue;
+			}
+
+			FStringView LinkScope = FStringView(LinkScopeStart, int32(LinkScopeEnd - LinkScopeStart)).TrimStartAndEnd();
+
+			while (true)
+			{
+				int32 SpaceIndex;
+				if (LinkScope.FindChar(' ', SpaceIndex))
+				{
+					FStringView Link = LinkScope.Left(SpaceIndex).TrimStartAndEnd();
+					if (!Link.IsEmpty())
+					{
+						OutHyperLinks->Add(FString(Link));
+					}
+					LinkScope.RightChopInline(SpaceIndex + 1);
+				}
+				else
+				{
+					FStringView Link = LinkScope.TrimStartAndEnd();
+					if (!Link.IsEmpty())
+					{
+						OutHyperLinks->Add(FString(Link));
+					}
+					break;
+				}
+			}
+		}
+
+		OutStatus << Str;
+
+		return OutHyperLinks && OutHyperLinks->Num() != 0;
+	}
+
 	void HandleActivityModifications()
 	{
 		bool bActivitiesLightDirty = false;
@@ -1585,13 +1647,19 @@ public:
 			if (!A.NameHwnd)
 			{
 				A.NameHwnd = CreateTextHwnd(*A.Name, Font);
-				A.StatusHwnd = CreateTextHwnd(*A.Status, Font);
+
+				TStringBuilder<256> Status;
+				ExtractHyperLinks(Status, nullptr, A.Status);
+				A.StatusHwnd = CreateTextHwnd(*Status, Font, ID_TRACKEDACTIVITY, SS_NOTIFY);
+
 				bUpdatePositions = true;
 			}
 			else if (A.bStatusDirty)
 			{
 				A.bStatusDirty = false;
-				SetWindowTextW(A.StatusHwnd, *A.Status);
+				TStringBuilder<256> Status;
+				ExtractHyperLinks(Status, nullptr, A.Status);
+				SetWindowTextW(A.StatusHwnd, *Status);
 			}
 			++i;
 		}
@@ -2112,6 +2180,22 @@ public:
 				FCString::Sprintf(ButtonString, TEXT("Log CHECKPOINT%i"), CheckpointIndex);
 				SetDlgItemText(MainHwnd, ID_ADDCHECKPOINTBUTTON, ButtonString);
 			}
+			else if (LOWORD(wParam) == ID_TRACKEDACTIVITY)
+			{
+				TStringBuilder<256> Status;
+				TArray<FString> HyperLinks;
+
+				for (const Activity& A : Activities)
+				{
+					if (A.StatusHwnd == (HWND)lParam)
+					{
+						ExtractHyperLinks(Status, &HyperLinks, A.Status);
+					}
+				}
+
+				for (const FString& Link : HyperLinks)
+					FPlatformProcess::LaunchURL(*Link, NULL, NULL);
+			}
 			break;
 
 		case DM_GETDEFID:
@@ -2362,6 +2446,8 @@ public:
 	TArray<FString> IncludeFilter;
 	TArray<FString> ExcludeFilter;
 	TRingBuffer<LogEntry> Log;
+	FDateTime StartDateTime;
+	double StartTime;
 	int32 AddedEntryLogVirtualIndex = -1;
 	HICON Icon;
 	HFONT Font;
@@ -2804,7 +2890,7 @@ void FWindowsConsoleOutputDevice2::Serialize( const TCHAR* Data, ELogVerbosity::
 			FReadScopeLock lock(WindowRWLock);
 			if (Window)
 			{
-				Window->AddLogEntry(Data, Verbosity, Category, FPlatformTime::Seconds() - GStartTime, TextAttribute);
+				Window->AddLogEntry(Data, Verbosity, Category, RealTime, TextAttribute);
 			}
 		}
 
