Moved DebugConsole.ini to EngineUserDir's /Config/DebugConsole.ini

diff --git a/Engine/Source/Runtime/ApplicationCore/Private/Windows/WindowsConsoleOutputDevice2.cpp b/Engine/Source/Runtime/ApplicationCore/Private/Windows/WindowsConsoleOutputDevice2.cpp
--- a/Engine/Source/Runtime/ApplicationCore/Private/Windows/WindowsConsoleOutputDevice2.cpp
+++ b/Engine/Source/Runtime/ApplicationCore/Private/Windows/WindowsConsoleOutputDevice2.cpp
@@ -2684,7 +2684,7 @@ void FWindowsConsoleOutputDevice2::SaveToINI()
 	bool bConsoleExpanded = Window->bConsoleExpanded;
 	WindowRWLock.WriteUnlock();
 
-	FString Filename = FPlatformProcess::GetModulesDirectory() + TEXT("\\DebugConsole.ini");
+	FString Filename = GetConfigFilename();
 	const TCHAR* Selection = TEXT("ConsoleWindows");
 	if (IsRunningDedicatedServer())
 		Selection = TEXT("ServerConsoleWindows");
@@ -2712,6 +2712,12 @@ void FWindowsConsoleOutputDevice2::SaveToINI()
 	Config.Flush(false, Filename);
 }
 
+const FString& FWindowsConsoleOutputDevice2::GetConfigFilename()
+{
+	static FString Filename = FPaths::EngineUserDir() / TEXT("Config") / TEXT("DebugConsole.ini");
+	return Filename;
+}
+
 void FWindowsConsoleOutputDevice2::Show( bool ShowWindow )
 {
 	if (ShowWindow)
@@ -2731,7 +2737,8 @@ void FWindowsConsoleOutputDevice2::Show( bool ShowWindow )
 		bool bHasX = false;
 		bool bHasY = false;
 
-		FString Filename = FPlatformProcess::GetModulesDirectory() + TEXT("\\DebugConsole.ini");
+		FString Filename = GetConfigFilename();
+
 		const TCHAR* Selection = TEXT("ConsoleWindows");
 		if (IsRunningDedicatedServer())
 			Selection = TEXT("ServerConsoleWindows");
@@ -2740,6 +2747,10 @@ void FWindowsConsoleOutputDevice2::Show( bool ShowWindow )
 
 		FConfigCacheIni Config(EConfigCacheType::Temporary);
 		Config.LoadFile(Filename);
+		if (!Config.FindConfigFile(Filename))
+		{
+			Filename = FPlatformProcess::GetModulesDirectory() / TEXT("DebugConsole.ini");
+		}
 		Config.GetInt(Selection, TEXT("ConsoleWidth"), ConsoleWidth, Filename);
 		Config.GetInt(Selection, TEXT("ConsoleHeight"), ConsoleHeight, Filename);
 		bHasX = Config.GetInt(Selection, TEXT("ConsoleX"), ConsolePosX, Filename);
diff --git a/Engine/Source/Runtime/ApplicationCore/Public/Windows/WindowsConsoleOutputDevice2.h b/Engine/Source/Runtime/ApplicationCore/Public/Windows/WindowsConsoleOutputDevice2.h
--- a/Engine/Source/Runtime/ApplicationCore/Public/Windows/WindowsConsoleOutputDevice2.h
+++ b/Engine/Source/Runtime/ApplicationCore/Public/Windows/WindowsConsoleOutputDevice2.h
@@ -59,6 +59,8 @@ private:
 	 */
 	void SaveToINI();
 
+	static const FString& GetConfigFilename();
+
 public:
 
 	/** 
