[ApplicationCore]
* WindowsConsoleOutputDevice2 - Changed ::SetThreadDescription to use wrapped function that does not depend on new versions of windows runtime

diff --git a/Engine/Source/Runtime/ApplicationCore/Private/Windows/WindowsConsoleOutputDevice2.cpp b/Engine/Source/Runtime/ApplicationCore/Private/Windows/WindowsConsoleOutputDevice2.cpp
--- a/Engine/Source/Runtime/ApplicationCore/Private/Windows/WindowsConsoleOutputDevice2.cpp
+++ b/Engine/Source/Runtime/ApplicationCore/Private/Windows/WindowsConsoleOutputDevice2.cpp
@@ -166,14 +166,11 @@ public:
 	void Start()
 	{
 		Thread = CreateThread(NULL, 0, StaticThreadProc, this, 0, NULL);
-		if (ensure(Thread))
-		{
-			SetThreadDescription(Thread, TEXT("LogConsoleHwnd"));
-		}
 	}
 
 	static DWORD WINAPI StaticThreadProc(LPVOID lpParameter)
 	{
+		FWindowsPlatformProcess::SetThreadName(TEXT("LogConsoleHwnd"));
 		return ((FConsoleWindow*)lpParameter)->ThreadProc();
 	}
 
