Fix DebugConsole.ini so it is placed in /Engine/Saved/Config instead of /Engine/Config

diff --git a/Engine/Source/Runtime/ApplicationCore/Private/Windows/WindowsConsoleOutputDevice2.cpp b/Engine/Source/Runtime/ApplicationCore/Private/Windows/WindowsConsoleOutputDevice2.cpp
--- a/Engine/Source/Runtime/ApplicationCore/Private/Windows/WindowsConsoleOutputDevice2.cpp
+++ b/Engine/Source/Runtime/ApplicationCore/Private/Windows/WindowsConsoleOutputDevice2.cpp
@@ -2711,7 +2711,7 @@ void FWindowsConsoleOutputDevice2::SaveToINI()
 
 const FString& FWindowsConsoleOutputDevice2::GetConfigFilename()
 {
-	static FString Filename = FPaths::EngineUserDir() / TEXT("Config") / TEXT("DebugConsole.ini");
+	static FString Filename = FPaths::EngineSavedDir() / TEXT("Config") / TEXT("DebugConsole.ini");
 	return Filename;
 }
 
@@ -2744,10 +2744,6 @@ void FWindowsConsoleOutputDevice2::Show( bool ShowWindow )
 
 		FConfigCacheIni Config(EConfigCacheType::Temporary);
 		Config.LoadFile(Filename);
-		if (!Config.FindConfigFile(Filename))
-		{
-			Filename = FPlatformProcess::GetModulesDirectory() / TEXT("DebugConsole.ini");
-		}
 		Config.GetInt(Selection, TEXT("ConsoleWidth"), ConsoleWidth, Filename);
 		Config.GetInt(Selection, TEXT("ConsoleHeight"), ConsoleHeight, Filename);
 		bHasX = Config.GetInt(Selection, TEXT("ConsoleX"), ConsolePosX, Filename);
