Fixed deadlock in the splash screen happening when the main thread goes faster than the splash screen and try to hide the splash before the splash screen thread even created the window.

diff --git a/Engine/Source/Runtime/ApplicationCore/Private/Windows/WindowsPlatformSplash.cpp b/Engine/Source/Runtime/ApplicationCore/Private/Windows/WindowsPlatformSplash.cpp
--- a/Engine/Source/Runtime/ApplicationCore/Private/Windows/WindowsPlatformSplash.cpp
+++ b/Engine/Source/Runtime/ApplicationCore/Private/Windows/WindowsPlatformSplash.cpp
@@ -50,7 +50,7 @@ static HFONT GSplashScreenSmallTextFontHandle = NULL;
 static HFONT GSplashScreenNormalTextFontHandle = NULL;
 static HFONT GSplashScreenTitleTextFontHandle = NULL;
 static FCriticalSection GSplashScreenSynchronizationObject;
-
+static HANDLE GSplashWindowCreationEvent = INVALID_HANDLE_VALUE;
 
 
 /**
@@ -396,53 +396,63 @@ uint32 WINAPI StartSplashScreenThread( LPVOID unused )
 	wc.hbrBackground = (HBRUSH)GetStockObject(WHITE_BRUSH);
 	wc.lpszMenuName  = NULL;
 	wc.lpszClassName = TEXT("SplashScreenClass"); 
- 
-	if(!RegisterClass(&wc)) 
-	{
-		return 0; 
-	} 
 
-	// Load splash screen image, display it and handle all window's messages
-	GSplashScreenBitmap = LoadSplashBitmap();
-	if(GSplashScreenBitmap)
+	BITMAP bm;
+	FMemory::Memzero(bm);
+
+	const bool bAllowFading = false;
 	{
-		BITMAP bm;
-		GetObjectW(GSplashScreenBitmap, sizeof(bm), &bm);
-		const int32 WindowWidth = bm.bmWidth;
-		const int32 WindowHeight = bm.bmHeight;
-		int32 ScreenPosX = (GetSystemMetrics(SM_CXSCREEN) - WindowWidth) / 2;
-		int32 ScreenPosY = (GetSystemMetrics(SM_CYSCREEN) - WindowHeight) / 2;
-
-		const bool bAllowFading = false;
-
-		// Force the editor splash screen to show up in the taskbar and alt-tab lists
-		uint32 dwWindowStyle = GIsEditor ? WS_EX_APPWINDOW : WS_EX_TOOLWINDOW;
-		if( bAllowFading )
+		// Ensure that once this scope exits, we signal the main thread that the code constructing the splash window was executed (successfully or not)
+		// so the main thread can read GSplashScreenWnd value and act properly to resume the execution.
+		ON_SCOPE_EXIT { ::SetEvent(GSplashWindowCreationEvent); };
+	
+		if(!RegisterClass(&wc)) 
 		{
-			dwWindowStyle |= WS_EX_LAYERED;
+			return 0; 
+		} 
+
+		// Load splash screen image, display it and handle all window's messages
+		GSplashScreenBitmap = LoadSplashBitmap();
+		if (GSplashScreenBitmap)
+		{
+			GetObjectW(GSplashScreenBitmap, sizeof(bm), &bm);
+			const int32 WindowWidth = bm.bmWidth;
+			const int32 WindowHeight = bm.bmHeight;
+			int32 ScreenPosX = (GetSystemMetrics(SM_CXSCREEN) - WindowWidth) / 2;
+			int32 ScreenPosY = (GetSystemMetrics(SM_CYSCREEN) - WindowHeight) / 2;
+
+			// Force the editor splash screen to show up in the taskbar and alt-tab lists
+			uint32 dwWindowStyle = GIsEditor ? WS_EX_APPWINDOW : WS_EX_TOOLWINDOW;
+			if( bAllowFading )
+			{
+				dwWindowStyle |= WS_EX_LAYERED;
+			}
+
+			GSplashScreenWnd = CreateWindowEx(
+				dwWindowStyle,
+				wc.lpszClassName, 
+				TEXT("SplashScreen"),
+				WS_POPUP,
+				ScreenPosX,
+				ScreenPosY,
+				WindowWidth,
+				WindowHeight,
+				(HWND) NULL,
+				(HMENU) NULL,
+				hInstance,
+				(LPVOID) NULL);
 		}
+	}
 
-		GSplashScreenWnd = CreateWindowEx(
-			dwWindowStyle,
-			wc.lpszClassName, 
-			TEXT("SplashScreen"),
-			WS_POPUP,
-			ScreenPosX,
-			ScreenPosY,
-			WindowWidth,
-			WindowHeight,
-			(HWND) NULL,
-			(HMENU) NULL,
-			hInstance,
-			(LPVOID) NULL); 
-
-		if( bAllowFading )
+	if (GSplashScreenWnd) // also implies GSplashScreenBitmap is not null.
+	{
+		check(GSplashScreenBitmap);
+		if (bAllowFading)
 		{
 			// Set window to fully transparent to start out
 			SetLayeredWindowAttributes( GSplashScreenWnd, 0, 0, LWA_ALPHA );
 		}
 
-
 		// Setup font
 		{
 			HFONT SystemFontHandle = ( HFONT )GetStockObject( DEFAULT_GUI_FONT );
@@ -532,53 +542,50 @@ uint32 WINAPI StartSplashScreenThread( LPVOID unused )
 		GSplashScreenTextRects[ SplashTextType::StartupProgress ].left = 10;
 		GSplashScreenTextRects[ SplashTextType::StartupProgress ].right = bm.bmWidth - 20;
 
-		if (GSplashScreenWnd)
-		{
-			SetWindowText(GSplashScreenWnd, *GSplashScreenAppName.ToString());
-			ShowWindow(GSplashScreenWnd, SW_SHOW); 
-			UpdateWindow(GSplashScreenWnd); 
+		SetWindowText(GSplashScreenWnd, *GSplashScreenAppName.ToString());
+		ShowWindow(GSplashScreenWnd, SW_SHOW); 
+		UpdateWindow(GSplashScreenWnd); 
 		 
-			const double FadeStartTime = FPlatformTime::Seconds();
-			const float FadeDuration = 0.2f;
-			int32 CurrentOpacityByte = 0;
+		const double FadeStartTime = FPlatformTime::Seconds();
+		const float FadeDuration = 0.2f;
+		int32 CurrentOpacityByte = 0;
 
-			MSG message;
-			bool bIsSplashFinished = false;
-			while( !bIsSplashFinished )
+		MSG message;
+		bool bIsSplashFinished = false;
+		while (!bIsSplashFinished)
+		{
+			if( PeekMessage(&message, NULL, 0, 0, PM_REMOVE) )
 			{
-				if( PeekMessage(&message, NULL, 0, 0, PM_REMOVE) )
-				{
-					TranslateMessage(&message);
-					DispatchMessage(&message);
+				TranslateMessage(&message);
+				DispatchMessage(&message);
 
-					if( message.message == WM_QUIT )
-					{
-						bIsSplashFinished = true;
-					}
-				}
-
-				// Update window opacity
-				CA_SUPPRESS(6239)
-				if( bAllowFading && CurrentOpacityByte < 255 )
+				if( message.message == WM_QUIT )
 				{
-					// Set window to fully transparent to start out
-					const float TimeSinceFadeStart = (float)( FPlatformTime::Seconds() - FadeStartTime );
-					const float FadeAmount = FMath::Clamp( TimeSinceFadeStart / FadeDuration, 0.0f, 1.0f );
-					const int32 NewOpacityByte = (int32)(255 * FadeAmount);
-					if( NewOpacityByte != CurrentOpacityByte )
-					{
-						CurrentOpacityByte = NewOpacityByte;
-						SetLayeredWindowAttributes( GSplashScreenWnd, 0, (BYTE)CurrentOpacityByte, LWA_ALPHA );
-					}
-
-					// We're still fading, but still yield a timeslice
-					FPlatformProcess::Sleep( 0.0f );
+					bIsSplashFinished = true;
 				}
-				else
+			}
+
+			// Update window opacity
+			CA_SUPPRESS(6239)
+			if( bAllowFading && CurrentOpacityByte < 255 )
+			{
+				// Set window to fully transparent to start out
+				const float TimeSinceFadeStart = (float)( FPlatformTime::Seconds() - FadeStartTime );
+				const float FadeAmount = FMath::Clamp( TimeSinceFadeStart / FadeDuration, 0.0f, 1.0f );
+				const int32 NewOpacityByte = (int32)(255 * FadeAmount);
+				if( NewOpacityByte != CurrentOpacityByte )
 				{
-					// Give up some time
-					FPlatformProcess::Sleep( 1.0f / 60.0f );
+					CurrentOpacityByte = NewOpacityByte;
+					SetLayeredWindowAttributes( GSplashScreenWnd, 0, (BYTE)CurrentOpacityByte, LWA_ALPHA );
 				}
+
+				// We're still fading, but still yield a timeslice
+				FPlatformProcess::Sleep( 0.0f );
+			}
+			else
+			{
+				// Give up some time
+				FPlatformProcess::Sleep( 1.0f / 60.0f );
 			}
 		}
 
@@ -587,7 +594,7 @@ uint32 WINAPI StartSplashScreenThread( LPVOID unused )
 	}
 
 	UnregisterClass(wc.lpszClassName, hInstance);
-	return 0; 
+	return 0;
 }
 
 /**
@@ -677,7 +684,8 @@ void FWindowsPlatformSplash::Show()
 			{
 				ShowWindow(GSplashScreenGuard, SW_SHOW); 
 			}
-			
+
+			GSplashWindowCreationEvent = CreateEvent(NULL, FALSE, FALSE, NULL);
 			GSplashScreenFileName = SplashPath;
 			DWORD ThreadID = 0;
 			GSplashScreenThread = CreateThread(NULL, 128 * 1024, (LPTHREAD_START_ROUTINE)StartSplashScreenThread, (LPVOID)NULL, STACK_SIZE_PARAM_IS_A_RESERVATION, &ThreadID);
@@ -693,10 +701,15 @@ void FWindowsPlatformSplash::Hide()
 	if(GSplashScreenThread)
 	{
 		TRACE_CPUPROFILER_EVENT_SCOPE(FWindowsPlatformSplash::Hide);
-		if(GSplashScreenWnd)
+
+		// The call to Hide() can be executed on the main thread before GSplashScreenWnd is created on the splash screen thread. Wait until the
+		// the thread had a chance to create the splash screen window and then notify the thread back with a WM_DESTROY message that will end
+		// the loop in that thread and prevent deadlocking the engine.
+		WaitForSingleObject(GSplashWindowCreationEvent, INFINITE);
+		if (GSplashScreenWnd)
 		{
-			// Send message to splash screen window to destroy itself
-			PostMessageW(GSplashScreenWnd, WM_DESTROY, 0, 0);
+			// Send message to splash screen window to close itself
+			PostMessageW(GSplashScreenWnd, WM_CLOSE, 0, 0);
 		}
 
 		// Wait for splash screen thread to finish
