Fixed CIS non-unity ShaderCompileWorker errors reported in WindowsPlatformSplash.cpp.

diff --git a/Engine/Source/Runtime/ApplicationCore/Private/Windows/WindowsPlatformSplash.cpp b/Engine/Source/Runtime/ApplicationCore/Private/Windows/WindowsPlatformSplash.cpp
--- a/Engine/Source/Runtime/ApplicationCore/Private/Windows/WindowsPlatformSplash.cpp
+++ b/Engine/Source/Runtime/ApplicationCore/Private/Windows/WindowsPlatformSplash.cpp
@@ -19,6 +19,7 @@
 #include "Stats/Stats.h"
 #include "Misc/App.h"
 #include "Misc/EngineVersion.h"
+#include "Misc/ScopeExit.h"
 #include "Windows/WindowsApplication.h"
 #include "Windows/WindowsHWrapper.h"
 #include "Windows/AllowWindowsPlatformTypes.h"
