Fixes incorrect screen launch rotation due to invalid mask value before IOSView is initialized.

#rnx

diff --git a/Engine/Source/Runtime/ApplicationCore/Private/IOS/IOSView.cpp b/Engine/Source/Runtime/ApplicationCore/Private/IOS/IOSView.cpp
--- a/Engine/Source/Runtime/ApplicationCore/Private/IOS/IOSView.cpp
+++ b/Engine/Source/Runtime/ApplicationCore/Private/IOS/IOSView.cpp
@@ -1050,9 +1050,9 @@ self.accessibilityElements = @[Window.accessibilityContainer];
 		return View->SupportedInterfaceOrientations;
 	}
 	else
-	{	
-		return 0;
-	}	
+	{
+		return UIInterfaceOrientationMaskAll;
+	}
 }
 
 /**
