Depreacte the old IForceFeedbackSystem from 2015 because it is just an empty subclass of IInputInterface. It has a comment for "This class is deprecated in favor of IInputInputface", but none of the references were ever replaced.

diff --git a/Engine/Plugins/Media/PixelStreaming/Source/PixelStreaming/Private/PixelStreamingInputHandler.h b/Engine/Plugins/Media/PixelStreaming/Source/PixelStreaming/Private/PixelStreamingInputHandler.h
--- a/Engine/Plugins/Media/PixelStreaming/Source/PixelStreaming/Private/PixelStreamingInputHandler.h
+++ b/Engine/Plugins/Media/PixelStreaming/Source/PixelStreaming/Private/PixelStreamingInputHandler.h
@@ -31,7 +31,7 @@ namespace UE::PixelStreaming
 		/** Exec handler to allow console commands to be passed through for debugging */
 		virtual bool Exec(UWorld* InWorld, const TCHAR* Cmd, FOutputDevice& Ar) override;
 		/**
-		 * IForceFeedbackSystem pass through functions
+		 * IInputInterface pass through functions
 		 */
 		virtual void SetChannelValue(int32 ControllerId, FForceFeedbackChannelType ChannelType, float Value) override;
 
diff --git a/Engine/Source/Runtime/ApplicationCore/Private/Apple/AppleControllerInterface.h b/Engine/Source/Runtime/ApplicationCore/Private/Apple/AppleControllerInterface.h
--- a/Engine/Source/Runtime/ApplicationCore/Private/Apple/AppleControllerInterface.h
+++ b/Engine/Source/Runtime/ApplicationCore/Private/Apple/AppleControllerInterface.h
@@ -3,8 +3,8 @@
 #pragma once
 
 #include "CoreMinimal.h"
-#include "GenericPlatform/IForceFeedbackSystem.h"
 #include "Misc/CoreMisc.h"
+#include "GenericPlatform/IInputInterface.h"
 #include "GenericPlatform/GenericApplicationMessageHandler.h"
 #import <GameController/GameController.h>
 
@@ -60,7 +60,7 @@ struct FDeferredAppleControllerEvent
 /**
  * Interface class for Apple Controllers
  */
-class FAppleControllerInterface : public IForceFeedbackSystem
+class FAppleControllerInterface : public IInputInterface
 {
 public:
 
@@ -82,7 +82,7 @@ public:
 	void SendControllerEvents();
 
 	/**
-	 * IForceFeedbackSystem implementation
+	 * Force Feedback implementation
 	 */
 	virtual void SetForceFeedbackChannelValue(int32 ControllerId, FForceFeedbackChannelType ChannelType, float Value) override {}
 	virtual void SetForceFeedbackChannelValues(int32 ControllerId, const FForceFeedbackValues &values) override {}
diff --git a/Engine/Source/Runtime/ApplicationCore/Public/Android/AndroidInputInterface.h b/Engine/Source/Runtime/ApplicationCore/Public/Android/AndroidInputInterface.h
--- a/Engine/Source/Runtime/ApplicationCore/Public/Android/AndroidInputInterface.h
+++ b/Engine/Source/Runtime/ApplicationCore/Public/Android/AndroidInputInterface.h
@@ -7,7 +7,6 @@
 #include <android/api-level.h>
 #include "GenericPlatform/ICursor.h"
 #include "GenericPlatform/IInputInterface.h"
-#include "GenericPlatform/IForceFeedbackSystem.h"
 #include "Math/Vector.h"
 #include "Math/Vector2D.h"
 #include "Math/Color.h"
@@ -215,7 +214,7 @@ struct FDeferredAndroidMessage
 /**
  * Interface class for Android input devices                 
  */
-class FAndroidInputInterface : public IForceFeedbackSystem
+class FAndroidInputInterface : public IInputInterface
 {
 public:
 
@@ -256,7 +255,7 @@ public:
 	static void QueueMotionData(const FVector& Tilt, const FVector& RotationRate, const FVector& Gravity, const FVector& Acceleration);
 
 	/**
-	* IForceFeedbackSystem implementation
+	* Force Feedback implementation
 	*/
 	virtual void SetForceFeedbackChannelValue(int32 ControllerId, FForceFeedbackChannelType ChannelType, float Value) override;
 	virtual void SetForceFeedbackChannelValues(int32 ControllerId, const FForceFeedbackValues &values) override;
diff --git a/Engine/Source/Runtime/ApplicationCore/Public/GenericPlatform/GenericApplication.h b/Engine/Source/Runtime/ApplicationCore/Public/GenericPlatform/GenericApplication.h
--- a/Engine/Source/Runtime/ApplicationCore/Public/GenericPlatform/GenericApplication.h
+++ b/Engine/Source/Runtime/ApplicationCore/Public/GenericPlatform/GenericApplication.h
@@ -22,7 +22,6 @@
 class FSlateApplication;
 class IAnalyticsProvider;
 class ICursor;
-class IForceFeedbackSystem;
 class IInputInterface;
 class ITextInputMethodSystem;
 
diff --git a/Engine/Source/Runtime/ApplicationCore/Public/GenericPlatform/IForceFeedbackSystem.h b/Engine/Source/Runtime/ApplicationCore/Public/GenericPlatform/IForceFeedbackSystem.h
--- a/Engine/Source/Runtime/ApplicationCore/Public/GenericPlatform/IForceFeedbackSystem.h
+++ b/Engine/Source/Runtime/ApplicationCore/Public/GenericPlatform/IForceFeedbackSystem.h
@@ -2,6 +2,17 @@
 
 #pragma once
 
+/**
+ * DEPRECATED GenericPlatform/IForceFeedbackSystem.h. Use GenericPlatform/IInputInterface.h instead
+ */
+#pragma once
+  
+#ifdef _MSC_VER
+	#pragma message(__FILE__"(9): warning: use GenericPlatform/IInputInterface.h instead of GenericPlatform/IForceFeedbackSystem.h")
+#else
+	#pragma message("#include GenericPlatform/IInputInterface.h instead of GenericPlatform/IForceFeedbackSystem.h")
+#endif
+
 #include "GenericPlatform/IInputInterface.h"
 
 
@@ -10,6 +21,7 @@
  *
  * Note: This class is deprecated and will be removed in favor of IInputInterface
  */
+UE_DEPRECATED(5.2, "IForceFeedbackSystem has been deprecated, use IInputInterface instead.")
 class IForceFeedbackSystem
 	: public IInputInterface
 {
diff --git a/Engine/Source/Runtime/ApplicationCore/Public/IOS/IOSInputInterface.h b/Engine/Source/Runtime/ApplicationCore/Public/IOS/IOSInputInterface.h
--- a/Engine/Source/Runtime/ApplicationCore/Public/IOS/IOSInputInterface.h
+++ b/Engine/Source/Runtime/ApplicationCore/Public/IOS/IOSInputInterface.h
@@ -4,7 +4,6 @@
 
 #include "CoreMinimal.h"
 #include "GenericPlatform/IInputInterface.h"
-#include "GenericPlatform/IForceFeedbackSystem.h"
 #if !PLATFORM_TVOS
 #import <CoreMotion/CoreMotion.h>
 #endif
@@ -81,7 +80,7 @@ public:
 	void SendControllerEvents();
 	
 	/**
-	 * IForceFeedbackSystem implementation
+	 * IInputInterface implementation
 	 */
 	virtual void SetForceFeedbackChannelValue(int32 ControllerId, FForceFeedbackChannelType ChannelType, float Value) override;
 	virtual void SetForceFeedbackChannelValues(int32 ControllerId, const FForceFeedbackValues &values) override;
diff --git a/Engine/Source/Runtime/ApplicationCore/Public/Windows/WindowsApplication.h b/Engine/Source/Runtime/ApplicationCore/Public/Windows/WindowsApplication.h
--- a/Engine/Source/Runtime/ApplicationCore/Public/Windows/WindowsApplication.h
+++ b/Engine/Source/Runtime/ApplicationCore/Public/Windows/WindowsApplication.h
@@ -7,19 +7,17 @@
 #include "Math/Color.h"
 #include "HAL/IConsoleManager.h"
 #include "GenericPlatform/GenericApplication.h"
-
+#include "GenericPlatform/IInputInterface.h"
 #include "Windows/WindowsHWrapper.h"
 #include "Windows/AllowWindowsPlatformTypes.h"
 	#include <Ole2.h>
 	#include <oleidl.h>
 	#include <ShObjIdl.h>
 #include "Windows/HideWindowsPlatformTypes.h"
-#include "GenericPlatform/IForceFeedbackSystem.h"
 #include "Windows/WindowsTextInputMethodSystem.h"
 
 class FGenericWindow;
 enum class EWindowTransparency;
-class IInputInterface;
 class ITextInputMethodSystem;
 enum class FForceFeedbackChannelType;
 struct FForceFeedbackValues;
@@ -284,7 +282,7 @@ public:
  */
 class APPLICATIONCORE_API FWindowsApplication
 	: public GenericApplication
-	, public IForceFeedbackSystem
+	, public IInputInterface
 {
 public:
 
diff --git a/Engine/Source/Runtime/InputDevice/Public/IInputDevice.h b/Engine/Source/Runtime/InputDevice/Public/IInputDevice.h
--- a/Engine/Source/Runtime/InputDevice/Public/IInputDevice.h
+++ b/Engine/Source/Runtime/InputDevice/Public/IInputDevice.h
@@ -31,7 +31,7 @@ public:
     virtual bool Exec( UWorld* InWorld, const TCHAR* Cmd, FOutputDevice& Ar ) = 0;
 
 	/**
-	 * IForceFeedbackSystem pass through functions
+	 * Force Feedback pass through functions
 	 */
 	virtual void SetChannelValue (int32 ControllerId, FForceFeedbackChannelType ChannelType, float Value) = 0;
 	virtual void SetChannelValues (int32 ControllerId, const FForceFeedbackValues &values) = 0;
