Adding a remap type on android 12 and up.

diff --git a/Engine/Source/Runtime/ApplicationCore/Private/Android/AndroidInputInterface.cpp b/Engine/Source/Runtime/ApplicationCore/Private/Android/AndroidInputInterface.cpp
--- a/Engine/Source/Runtime/ApplicationCore/Private/Android/AndroidInputInterface.cpp
+++ b/Engine/Source/Runtime/ApplicationCore/Private/Android/AndroidInputInterface.cpp
@@ -975,7 +975,7 @@ void FAndroidInputInterface::SendControllerEvents()
 						}
 						else if (CurrentDevice.DeviceInfo.Name.StartsWith(TEXT("PS5 Wireless Controller")))
 						{
-							CurrentDevice.ButtonRemapping = ButtonRemapType::PS5;
+							CurrentDevice.ButtonRemapping = FAndroidMisc::GetAndroidBuildVersion() > 11 ? ButtonRemapType::PS5New : ButtonRemapType::PS5;
 							CurrentDevice.ControllerClass = ControllerClassType::PlaystationWireless;
 							CurrentDevice.bSupportsHat = true;
 							CurrentDevice.bRightStickZRZ = true;
@@ -1618,6 +1618,28 @@ void FAndroidInputInterface::JoystickButtonEvent(int32 deviceId, int32 buttonId,
 				case AKEYCODE_BUTTON_R1:     NewControllerData[deviceId].ButtonStates[11] = buttonDown; break; // R2
 			}
 			break;
+		case ButtonRemapType::PS5New:
+			switch (buttonId)
+			{
+				case AKEYCODE_BUTTON_A:		NewControllerData[deviceId].ButtonStates[0] = buttonDown; break; // Cross
+				case AKEYCODE_BUTTON_B:		NewControllerData[deviceId].ButtonStates[1] = buttonDown; break; // Circle
+				case AKEYCODE_BUTTON_X:		NewControllerData[deviceId].ButtonStates[2] = buttonDown; break; // Triangle
+				case AKEYCODE_BUTTON_Y:		NewControllerData[deviceId].ButtonStates[3] = buttonDown; break; // Square
+				case AKEYCODE_BUTTON_L1:	NewControllerData[deviceId].ButtonStates[4] = buttonDown; break; // L1
+				case AKEYCODE_BUTTON_R1:	NewControllerData[deviceId].ButtonStates[5] = buttonDown; break; // R1
+				case AKEYCODE_BUTTON_THUMBL:NewControllerData[deviceId].ButtonStates[8] = buttonDown; break; // L3
+				case AKEYCODE_BUTTON_THUMBR:NewControllerData[deviceId].ButtonStates[9] = buttonDown; break; // R3
+				case AKEYCODE_BUTTON_L2:	NewControllerData[deviceId].ButtonStates[10] = buttonDown; break; // L2
+				case AKEYCODE_BUTTON_R2:	NewControllerData[deviceId].ButtonStates[11] = buttonDown; break; // R2
+				case 3002:		NewControllerData[deviceId].ButtonStates[16] = buttonDown; break; // Touchpad
+				case AKEYCODE_BUTTON_START:	NewControllerData[deviceId].ButtonStates[6] = buttonDown; // Options
+					if (!bBlockAndroidKeysOnControllers)
+					{
+						NewControllerData[deviceId].ButtonStates[17] = buttonDown; // Options
+					}
+					break;
+			}
+			break;
 	}
 }
 
diff --git a/Engine/Source/Runtime/ApplicationCore/Public/Android/AndroidInputInterface.h b/Engine/Source/Runtime/ApplicationCore/Public/Android/AndroidInputInterface.h
--- a/Engine/Source/Runtime/ApplicationCore/Public/Android/AndroidInputInterface.h
+++ b/Engine/Source/Runtime/ApplicationCore/Public/Android/AndroidInputInterface.h
@@ -101,7 +101,8 @@ enum ButtonRemapType
 	Normal,
 	XBox,
 	PS4,
-	PS5
+	PS5,
+	PS5New
 };
 
 struct FAndroidInputDeviceInfo {
