Fix iOS delegate method name typo for remote notification registration failure.  Github PR #8541

#rnx

diff --git a/Engine/Source/Runtime/ApplicationCore/Private/IOS/IOSAppDelegate.cpp b/Engine/Source/Runtime/ApplicationCore/Private/IOS/IOSAppDelegate.cpp
--- a/Engine/Source/Runtime/ApplicationCore/Private/IOS/IOSAppDelegate.cpp
+++ b/Engine/Source/Runtime/ApplicationCore/Private/IOS/IOSAppDelegate.cpp
@@ -1558,7 +1558,7 @@ extern double GCStartTime;
     }, TStatId(), NULL, ENamedThreads::GameThread);
 }
 
--(void)application:(UIApplication *)application didFailtoRegisterForRemoteNotificationsWithError:(NSError *)error
+-(void)application:(UIApplication *)application didFailToRegisterForRemoteNotificationsWithError:(NSError *)error
 {
 	FString errorDescription([error description]);
 	
