Update the platform input device mapper to expose the Allocate User and Device functions as public, allowing for platforms to allocate new input device ID's and platform user's if they want, making for a little bit easier to use API.

When setting the input device connection state, only call the delegate that there was a connection state change if the state is actually different than the existing one.

#rnx

diff --git a/Engine/Source/Runtime/ApplicationCore/Private/GenericPlatform/GenericPlatformInputDeviceMapper.cpp b/Engine/Source/Runtime/ApplicationCore/Private/GenericPlatform/GenericPlatformInputDeviceMapper.cpp
--- a/Engine/Source/Runtime/ApplicationCore/Private/GenericPlatform/GenericPlatformInputDeviceMapper.cpp
+++ b/Engine/Source/Runtime/ApplicationCore/Private/GenericPlatform/GenericPlatformInputDeviceMapper.cpp
@@ -132,6 +132,12 @@ bool IPlatformInputDeviceMapper::Internal_SetInputDeviceConnectionState(FInputDe
 		return false;
 	}
 
+	// If the connection state hasn't changed, then there is no point to calling the map function below
+	if (GetInputDeviceConnectionState(DeviceId) == NewState)
+	{
+		return false;
+	}
+
 	// Determine the owning user for this input device
 	FPlatformUserId OwningUser = GetUserForInputDevice(DeviceId);
 
