Create a generic "FInputDeviceLightColorProperty" that can be used on any platform. Update the Sony platforms to use it instead and deprecate the old one.

diff --git a/Engine/Source/Runtime/ApplicationCore/Public/GenericPlatform/IInputInterface.h b/Engine/Source/Runtime/ApplicationCore/Public/GenericPlatform/IInputInterface.h
--- a/Engine/Source/Runtime/ApplicationCore/Public/GenericPlatform/IInputInterface.h
+++ b/Engine/Source/Runtime/ApplicationCore/Public/GenericPlatform/IInputInterface.h
@@ -98,6 +98,21 @@ struct FInputDeviceProperty
 	FName Name;
 };
 
+/** A generic light color property for input devices that have lights on them */
+struct FInputDeviceLightColorProperty : public FInputDeviceProperty
+{
+	FInputDeviceLightColorProperty()
+		: FInputDeviceProperty(PropertyName())
+	{}
+
+	static FName PropertyName() { return FName("InputDeviceLightColor"); }
+
+	/** If the light should be enabled or not. */
+	bool bEnable = true;
+	
+	/** The color to set the light */
+	FColor Color = FColor::White;
+}; 
 
 /**
  * Interface for the input interface.
