FInputDeviceLightColorProperty: add constructor that takes bool and color to match the deprecated old class. Fix callsites using the deprecated old class.
#rnx

diff --git a/Engine/Source/Runtime/ApplicationCore/Public/GenericPlatform/IInputInterface.h b/Engine/Source/Runtime/ApplicationCore/Public/GenericPlatform/IInputInterface.h
--- a/Engine/Source/Runtime/ApplicationCore/Public/GenericPlatform/IInputInterface.h
+++ b/Engine/Source/Runtime/ApplicationCore/Public/GenericPlatform/IInputInterface.h
@@ -105,6 +105,12 @@ struct FInputDeviceLightColorProperty : public FInputDeviceProperty
 		: FInputDeviceProperty(PropertyName())
 	{}
 
+	FInputDeviceLightColorProperty(bool bInEnable, FColor InColor)
+		: FInputDeviceProperty(PropertyName())
+		, bEnable(bInEnable)
+		, Color(InColor)
+	{}
+
 	static FName PropertyName() { return FName("InputDeviceLightColor"); }
 
 	/** If the light should be enabled or not. */
