Add the low level Trigger Input Device properties to IInputInterface and use them in the PS5 controllers

diff --git a/Engine/Source/Runtime/ApplicationCore/Public/GenericPlatform/IInputInterface.h b/Engine/Source/Runtime/ApplicationCore/Public/GenericPlatform/IInputInterface.h
--- a/Engine/Source/Runtime/ApplicationCore/Public/GenericPlatform/IInputInterface.h
+++ b/Engine/Source/Runtime/ApplicationCore/Public/GenericPlatform/IInputInterface.h
@@ -134,6 +134,81 @@ struct FInputDeviceLightColorProperty : public FInputDeviceProperty
 	FColor Color = FColor::White;
 }; 
 
+/** Base class for device properties that affect Triggers */
+struct FInputDeviceTriggerProperty : public FInputDeviceProperty
+{
+	FInputDeviceTriggerProperty(FName InName)
+		: FInputDeviceProperty(InName)
+	{}
+
+	/** Which trigger this property should effect */
+	EInputDeviceTriggerMask AffectedTriggers = EInputDeviceTriggerMask::None;
+};
+
+/** This property can be used to reset the state of a given trigger. */
+struct FInputDeviceTriggerResetProperty : public FInputDeviceTriggerProperty
+{
+	FInputDeviceTriggerResetProperty()
+		: FInputDeviceTriggerProperty(PropertyName())
+	{
+		AffectedTriggers = EInputDeviceTriggerMask::All;
+	}
+
+	static FName PropertyName() { return FName("InputDeviceTriggerResetProperty"); }
+};
+
+/** Trigger resistance that is applied at a single position with the given strength. */
+struct FInputDeviceTriggerFeedbackProperty : public FInputDeviceTriggerProperty
+{
+	FInputDeviceTriggerFeedbackProperty()
+		: FInputDeviceTriggerProperty(PropertyName())
+	{}
+
+	static FName PropertyName() { return FName("InputDeviceTriggerFeedback"); }
+
+	int32 Position = 0;
+
+	int32 Strengh = 0;
+};
+
+/** A generic trigger effect that allows analog triggers to have a resistance curve between two points (Start and End) */
+struct FInputDeviceTriggerResistanceProperty : public FInputDeviceTriggerProperty
+{
+	FInputDeviceTriggerResistanceProperty()
+		: FInputDeviceTriggerProperty(PropertyName())
+	{}
+
+	static FName PropertyName() { return FName("InputDeviceTriggerResistance"); }
+
+	/** The position that the trigger should start providing resistance */
+	int32 StartPosition = 0;
+
+	/** How strong the resistance is */
+	int32 StartStrengh = 0;
+
+	/** The position that the trigger should start providing resistance */
+	int32 EndPosition = 0;
+
+	/** How strong the resistance is */
+	int32 EndStrengh = 0;
+};
+
+/** A generic input device property that sets vibration on triggers */
+struct FInputDeviceTriggerVibrationProperty : public FInputDeviceTriggerProperty
+{
+	FInputDeviceTriggerVibrationProperty()
+		: FInputDeviceTriggerProperty(PropertyName())
+	{}
+
+	static FName PropertyName() { return FName("InputDeviceTriggerVibration"); }
+
+	int32 TriggerPosition = 0;
+
+	int32 VibrationFrequency = 0;
+
+	int32 VibrationAmplitude = 0;
+};
+
 /**
  * Interface for the input interface.
  */
