Add missing dealloc instance method found by static analysis

diff --git a/Engine/Source/Runtime/ApplicationCore/Private/IOS/ProxyPaymentTransactionObserver.cpp b/Engine/Source/Runtime/ApplicationCore/Private/IOS/ProxyPaymentTransactionObserver.cpp
--- a/Engine/Source/Runtime/ApplicationCore/Private/IOS/ProxyPaymentTransactionObserver.cpp
+++ b/Engine/Source/Runtime/ApplicationCore/Private/IOS/ProxyPaymentTransactionObserver.cpp
@@ -21,6 +21,11 @@
 	return self;
 }
 
+-(void) dealloc
+{
+    [super dealloc];
+}
+
 -(void)paymentQueue:(SKPaymentQueue*)queue updatedTransactions : (NSArray*)transactions
 {
     if(self.ossObserver)
