Add missing release for a property

diff --git a/Engine/Source/Runtime/ApplicationCore/Private/IOS/ProxyPaymentTransactionObserver.cpp b/Engine/Source/Runtime/ApplicationCore/Private/IOS/ProxyPaymentTransactionObserver.cpp
--- a/Engine/Source/Runtime/ApplicationCore/Private/IOS/ProxyPaymentTransactionObserver.cpp
+++ b/Engine/Source/Runtime/ApplicationCore/Private/IOS/ProxyPaymentTransactionObserver.cpp
@@ -23,6 +23,10 @@
 
 -(void) dealloc
 {
+    if(_ossObserver)
+    {
+        [_ossObserver release];
+    }
     [super dealloc];
 }
 
