Add more info to the log when the createwindow function fails and remove unlocalized messagebox that was highly likely to fail to show a window anyway

#rnx

diff --git a/Engine/Source/Runtime/ApplicationCore/Private/Windows/WindowsWindow.cpp b/Engine/Source/Runtime/ApplicationCore/Private/Windows/WindowsWindow.cpp
--- a/Engine/Source/Runtime/ApplicationCore/Private/Windows/WindowsWindow.cpp
+++ b/Engine/Source/Runtime/ApplicationCore/Private/Windows/WindowsWindow.cpp
@@ -201,17 +201,22 @@ void FWindowsWindow::Initialize( FWindowsApplication* const Application, const T
 
 	if (HWnd == NULL)
 	{
-		FSlowHeartBeatScope SuspendHeartBeat;
-
-		// @todo Error message should be localized!
-		MessageBox(NULL, TEXT("Window Creation Failed!"), TEXT("Error!"), MB_ICONEXCLAMATION | MB_OK);
-
+		// Make sure the call to GetLastError happens immediately after CreateWindowEx to capture the proper error context for that API call.
 		const uint32 Error = GetLastError();
 
-		// Get the number of handles.  A large number of windows has been known to cause window creation to fail because windows only allows so many.
+		FSlowHeartBeatScope SuspendHeartBeat;
+
+		// Get the number of handles along with global and per process GDI and User Objects to give an idea of why the problem is happening.
+		// A large number of windows has been known to cause window creation to fail because windows only allows so many.
 		DWORD NumHandles = 0;
 		GetProcessHandleCount(GetCurrentProcess(), &NumHandles);
-		checkf(0, TEXT("Window Creation Failed (%d). %d"), Error, NumHandles);
+		checkf(0, TEXT("Window Creation Failed (Error Code %d). \n  Current Process (Handles: %d, GDI Objects: %d, User Objects: %d)\n  Global (GDI Objects: %d, User Objects: %d)"),
+			Error, NumHandles,
+			GetGuiResources(GetCurrentProcess(), GR_GDIOBJECTS),
+			GetGuiResources(GetCurrentProcess(), GR_USEROBJECTS),
+			GetGuiResources(GR_GLOBAL, GR_GDIOBJECTS),
+			GetGuiResources(GR_GLOBAL, GR_USEROBJECTS)
+			);
 
 		return;
 	}
