Fix a bug of android safezone not updated if orientation changes from landscape to reverse landscape.

diff --git a/Engine/Source/Runtime/ApplicationCore/Private/Android/AndroidApplication.cpp b/Engine/Source/Runtime/ApplicationCore/Private/Android/AndroidApplication.cpp
--- a/Engine/Source/Runtime/ApplicationCore/Private/Android/AndroidApplication.cpp
+++ b/Engine/Source/Runtime/ApplicationCore/Private/Android/AndroidApplication.cpp
@@ -106,7 +106,10 @@ void FAndroidApplication::HandleDeviceOrientation()
 	if (PreviousDeviceOrientation != DeviceOrientation)
 	{
 		FCoreDelegates::ApplicationReceivedScreenOrientationChangedNotificationDelegate.Broadcast((int32)DeviceOrientation);
+	}
 
+	if (FAndroidWindow::SafezoneUpdated())
+	{
 		//we also want to fire off the safe frame event
 		FCoreDelegates::OnSafeFrameChangedEvent.Broadcast();
 	}
diff --git a/Engine/Source/Runtime/ApplicationCore/Private/Android/AndroidWindow.cpp b/Engine/Source/Runtime/ApplicationCore/Private/Android/AndroidWindow.cpp
--- a/Engine/Source/Runtime/ApplicationCore/Private/Android/AndroidWindow.cpp
+++ b/Engine/Source/Runtime/ApplicationCore/Private/Android/AndroidWindow.cpp
@@ -89,6 +89,7 @@ void FAndroidWindow::SetOSWindowHandle(void* InWindow)
 
 //This function is declared in the Java-defined class, GameActivity.java: "public native void nativeSetObbInfo(String PackageName, int Version, int PatchVersion);"
 static bool GAndroidIsPortrait = false;
+static std::atomic<bool> GAndroidSafezoneRequiresUpdate = false;
 static int GAndroidDepthBufferPreference = 0;
 static FVector4 GAndroidPortraitSafezone = FVector4(-1.0f, -1.0f, -1.0f, -1.0f);
 static FVector4 GAndroidLandscapeSafezone = FVector4(-1.0f, -1.0f, -1.0f, -1.0f);
@@ -126,6 +127,8 @@ JNI_METHOD void Java_com_epicgames_unreal_GameActivity_nativeSetSafezoneInfo(JNI
 		GAndroidLandscapeSafezone.Z = right;
 		GAndroidLandscapeSafezone.W = bottom;
 	}
+
+	GAndroidSafezoneRequiresUpdate = true;
 }
 #endif
 
@@ -143,6 +146,12 @@ FVector4 FAndroidWindow::GetSafezone(bool bPortrait)
 	return bPortrait ? GAndroidPortraitSafezone : GAndroidLandscapeSafezone;
 }
 
+bool FAndroidWindow::SafezoneUpdated()
+{
+	bool bRequiresUpdate = true;
+	return GAndroidSafezoneRequiresUpdate.compare_exchange_weak(bRequiresUpdate, false);
+}
+
 int32 FAndroidWindow::GetDepthBufferPreference()
 {
 	return GAndroidDepthBufferPreference;
diff --git a/Engine/Source/Runtime/ApplicationCore/Public/Android/AndroidWindow.h b/Engine/Source/Runtime/ApplicationCore/Public/Android/AndroidWindow.h
--- a/Engine/Source/Runtime/ApplicationCore/Public/Android/AndroidWindow.h
+++ b/Engine/Source/Runtime/ApplicationCore/Public/Android/AndroidWindow.h
@@ -61,6 +61,7 @@ public:
 
 	static bool IsPortraitOrientation();
 	static FVector4 GetSafezone(bool bPortrait);
+	static bool SafezoneUpdated();
 
 	// called by the Android event thread to initially set the current window dimensions.
 	static void SetWindowDimensions_EventThread(ANativeWindow* DimensionWindow);
