Prefer using Windows functions to get user32.dll to follow best practices and avoid a LoadLibrary with an absolute path

#rnx

diff --git a/Engine/Source/Runtime/ApplicationCore/Private/Windows/WindowsPlatformApplicationMisc.cpp b/Engine/Source/Runtime/ApplicationCore/Private/Windows/WindowsPlatformApplicationMisc.cpp
--- a/Engine/Source/Runtime/ApplicationCore/Private/Windows/WindowsPlatformApplicationMisc.cpp
+++ b/Engine/Source/Runtime/ApplicationCore/Private/Windows/WindowsPlatformApplicationMisc.cpp
@@ -228,7 +228,7 @@ void FWindowsPlatformApplicationMisc::SetHighDPIMode()
 
 			FPlatformProcess::FreeDllHandle(ShCoreDll);
 		}
-		else if (void* User32Dll = FPlatformProcess::GetDllHandle(TEXT("user32.dll")))
+		else if (void* User32Dll = GetModuleHandle(L"user32.dll"))
 		{
 			typedef BOOL(WINAPI *SetProcessDpiAwareProc)(void);
 			SetProcessDpiAwareProc SetProcessDpiAware = (SetProcessDpiAwareProc)FPlatformProcess::GetDllExport(User32Dll, TEXT("SetProcessDPIAware"));
@@ -243,8 +243,6 @@ void FWindowsPlatformApplicationMisc::SetHighDPIMode()
 					UE_LOG(LogInit, Warning, TEXT("SetProcessDpiAware failed"));
 				}
 			}
-
-			FPlatformProcess::FreeDllHandle(User32Dll);
 		}
 	}
 }
diff --git a/Engine/Source/Runtime/ApplicationCore/Private/Windows/WindowsWindow.cpp b/Engine/Source/Runtime/ApplicationCore/Private/Windows/WindowsWindow.cpp
--- a/Engine/Source/Runtime/ApplicationCore/Private/Windows/WindowsWindow.cpp
+++ b/Engine/Source/Runtime/ApplicationCore/Private/Windows/WindowsWindow.cpp
@@ -429,7 +429,7 @@ HRGN FWindowsWindow::MakeWindowRegionObject(bool bIncludeBorderWhenMaximized) co
 
 void FWindowsWindow::DisableTouchFeedback()
 {
-	if (void* User32Dll = FPlatformProcess::GetDllHandle(TEXT("User32.dll")))
+	if (void* User32Dll = GetModuleHandle(L"user32.dll"))
 	{
 		typedef enum tagWINVER602FEEDBACK_TYPE
 		{
