Logging: Removed FFeedbackContext overrides that are now covered by the base

#rnx

diff --git a/Engine/Source/Runtime/ApplicationCore/Private/Android/AndroidFeedbackContext.cpp b/Engine/Source/Runtime/ApplicationCore/Private/Android/AndroidFeedbackContext.cpp
deleted file mode 100644
diff --git a/Engine/Source/Runtime/ApplicationCore/Private/IOS/IOSFeedbackContext.cpp b/Engine/Source/Runtime/ApplicationCore/Private/IOS/IOSFeedbackContext.cpp
deleted file mode 100644
diff --git a/Engine/Source/Runtime/ApplicationCore/Public/Android/AndroidFeedbackContext.h b/Engine/Source/Runtime/ApplicationCore/Public/Android/AndroidFeedbackContext.h
--- a/Engine/Source/Runtime/ApplicationCore/Public/Android/AndroidFeedbackContext.h
+++ b/Engine/Source/Runtime/ApplicationCore/Public/Android/AndroidFeedbackContext.h
@@ -1,20 +1,11 @@
 // Copyright Epic Games, Inc. All Rights Reserved.
 
 #pragma once
+
 #include "Misc/App.h"
 #include "Misc/FeedbackContext.h"
 #include "Misc/OutputDeviceRedirector.h"
 
-/**
- * Feedback context implementation for Android.
- */
 class FAndroidFeedbackContext : public FFeedbackContext
 {
-public:
-
-	/** Default constructor. */
-	FAndroidFeedbackContext();
-
-	virtual void Serialize( const TCHAR* V, ELogVerbosity::Type Verbosity, const class FName& Category ) override;
-	virtual bool YesNof(const FText& Question) override;
 };
diff --git a/Engine/Source/Runtime/ApplicationCore/Public/IOS/IOSFeedbackContext.h b/Engine/Source/Runtime/ApplicationCore/Public/IOS/IOSFeedbackContext.h
--- a/Engine/Source/Runtime/ApplicationCore/Public/IOS/IOSFeedbackContext.h
+++ b/Engine/Source/Runtime/ApplicationCore/Public/IOS/IOSFeedbackContext.h
@@ -1,20 +1,11 @@
 // Copyright Epic Games, Inc. All Rights Reserved.
 
 #pragma once
+
 #include "Misc/App.h"
 #include "Misc/FeedbackContext.h"
 #include "Misc/OutputDeviceRedirector.h"
 
-/**
- * Feedback context implementation for IOS.
- */
 class FIOSFeedbackContext : public FFeedbackContext
 {
-public:
-
-	/** Default constructor. */
-	FIOSFeedbackContext();
-
-	virtual void Serialize( const TCHAR* V, ELogVerbosity::Type Verbosity, const class FName& Category ) override;
-	virtual bool YesNof(const FText& Question) override;
 };
