PR #9813: Fix UMG Overlay mouse/touch input not working after launching fullscreen Wayland application (Contributed by [REDACTED_PERSON])

diff --git a/Engine/Source/Runtime/ApplicationCore/Private/Linux/LinuxApplication.cpp b/Engine/Source/Runtime/ApplicationCore/Private/Linux/LinuxApplication.cpp
--- a/Engine/Source/Runtime/ApplicationCore/Private/Linux/LinuxApplication.cpp
+++ b/Engine/Source/Runtime/ApplicationCore/Private/Linux/LinuxApplication.cpp
@@ -793,8 +793,10 @@ void FLinuxApplication::ProcessDeferredMessage( SDL_Event Event )
 
 				case SDL_WINDOWEVENT_MOVED:
 					{
-						int32 ClientScreenX = windowEvent.data1;
-						int32 ClientScreenY = windowEvent.data2;
+						// Mask away the higher bits, as SDL uses those as flags
+						// See: SDL_WINDOWPOS_UNDEFINED_MASK & SDL_WINDOWPOS_CENTERED_MASK for context
+						int32 ClientScreenX = windowEvent.data1 & 0xFFFF;
+						int32 ClientScreenY = windowEvent.data2 & 0xFFFF;
 
 						int32 BorderSizeX, BorderSizeY;
 						CurrentEventWindow->GetNativeBordersSize(BorderSizeX, BorderSizeY);
