Fix some cases of the notification windows moving behind the main window

Still some cases to fix

diff --git a/Engine/Source/Runtime/ApplicationCore/Private/Linux/LinuxApplication.cpp b/Engine/Source/Runtime/ApplicationCore/Private/Linux/LinuxApplication.cpp
--- a/Engine/Source/Runtime/ApplicationCore/Private/Linux/LinuxApplication.cpp
+++ b/Engine/Source/Runtime/ApplicationCore/Private/Linux/LinuxApplication.cpp
@@ -442,11 +442,6 @@ void FLinuxApplication::ProcessDeferredMessage( SDL_Event Event )
 					if (CurrentlyActiveWindow != CurrentEventWindow)
 					{
 						ActivateWindow(CurrentEventWindow);
-
-						if(NotificationWindows.Num() > 0)
-						{
-							RaiseNotificationWindows(CurrentEventWindow);
-						}
 					}
 
 					// Check if we have to set the focus.
@@ -456,6 +451,11 @@ void FLinuxApplication::ProcessDeferredMessage( SDL_Event Event )
 					}
 
 					MessageHandler->OnMouseDown(CurrentEventWindow, button);
+
+					if (NotificationWindows.Num() > 0)
+					{
+						RaiseNotificationWindows(CurrentEventWindow);
+					}
 				}
 			}
 		}
@@ -906,6 +906,11 @@ void FLinuxApplication::ProcessDeferredMessage( SDL_Event Event )
 
 						CurrentFocusWindow = CurrentEventWindow;
 
+						if (NotificationWindows.Num() > 0)
+						{
+							RaiseNotificationWindows(CurrentEventWindow);
+						}
+
 						// We have gained focus, we can stop trying to check if we need to deactivate
 						FocusOutDeactivationTime = 0.0;
 					}
