Background audio on iOS

diff --git a/Engine/Source/Programs/UnrealBuildTool/Platform/IOS/UEDeployIOS.cs b/Engine/Source/Programs/UnrealBuildTool/Platform/IOS/UEDeployIOS.cs
--- a/Engine/Source/Programs/UnrealBuildTool/Platform/IOS/UEDeployIOS.cs
+++ b/Engine/Source/Programs/UnrealBuildTool/Platform/IOS/UEDeployIOS.cs
@@ -347,6 +347,10 @@ namespace UnrealBuildTool
 			bool bRemoteNotificationsSupported = false;
 			Ini.GetBool("/Script/IOSRuntimeSettings.IOSRuntimeSettings", "bEnableRemoteNotificationsSupport", out bRemoteNotificationsSupported);
 
+			// Add audio as background mode
+			bool bBackgroundAudioSupported = false;
+			Ini.GetBool("/Script/IOSRuntimeSettings.IOSRuntimeSettings", "bSupportsBackgroundAudio", out bBackgroundAudioSupported);
+
 			// Add background fetch as background mode
 			bool bBackgroundFetch = false;
 			Ini.GetBool("/Script/IOSRuntimeSettings.IOSRuntimeSettings", "bEnableBackgroundFetch", out bBackgroundFetch);
@@ -539,10 +543,14 @@ namespace UnrealBuildTool
 			}
 
 			// Add remote-notifications as background mode
-			if (bRemoteNotificationsSupported || bBackgroundFetch)
+			if (bRemoteNotificationsSupported || bBackgroundFetch || bBackgroundAudioSupported)
 			{
 				Text.AppendLine("\t<key>UIBackgroundModes</key>");
 				Text.AppendLine("\t<array>");
+				if (bBackgroundAudioSupported)
+				{
+					Text.AppendLine("\t\t<string>audio</string>");
+				}
 				if (bRemoteNotificationsSupported)
 				{
 					Text.AppendLine("\t\t<string>remote-notification</string>");
diff --git a/Engine/Source/Runtime/Apple/AudioMixerAudioUnit/Private/AudioMixerPlatformAudioUnit.cpp b/Engine/Source/Runtime/Apple/AudioMixerAudioUnit/Private/AudioMixerPlatformAudioUnit.cpp
--- a/Engine/Source/Runtime/Apple/AudioMixerAudioUnit/Private/AudioMixerPlatformAudioUnit.cpp
+++ b/Engine/Source/Runtime/Apple/AudioMixerAudioUnit/Private/AudioMixerPlatformAudioUnit.cpp
@@ -82,6 +82,8 @@ namespace Audio
 			return false;
 		}
 		
+        bSupportsBackgroundAudio = GConfig->GetBool(TEXT("/Script/IOSRuntimeSettings.IOSRuntimeSettings"), TEXT("bSupportsBackgroundAudio"), bSupportsBackgroundAudio, GEngineIni);
+
 		OSStatus Status;
 		GraphSampleRate = (double) InternalPlatformSettings.SampleRate;
 		UInt32 BufferSize = (UInt32) GetNumFrames(InternalPlatformSettings.CallbackBufferFrameSize);
@@ -513,6 +515,12 @@ namespace Audio
 	
 	void FMixerPlatformAudioUnit::SuspendContext()
 	{
+#if PLATFORM_IOS
+        if (bSupportsBackgroundAudio)
+        {
+            return;
+        }
+#endif
 		if (SuspendCounter == 0)
 		{
 			FPlatformAtomics::InterlockedIncrement(&SuspendCounter);
diff --git a/Engine/Source/Runtime/Apple/AudioMixerAudioUnit/Private/AudioMixerPlatformAudioUnit.h b/Engine/Source/Runtime/Apple/AudioMixerAudioUnit/Private/AudioMixerPlatformAudioUnit.h
--- a/Engine/Source/Runtime/Apple/AudioMixerAudioUnit/Private/AudioMixerPlatformAudioUnit.h
+++ b/Engine/Source/Runtime/Apple/AudioMixerAudioUnit/Private/AudioMixerPlatformAudioUnit.h
@@ -72,6 +72,8 @@ namespace Audio
 		int32       BytesPerSubmittedBuffer;
 		
 		double GraphSampleRate;
+        
+        bool bSupportsBackgroundAudio;
 		
 		// We may have to grow the circular buffer capacity since Audio Unit callback size is not guaranteed to be constant
 		// Currently, this just zero's-out and reallocates, so it will pop. (We always keep largest capacity)
diff --git a/Engine/Source/Runtime/ApplicationCore/Private/IOS/IOSAppDelegate.cpp b/Engine/Source/Runtime/ApplicationCore/Private/IOS/IOSAppDelegate.cpp
--- a/Engine/Source/Runtime/ApplicationCore/Private/IOS/IOSAppDelegate.cpp
+++ b/Engine/Source/Runtime/ApplicationCore/Private/IOS/IOSAppDelegate.cpp
@@ -606,8 +606,20 @@ static IOSAppDelegate* CachedDelegate = nil;
 	
 	self.bAudioActive = bActive;
 	
-	// get the category and settings to use
-	NSString* Category = [self IsFeatureActive:EAudioFeature::DoNotMixWithOthers] ? AVAudioSessionCategorySoloAmbient : AVAudioSessionCategoryAmbient;
+    // get the category and settings to use
+        NSString* Category = AVAudioSessionCategoryAmbient;
+        if([self IsFeatureActive:EAudioFeature::DoNotMixWithOthers])
+        {
+            Category = AVAudioSessionCategorySoloAmbient;
+        }
+    #if !PLATFORM_TVOS
+        bool bSupportsBackgroundAudio = GConfig->GetBool(TEXT("/Script/IOSRuntimeSettings.IOSRuntimeSettings"), TEXT("bSupportsBackgroundAudio"), bSupportsBackgroundAudio, GEngineIni);
+        if (bSupportsBackgroundAudio)
+        {
+            Category = AVAudioSessionCategoryPlayback;
+        }
+    #endif
+    
 	NSString* Mode = AVAudioSessionModeDefault;
 	AVAudioSessionCategoryOptions Options = 0;
 	if (self.bAudioActive || [self IsBackgroundAudioPlaying] || [self IsFeatureActive:EAudioFeature::BackgroundAudio])
@@ -1327,11 +1339,8 @@ FCriticalSection RenderSuspend;
  		FEmbeddedCommunication::KeepAwake(TEXT("Background"), false);
         FGraphEventRef ResignTask = FFunctionGraphTask::CreateAndDispatchWhenReady([]()
         {
-			UE_LOG(LogTemp, Display, TEXT("Calling Delegate"));
-
-			FCoreDelegates::ApplicationWillDeactivateDelegate.Broadcast();
-
-			FEmbeddedCommunication::AllowSleep(TEXT("Background"));
+            FCoreDelegates::ApplicationWillDeactivateDelegate.Broadcast();
+            FEmbeddedCommunication::AllowSleep(TEXT("Background"));
         }, TStatId(), NULL, ENamedThreads::GameThread);
 		
 		// Do not wait forever for this task to complete since the game thread may be stuck on waiting for user input from a modal dialog box
@@ -1347,11 +1356,18 @@ FCriticalSection RenderSuspend;
 		}
 		UE_LOG(LogTemp, Display, TEXT("Done with entering background tasks time."));
     }
+    bool bSupportsBackgroundAudio = GConfig->GetBool(TEXT("/Script/IOSRuntimeSettings.IOSRuntimeSettings"), TEXT("bSupportsBackgroundAudio"), bSupportsBackgroundAudio, GEngineIni);
+    
 // fix for freeze on tvOS, moving to applicationDidEnterBackground. Not making the changes for iOS platforms as the bug does not happen and could bring some side effets.
 #if !PLATFORM_TVOS
-    [self ToggleSuspend:true];
+    if (!bSupportsBackgroundAudio)
+    {
+        [self ToggleSuspend:true];
+        [self ToggleAudioSession:false];
+    }
+#else
+    [self ToggleAudioSession:false];
 #endif
-	[self ToggleAudioSession:false];
     
     RenderSuspend.TryLock();
     if (FTaskGraphInterface::IsRunning())
diff --git a/Engine/Source/Runtime/IOS/IOSRuntimeSettings/Classes/IOSRuntimeSettings.h b/Engine/Source/Runtime/IOS/IOSRuntimeSettings/Classes/IOSRuntimeSettings.h
--- a/Engine/Source/Runtime/IOS/IOSRuntimeSettings/Classes/IOSRuntimeSettings.h
+++ b/Engine/Source/Runtime/IOS/IOSRuntimeSettings/Classes/IOSRuntimeSettings.h
@@ -503,6 +503,10 @@ public:
 	UPROPERTY(config, EditAnywhere, Category = "Audio")
 	FPlatformRuntimeAudioCompressionOverrides CompressionOverrides;
 
+    /** Whether this app's audio can be played when using other apps or on the srpingboard */
+    UPROPERTY(config, EditAnywhere, Category = "Audio", meta = (DisplayName = "Whether audio from this plays in background (iOS and iPadOS only)"))
+    bool bSupportsBackgroundAudio;
+
 	/** This determines the max amount of memory that should be used for the cache at any given time. If set low (<= 8 MB), it lowers the size of individual chunks of audio during cook. */
 	UPROPERTY(GlobalConfig, EditAnywhere, Category = "Audio|CookOverrides|Stream Caching", meta = (DisplayName = "Max Cache Size (KB)"))
 	int32 CacheSizeKB;
diff --git a/Engine/Source/Runtime/IOS/IOSRuntimeSettings/Private/IOSRuntimeSettings.cpp b/Engine/Source/Runtime/IOS/IOSRuntimeSettings/Private/IOSRuntimeSettings.cpp
--- a/Engine/Source/Runtime/IOS/IOSRuntimeSettings/Private/IOSRuntimeSettings.cpp
+++ b/Engine/Source/Runtime/IOS/IOSRuntimeSettings/Private/IOSRuntimeSettings.cpp
@@ -50,6 +50,7 @@ UIOSRuntimeSettings::UIOSRuntimeSettings(const FObjectInitializer& ObjectInitial
 	bSupportsMetal = true;
 	bSupportsMetalMRT = false;
 	bDisableHTTPS = false;
+    bSupportsBackgroundAudio = false;
 }
 
 void UIOSRuntimeSettings::PostReloadConfig(class FProperty* PropertyThatWasLoaded)
