- When in full screen, the Mac now returns the "visibleFrame" property instead of the full "frame" to account for the bezel area on laptops with a camera notch

diff --git a/Engine/Source/Runtime/ApplicationCore/Private/Mac/MacApplication.cpp b/Engine/Source/Runtime/ApplicationCore/Private/Mac/MacApplication.cpp
--- a/Engine/Source/Runtime/ApplicationCore/Private/Mac/MacApplication.cpp
+++ b/Engine/Source/Runtime/ApplicationCore/Private/Mac/MacApplication.cpp
@@ -1224,8 +1224,8 @@ void FMacApplication::OnWindowDidResize(TSharedRef<FMacWindow> Window, bool bRes
 	if (Window->GetWindowHandle().TargetWindowMode == EWindowMode::WindowedFullscreen)
 	{
 		// Grab current monitor data for sizing
-		Width = FMath::TruncToInt([[Window->GetWindowHandle() screen] frame].size.width * Window->GetDPIScaleFactor());
-		Height = FMath::TruncToInt([[Window->GetWindowHandle() screen] frame].size.height * Window->GetDPIScaleFactor());
+		Width = FMath::TruncToInt([[Window->GetWindowHandle() screen] visibleFrame].size.width * Window->GetDPIScaleFactor());
+		Height = FMath::TruncToInt([[Window->GetWindowHandle() screen] visibleFrame].size.height * Window->GetDPIScaleFactor());
 	}
 
 	if (bRestoreMouseCursorLocking)
