Fix mouse grab cursors in Linux

diff --git a/Engine/Source/Runtime/ApplicationCore/Private/Linux/LinuxCursor.cpp b/Engine/Source/Runtime/ApplicationCore/Private/Linux/LinuxCursor.cpp
--- a/Engine/Source/Runtime/ApplicationCore/Private/Linux/LinuxCursor.cpp
+++ b/Engine/Source/Runtime/ApplicationCore/Private/Linux/LinuxCursor.cpp
@@ -82,21 +82,21 @@ FLinuxCursor::FLinuxCursor()
 			break;
 
 		case EMouseCursor::GrabHand:
-			//CursorHandle = LoadCursorFromFile((LPCTSTR)*(FString( FPlatformProcess::BaseDir() ) / FString::Printf( TEXT("%sEditor/Slate/Old/grabhand.cur"), *FPaths::EngineContentDir() )));
-			//if (CursorHandle == NULL)
-			//{
-			//	// Failed to load file, fall back
-				CursorHandle = SDL_CreateSystemCursor( SDL_SYSTEM_CURSOR_HAND );
-			//}
+			{
+				FString path(FPlatformProcess::BaseDir() / FPaths::EngineContentDir() / FString("Editor/Slate/Cursor/grabhand.bmp") );
+				SDL_Surface *surf = SDL_LoadBMP(TCHAR_TO_ANSI(*path));
+				CursorHandle = SDL_CreateColorCursor(surf, 15, 15);
+				SDL_FreeSurface(surf);
+			}
 			break;
 
 		case EMouseCursor::GrabHandClosed:
-			//CursorHandle = LoadCursorFromFile((LPCTSTR)*(FString( FPlatformProcess::BaseDir() ) / FString::Printf( TEXT("%sEditor/Slate/Old/grabhand_closed.cur"), *FPaths::EngineContentDir() )));
-			//if (CursorHandle == NULL)
-			//{
-			//	// Failed to load file, fall back
-				CursorHandle = SDL_CreateSystemCursor( SDL_SYSTEM_CURSOR_HAND );
-			//}
+			{
+				FString path(FPlatformProcess::BaseDir() / FPaths::EngineContentDir() / FString("Editor/Slate/Cursor/grabhand_closed.bmp") );
+				SDL_Surface *surf = SDL_LoadBMP(TCHAR_TO_ANSI(*path));
+				CursorHandle = SDL_CreateColorCursor(surf, 8, 8);
+				SDL_FreeSurface(surf);
+			}
 			break;
 
 		case EMouseCursor::SlashedCircle:
@@ -154,12 +154,12 @@ FLinuxCursor::~FLinuxCursor()
 			case EMouseCursor::CardinalCross:
 			case EMouseCursor::Crosshairs:
 			case EMouseCursor::Hand:
-			case EMouseCursor::GrabHand:
-			case EMouseCursor::GrabHandClosed:
 			case EMouseCursor::SlashedCircle:
 			case EMouseCursor::Custom:
 				// Standard shared cursors don't need to be destroyed
 				break;
+			case EMouseCursor::GrabHand:
+			case EMouseCursor::GrabHandClosed:
 			case EMouseCursor::EyeDropper:
 				SDL_FreeCursor(CursorHandles[CurCursorIndex]);
 				break;
