PR #6649: Fixed Linux touch id handling (Contributed by [REDACTED_PERSON])

diff --git a/Engine/Source/Runtime/ApplicationCore/Private/Linux/LinuxApplication.cpp b/Engine/Source/Runtime/ApplicationCore/Private/Linux/LinuxApplication.cpp
--- a/Engine/Source/Runtime/ApplicationCore/Private/Linux/LinuxApplication.cpp
+++ b/Engine/Source/Runtime/ApplicationCore/Private/Linux/LinuxApplication.cpp
@@ -1000,19 +1000,23 @@ void FLinuxApplication::ProcessDeferredMessage( SDL_Event Event )
 			{
 				// remove touch context even if it existed
 				uint64 FingerId = static_cast<uint64>(Event.tfinger.fingerId);
+
+				FTouchContext* TouchContext = Touches.Find(FingerId);
 				if (UNLIKELY(Touches.Find(FingerId) != nullptr))
 				{
+					TouchIds[TouchContext->TouchIndex] = TOptional<uint64>();
 					Touches.Remove(FingerId);
 					UE_LOG(LogLinuxWindow, Warning, TEXT("Received another SDL_FINGERDOWN for finger %llu which was already down."), FingerId);
 				}
 
 				FTouchContext NewTouch;
-				NewTouch.TouchIndex = Touches.Num();
+				NewTouch.TouchIndex = GetFirstFreeTouchId();
 				NewTouch.Location = GetTouchEventLocation(NativeWindow, Event);
 				NewTouch.DeviceId = Event.tfinger.touchId;
 				Touches.Add(FingerId, NewTouch);
+				TouchIds[NewTouch.TouchIndex] = TOptional<uint64>(FingerId);
 
-				UE_LOG(LogLinuxWindow, Verbose, TEXT("OnTouchStarted at (%f, %f), finger %d (system touch id %llu)"), NewTouch.Location.X, NewTouch.Location.Y, NewTouch.TouchIndex, FingerId);
+				UE_LOG(LogLinuxWindow, Verbose, TEXT("OnTouchStarted at (%f, %f), finger %d (system touch id %llu)"), NewTouch.Location.X, NewTouch.Location.Y, FingerId, NewTouch.TouchIndex);
 				MessageHandler->OnTouchStarted(CurrentEventWindow, NewTouch.Location, 1.0f, NewTouch.TouchIndex, 0);// NewTouch.DeviceId);
 			}
 			else
@@ -1041,10 +1045,11 @@ void FLinuxApplication::ProcessDeferredMessage( SDL_Event Event )
 					TouchContext->Location = GetTouchEventLocation(NativeWindow, Event);
 					// check touch device?
 
-					UE_LOG(LogLinuxWindow, Verbose, TEXT("OnTouchEnded at (%f, %f), finger %d (system touch id %llu)"), TouchContext->Location.X, TouchContext->Location.Y, TouchContext->TouchIndex, FingerId);
+					UE_LOG(LogLinuxWindow, Verbose, TEXT("OnTouchEnded at (%f, %f), finger %d (system touch id %llu)"), TouchContext->Location.X, TouchContext->Location.Y, FingerId, TouchContext->TouchIndex);
 					MessageHandler->OnTouchEnded(TouchContext->Location, TouchContext->TouchIndex, 0);// TouchContext->DeviceId);
 
 					// remove the touch
+					TouchIds[TouchContext->TouchIndex] = TOptional<uint64>();
 					Touches.Remove(FingerId);
 				}
 			}
@@ -1118,6 +1123,19 @@ void FLinuxApplication::CheckIfApplicatioNeedsDeactivation()
 	}
 }
 
+int FLinuxApplication::GetFirstFreeTouchId()
+{
+	for (int i = 0; i < TouchIds.Num(); i++)
+	{
+		if (TouchIds[i].IsSet() == false)
+		{
+			return i;
+		}
+	}
+
+	return TouchIds.Add(TOptional<uint64>()); 
+}
+
 FVector2D FLinuxApplication::GetTouchEventLocation(SDL_HWindow NativeWindow, SDL_Event TouchEvent)
 {
 	checkf(TouchEvent.type == SDL_FINGERDOWN || TouchEvent.type == SDL_FINGERUP || TouchEvent.type == SDL_FINGERMOTION, TEXT("Wrong touch event."));
diff --git a/Engine/Source/Runtime/ApplicationCore/Public/Linux/LinuxApplication.h b/Engine/Source/Runtime/ApplicationCore/Public/Linux/LinuxApplication.h
--- a/Engine/Source/Runtime/ApplicationCore/Public/Linux/LinuxApplication.h
+++ b/Engine/Source/Runtime/ApplicationCore/Public/Linux/LinuxApplication.h
@@ -215,6 +215,9 @@ private:
 	/** Gets the location from a given touch event. */
 	FVector2D GetTouchEventLocation(SDL_HWindow NativeWindow, SDL_Event TouchEvent);
 
+	/** Searches for a free touch index. */
+	int GetFirstFreeTouchId();
+
 public:
 	virtual IInputInterface* GetInputInterface() override
 	{
@@ -247,6 +250,9 @@ private:
 	/** Holds currently active touches (i.e. fingers pressed but not released) */
 	TMap<uint64, FTouchContext> Touches;
 
+	/** Maps touch indexes to SDL touch IDs. */
+	TArray<TOptional<uint64>> TouchIds;
+
 	struct SDLControllerState
 	{
 		/** SDL controller */
