AndroidInputInterface - PS5 controller now has the correct mapping for Android Version less than 12 and triggers are working correctly on Android 12.

diff --git a/Engine/Source/Runtime/ApplicationCore/Private/Android/AndroidInputInterface.cpp b/Engine/Source/Runtime/ApplicationCore/Private/Android/AndroidInputInterface.cpp
--- a/Engine/Source/Runtime/ApplicationCore/Private/Android/AndroidInputInterface.cpp
+++ b/Engine/Source/Runtime/ApplicationCore/Private/Android/AndroidInputInterface.cpp
@@ -975,13 +975,15 @@ void FAndroidInputInterface::SendControllerEvents()
 						}
 						else if (CurrentDevice.DeviceInfo.Name.StartsWith(TEXT("PS5 Wireless Controller")))
 						{
-							CurrentDevice.ButtonRemapping = FAndroidMisc::GetAndroidBuildVersion() > 11 ? ButtonRemapType::PS5New : ButtonRemapType::PS5;
+							//FAndroidMisc::GetAndroidBuildVersion() actually returns the API Level instead of the Android Version
+							bool bUseNewPS5Mapping = FAndroidMisc::GetAndroidBuildVersion() > 30;
+							CurrentDevice.ButtonRemapping = bUseNewPS5Mapping ? ButtonRemapType::PS5New : ButtonRemapType::PS5;
 							CurrentDevice.ControllerClass = ControllerClassType::PlaystationWireless;
 							CurrentDevice.bSupportsHat = true;
 							CurrentDevice.bRightStickZRZ = true;
-							CurrentDevice.bMapRXRYToTriggers = true;
-							CurrentDevice.LTAnalogRangeMinimum = -1.0f;
-							CurrentDevice.RTAnalogRangeMinimum = -1.0f;
+							CurrentDevice.bMapRXRYToTriggers = !bUseNewPS5Mapping;
+							CurrentDevice.LTAnalogRangeMinimum = bUseNewPS5Mapping ? 0.0f : -1.0f;
+							CurrentDevice.RTAnalogRangeMinimum = bUseNewPS5Mapping ? 0.0f : -1.0f;
 						}
 						else if (CurrentDevice.DeviceInfo.Name.StartsWith(TEXT("glap QXPGP001")))
 						{
