Update Android controllers

#android

diff --git a/Engine/Build/Android/Java/src/com/epicgames/unreal/GameActivity.java.template b/Engine/Build/Android/Java/src/com/epicgames/unreal/GameActivity.java.template
--- a/Engine/Build/Android/Java/src/com/epicgames/unreal/GameActivity.java.template
+++ b/Engine/Build/Android/Java/src/com/epicgames/unreal/GameActivity.java.template
@@ -5209,14 +5209,19 @@ public class GameActivity extends $${gameActivitySuperClass}$$ implements Surfac
 		new DeviceInfoData(0x0738, 0x5266, "Mad Catz C.T.R.L.R"),
 		new DeviceInfoData(0x045e, 0x028e, "Xbox Wired Controller"),
 		new DeviceInfoData(0x045e, 0x02e0, "Xbox Wireless Controller"),
+		new DeviceInfoData(0x045e, 0x02fd, "Xbox Wireless Controller"),
 		new DeviceInfoData(0x045e, 0x0b05, "Xbox Elite Wireless Controller"),
+		new DeviceInfoData(0x045e, 0x0b12, "Generic X-Box pad"),
+		new DeviceInfoData(0x045e, 0x0b13, "Xbox Wireless Controller"),
 		new DeviceInfoData(0x0111, 0x1419, "SteelSeries Stratus XL"), 
 		new DeviceInfoData(0x054c, 0x05c4, "PS4 Wireless Controller"),
 		new DeviceInfoData(0x054c, 0x09cc, "PS4 Wireless Controller (v2)"),
 		new DeviceInfoData(0x054c, 0x0ce6, "PS5 Wireless Controller"),
 		new DeviceInfoData(0x05ac, 0x056a, "glap QXPGP001"),
 		new DeviceInfoData(0x0483, 0x5750, "STMicroelectronics Lenovo GamePad"),
-		new DeviceInfoData(0x1532, 0x0705, "Razer Razer Claire T1 Wired")
+		new DeviceInfoData(0x1532, 0x0705, "Razer Razer Claire T1 Wired"),
+		new DeviceInfoData(0x1532, 0x0712, "Razer Kishi V2"),
+		new DeviceInfoData(0x1532, 0x0718, "Razer Kishi V2 Pro XBox360")
 	};
 
 	public class InputDeviceInfo {
diff --git a/Engine/Source/Runtime/ApplicationCore/Private/Android/AndroidInputInterface.cpp b/Engine/Source/Runtime/ApplicationCore/Private/Android/AndroidInputInterface.cpp
--- a/Engine/Source/Runtime/ApplicationCore/Private/Android/AndroidInputInterface.cpp
+++ b/Engine/Source/Runtime/ApplicationCore/Private/Android/AndroidInputInterface.cpp
@@ -930,6 +930,17 @@ void FAndroidInputInterface::SendControllerEvents()
 						{
 							CurrentDevice.bSupportsHat = true;
 						}
+						else if (CurrentDevice.DeviceInfo.Name.StartsWith(TEXT("Generic X-Box pad")))
+						{
+							CurrentDevice.ControllerClass = ControllerClassType::XBoxWired;
+							CurrentDevice.bSupportsHat = true;
+							CurrentDevice.bTriggersUseThresholdForClick = true;
+							CurrentDevice.bRightStickZRZ = false;
+							CurrentDevice.bRightStickRXRY = true;
+							CurrentDevice.bMapZRZToTriggers = true;
+							CurrentDevice.LTAnalogRangeMinimum = -1.0f;
+							CurrentDevice.RTAnalogRangeMinimum = -1.0f;
+						}
 						else if (CurrentDevice.DeviceInfo.Name.StartsWith(TEXT("Xbox Wired Controller")))
 						{
 							CurrentDevice.ControllerClass = ControllerClassType::XBoxWired;
@@ -996,6 +1007,22 @@ void FAndroidInputInterface::SendControllerEvents()
 						else if (CurrentDevice.DeviceInfo.Name.StartsWith(TEXT("Razer")))
 						{
 							CurrentDevice.bSupportsHat = true;
+							if (CurrentDevice.DeviceInfo.Name.StartsWith(TEXT("Razer Kishi V2 Pro XBox360")))
+							{
+								CurrentDevice.ControllerClass = ControllerClassType::XBoxWired;
+								CurrentDevice.bSupportsHat = true;
+								CurrentDevice.bTriggersUseThresholdForClick = true;
+								CurrentDevice.bRightStickZRZ = false;
+								CurrentDevice.bRightStickRXRY = true;
+								CurrentDevice.bMapZRZToTriggers = true;
+								CurrentDevice.LTAnalogRangeMinimum = -1.0f;
+								CurrentDevice.RTAnalogRangeMinimum = -1.0f;
+							}
+							else if (CurrentDevice.DeviceInfo.Name.StartsWith(TEXT("Razer Kishi V2")))
+							{
+								CurrentDevice.ControllerClass = ControllerClassType::XBoxWired;
+								CurrentDevice.bTriggersUseThresholdForClick = true;
+							}
 						}
 
 						IPlatformInputDeviceMapper& DeviceMapper = IPlatformInputDeviceMapper::Get();
