AppCore/Mac: Adds handling for safeAreaInsets.

* Monitors with notched screens (e.g. screens with camera housing built in) need to take into account the safeAreaInset when in fullscreen mode.
* Remove some unused and unexecutable code around getting the monitor name.

diff --git a/Engine/Source/Runtime/ApplicationCore/Private/Mac/CocoaWindow.cpp b/Engine/Source/Runtime/ApplicationCore/Private/Mac/CocoaWindow.cpp
--- a/Engine/Source/Runtime/ApplicationCore/Private/Mac/CocoaWindow.cpp
+++ b/Engine/Source/Runtime/ApplicationCore/Private/Mac/CocoaWindow.cpp
@@ -60,22 +60,7 @@ NSString* NSPerformDragOperation = @"NSPerformDragOperation";
 - (NSView*)openGLView
 {
 	SCOPED_AUTORELEASE_POOL;
-	if (FPlatformMisc::IsRunningOnMavericks() && [self styleMask] & (NSWindowStyleMaskTexturedBackground))
-	{
-		NSView* SuperView = [[self contentView] superview];
-		for (NSView* View in [SuperView subviews])
-		{
-			if ([View isKindOfClass:[FCocoaTextView class]])
-			{
-				return View;
-			}
-		}
-		return nil;
-	}
-	else
-	{
-		return [self contentView];
-	}
+	return [self contentView];
 }
 
 - (void)setAcceptsInput:(bool)InAcceptsInput
@@ -83,6 +68,10 @@ NSString* NSPerformDragOperation = @"NSPerformDragOperation";
 	bAcceptsInput = InAcceptsInput;
 }
 
+- (NSApplicationPresentationOptions)window:(NSWindow *)window willUseFullScreenPresentationOptions:(NSApplicationPresentationOptions)proposedOptions {
+		return proposedOptions;
+	  }
+
 - (void)setWindowMode:(EWindowMode::Type)NewWindowMode
 {
 	WindowMode = NewWindowMode;
@@ -152,7 +141,6 @@ NSString* NSPerformDragOperation = @"NSPerformDragOperation";
 		MainThreadCall(^{
 			[super setAlphaValue:WindowAlpha];
 		}, UnrealShowEventMode, false);
-
 	}
 }
 
diff --git a/Engine/Source/Runtime/ApplicationCore/Private/Mac/MacApplication.cpp b/Engine/Source/Runtime/ApplicationCore/Private/Mac/MacApplication.cpp
--- a/Engine/Source/Runtime/ApplicationCore/Private/Mac/MacApplication.cpp
+++ b/Engine/Source/Runtime/ApplicationCore/Private/Mac/MacApplication.cpp
@@ -1217,17 +1217,17 @@ void FMacApplication::OnWindowDidResize(TSharedRef<FMacWindow> Window, bool bRes
 
 	OnWindowDidMove(Window);
 
-	// default is no override
-	uint32 Width = FMath::TruncToInt([Window->GetWindowHandle() openGLFrame].size.width * Window->GetDPIScaleFactor());
-	uint32 Height = FMath::TruncToInt([Window->GetWindowHandle() openGLFrame].size.height * Window->GetDPIScaleFactor());
+	const FCocoaWindow* CocoaWindow = Window->GetWindowHandle();
+	const NSScreen* Screen = [CocoaWindow screen];
 
-	if (Window->GetWindowHandle().TargetWindowMode == EWindowMode::WindowedFullscreen)
-	{
-		// Grab current monitor data for sizing
-		Width = FMath::TruncToInt([[Window->GetWindowHandle() screen] visibleFrame].size.width * Window->GetDPIScaleFactor());
-		Height = FMath::TruncToInt([[Window->GetWindowHandle() screen] visibleFrame].size.height * Window->GetDPIScaleFactor());
-	}
+	// default is no override
+	const uint32 ScreenWidth  = FMath::TruncToInt([CocoaWindow openGLFrame].size.width * Window->GetDPIScaleFactor());
+	const uint32 ScreenHeight = FMath::TruncToInt([CocoaWindow openGLFrame].size.height * Window->GetDPIScaleFactor());
 
+	// Grab current monitor data for sizing
+	const uint32 VisibleWidth = FMath::TruncToInt([Screen visibleFrame].size.width * Window->GetDPIScaleFactor());
+	const uint32 VisibleHeight = FMath::TruncToInt([Screen visibleFrame].size.height * Window->GetDPIScaleFactor());
+	
 	if (bRestoreMouseCursorLocking)
 	{
 		FMacCursor* MacCursor = (FMacCursor*)MacApplication->Cursor.Get();
@@ -1236,8 +1236,19 @@ void FMacApplication::OnWindowDidResize(TSharedRef<FMacWindow> Window, bool bRes
 			MacCursor->SetShouldIgnoreLocking(false);
 		}
 	}
-
-	MessageHandler->OnSizeChanged(Window, Width, Height);
+	// Depending on how the window is resized, it may result in actually moving the window slightly,
+	// e.g. going from fullscreenwindowed to fullscreen there's a few pixels of extra padding below
+	// camera housing on Apple screens for the menu bar, fullscreen doesn't have menu bars so that extra
+	// padding is removed, in effect making the window shift up.
+	if (Window->GetWindowHandle().TargetWindowMode == EWindowMode::WindowedFullscreen || Window->GetWindowHandle().TargetWindowMode == EWindowMode::Fullscreen  )
+	{
+		MessageHandler->OnMovedWindow(Window, ScreenWidth - VisibleWidth, ScreenHeight - VisibleHeight);
+		MessageHandler->OnSizeChanged(Window, VisibleWidth, VisibleHeight);
+	}
+	else
+	{
+		MessageHandler->OnSizeChanged(Window, ScreenWidth, ScreenHeight);
+	}
 	MessageHandler->OnResizingWindow(Window);
 }
 
@@ -1245,7 +1256,6 @@ void FMacApplication::OnWindowDidResize(TSharedRef<FMacWindow> Window, bool bRes
 void FMacApplication::OnWindowChangedScreen(TSharedRef<FMacWindow> Window)
 {
 	SCOPED_AUTORELEASE_POOL;
-
 	MessageHandler->HandleDPIScaleChanged(Window);
 }
 
@@ -2271,18 +2281,19 @@ void FDisplayMetrics::RebuildDisplayMetrics(FDisplayMetrics& OutDisplayMetrics)
 
 			Info.DisplayRect = FPlatformRect
 			(
-				FMath::TruncToInt(Screen->FramePixels.origin.x),
-				FMath::TruncToInt(Screen->FramePixels.origin.y),
-				FMath::TruncToInt(Screen->FramePixels.origin.x + Screen->FramePixels.size.width),
-				FMath::TruncToInt(Screen->FramePixels.origin.y + Screen->FramePixels.size.height)
+				FMath::TruncToInt(Screen->FramePixels.origin.x  + Screen->SafeAreaInsets.left),
+				FMath::TruncToInt(Screen->FramePixels.origin.y + Screen->SafeAreaInsets.top),
+				FMath::TruncToInt(Screen->FramePixels.size.width - Screen->SafeAreaInsets.right),
+				FMath::TruncToInt(Screen->FramePixels.size.height - Screen->SafeAreaInsets.bottom)
 			);
 			Info.WorkArea = FPlatformRect
 			(
 				FMath::TruncToInt(Screen->VisibleFramePixels.origin.x),
 				FMath::TruncToInt(Screen->VisibleFramePixels.origin.y),
-				FMath::TruncToInt(Screen->VisibleFramePixels.origin.x + Screen->VisibleFramePixels.size.width),
-				FMath::TruncToInt(Screen->VisibleFramePixels.origin.y + Screen->VisibleFramePixels.size.height)
+				FMath::TruncToInt(Screen->VisibleFramePixels.size.width),
+				FMath::TruncToInt(Screen->VisibleFramePixels.size.height)
 			);
+
 			Info.bIsPrimary = Screen->Screen == [NSScreen mainScreen];
 
 			// dpi computations
@@ -2292,50 +2303,7 @@ void FDisplayMetrics::RebuildDisplayMetrics(FDisplayMetrics& OutDisplayMetrics)
 			float VerticalDPI = MilimetreInch * (float)Info.NativeHeight / (float)DisplayPhysicalSize.height;
 			Info.DPI = FMath::CeilToInt((HorizontalDPI + VerticalDPI) / 2.0f);
 
-			// Monitor's name can only be obtained from IOKit
-			io_iterator_t IOIterator;
-            mach_port_t mainPortDefault;
-            if (@available(macOS 12.0, iOS 15.0, *))
-            {
-                mainPortDefault = kIOMainPortDefault;
-            }
-            else
-            {
-                mainPortDefault = kIOMasterPortDefault;
-            }
-			kern_return_t Result = IOServiceGetMatchingServices(kIOMainPortDefault, IOServiceMatching("IODisplayConnect"), &IOIterator);
-			if (Result == kIOReturnSuccess)
-			{
-				io_object_t Device;
-				while ((Device = IOIteratorNext(IOIterator)))
-				{
-					CFDictionaryRef Dictionary = IODisplayCreateInfoDictionary(Device, kIODisplayOnlyPreferredName);
-					if (Dictionary)
-					{
-						const uint32 VendorID = [(__bridge NSNumber*)CFDictionaryGetValue(Dictionary, CFSTR(kDisplayVendorID)) unsignedIntegerValue];
-						const uint32 ProductID = [(__bridge NSNumber*)CFDictionaryGetValue(Dictionary, CFSTR(kDisplayProductID)) unsignedIntegerValue];
-						const uint32 SerialNumber = [(__bridge NSNumber*)CFDictionaryGetValue(Dictionary, CFSTR(kDisplaySerialNumber)) unsignedIntegerValue];
-
-						if (VendorID == CGDisplayVendorNumber(DisplayID) && ProductID == CGDisplayModelNumber(DisplayID) && SerialNumber == CGDisplaySerialNumber(DisplayID))
-						{
-							NSDictionary* NamesDictionary = (__bridge NSDictionary*)CFDictionaryGetValue(Dictionary, CFSTR(kDisplayProductName));
-							if (NamesDictionary && NamesDictionary.count > 0)
-							{
-								Info.Name = (NSString*)[NamesDictionary objectForKey:[NamesDictionary.allKeys objectAtIndex:0]];
-								CFRelease(Dictionary);
-								IOObjectRelease(Device);
-								break;
-							}
-						}
-
-						CFRelease(Dictionary);
-					}
-
-					IOObjectRelease(Device);
-				}
-
-				IOObjectRelease(IOIterator);
-			}
+			Info.Name = Screen->Screen.localizedName;
 
 			OutDisplayMetrics.MonitorInfo.Add(Info);
 		}
@@ -2350,8 +2318,17 @@ void FDisplayMetrics::RebuildDisplayMetrics(FDisplayMetrics& OutDisplayMetrics)
 	// Get the screen rect of the primary monitor, excluding taskbar etc.
 	OutDisplayMetrics.PrimaryDisplayWorkAreaRect.Left = FMath::TruncToInt(VisibleFrame.origin.x);
 	OutDisplayMetrics.PrimaryDisplayWorkAreaRect.Top = FMath::TruncToInt(VisibleFrame.origin.y);
-	OutDisplayMetrics.PrimaryDisplayWorkAreaRect.Right = FMath::TruncToInt(VisibleFrame.origin.x + VisibleFrame.size.width);
-	OutDisplayMetrics.PrimaryDisplayWorkAreaRect.Bottom = FMath::TruncToInt(OutDisplayMetrics.PrimaryDisplayWorkAreaRect.Top + VisibleFrame.size.height);
+	OutDisplayMetrics.PrimaryDisplayWorkAreaRect.Right = FMath::TruncToInt(VisibleFrame.size.width);
+	OutDisplayMetrics.PrimaryDisplayWorkAreaRect.Bottom = FMath::TruncToInt(VisibleFrame.size.height + OutDisplayMetrics.PrimaryDisplayWorkAreaRect.Top);
+
+	OutDisplayMetrics.TitleSafePaddingSize.X = PrimaryScreen->SafeAreaInsets.left;
+    OutDisplayMetrics.TitleSafePaddingSize.Y = PrimaryScreen->SafeAreaInsets.top;
+    OutDisplayMetrics.TitleSafePaddingSize.Z = PrimaryScreen->SafeAreaInsets.right;
+    OutDisplayMetrics.TitleSafePaddingSize.W = PrimaryScreen->SafeAreaInsets.bottom;
+
+	OutDisplayMetrics.TitleSafePaddingSize *= PrimaryScreen->Screen.backingScaleFactor;
+    
+    OutDisplayMetrics.ActionSafePaddingSize = OutDisplayMetrics.TitleSafePaddingSize;
 
 	// Apply the debug safe zones
 	OutDisplayMetrics.ApplyDefaultSafeZones();
diff --git a/Engine/Source/Runtime/ApplicationCore/Private/Mac/MacWindow.cpp b/Engine/Source/Runtime/ApplicationCore/Private/Mac/MacWindow.cpp
--- a/Engine/Source/Runtime/ApplicationCore/Private/Mac/MacWindow.cpp
+++ b/Engine/Source/Runtime/ApplicationCore/Private/Mac/MacWindow.cpp
@@ -50,8 +50,8 @@ void FMacWindow::Initialize( FMacApplication* const Application, const TSharedRe
 	const double ScreenDPIScaleFactor = FPlatformApplicationMisc::IsHighDPIModeEnabled() ? TargetScreen->Screen.backingScaleFactor : 1.0;
 	const FVector2D CocoaPosition = FMacApplication::ConvertSlatePositionToCocoa(PositionX, PositionY);
 	const NSRect ViewRect = NSMakeRect(CocoaPosition.X, CocoaPosition.Y - (SizeY / ScreenDPIScaleFactor) + 1, SizeX / ScreenDPIScaleFactor, SizeY / ScreenDPIScaleFactor);
-
 	uint32 WindowStyle = 0;
+
 	if( Definition->IsRegularWindow )
 	{
 		if( Definition->HasCloseButton )
@@ -60,7 +60,7 @@ void FMacWindow::Initialize( FMacApplication* const Application, const TSharedRe
 		}
 
 		// In order to support rounded, shadowed windows set the window to be titled - we'll set the content view to cover the whole window
-		WindowStyle |= NSWindowStyleMaskTitled | (FPlatformMisc::IsRunningOnMavericks() ? NSWindowStyleMaskTexturedBackground : NSWindowStyleMaskFullSizeContentView);
+		WindowStyle |= NSWindowStyleMaskTitled | NSWindowStyleMaskFullSizeContentView;
 		
 		if( Definition->SupportsMinimize )
 		{
@@ -79,7 +79,7 @@ void FMacWindow::Initialize( FMacApplication* const Application, const TSharedRe
 	if( Definition->HasOSWindowBorder )
 	{
 		WindowStyle |= NSWindowStyleMaskTitled;
-		WindowStyle &= FPlatformMisc::IsRunningOnMavericks() ? ~NSWindowStyleMaskTexturedBackground : ~NSWindowStyleMaskFullSizeContentView;
+		WindowStyle &= ~NSWindowStyleMaskFullSizeContentView;
 	}
 
 	MainThreadCall(^{
@@ -384,7 +384,8 @@ void FMacWindow::SetWindowMode(EWindowMode::Type NewWindowMode)
 {
 	if(WindowHandle)
 	{
-		ApplySizeAndModeChanges(PositionX, PositionY, FMath::TruncToInt(WindowHandle.contentView.frame.size.width), FMath::TruncToInt(WindowHandle.contentView.frame.size.height), NewWindowMode);
+		ApplySizeAndModeChanges(PositionX, PositionY, FMath::TruncToInt(WindowHandle.contentView.frame.size.width),
+								FMath::TruncToInt(WindowHandle.contentView.frame.size.height), NewWindowMode);
 	}
 }
 
@@ -732,7 +733,9 @@ void FMacWindow::UpdateFullScreenState(bool bToggleFullScreen)
 				WindowedModeSavedState.WindowLevel = WindowHandle.level;
 				[WindowHandle setLevel:CGShieldingWindowLevel() + 1];
 			}
-			[NSApp setPresentationOptions:NSApplicationPresentationHideDock | NSApplicationPresentationHideMenuBar];
+			
+			// -toggleFullScreen implicitly sets these, but it doesn't hurt to set them ourselves to match.
+			[NSApp setPresentationOptions:NSApplicationPresentationAutoHideDock | NSApplicationPresentationAutoHideMenuBar];
 		}
 		else if (WindowHandle.level != WindowedModeSavedState.WindowLevel)
 		{
diff --git a/Engine/Source/Runtime/ApplicationCore/Public/Mac/MacApplication.h b/Engine/Source/Runtime/ApplicationCore/Public/Mac/MacApplication.h
--- a/Engine/Source/Runtime/ApplicationCore/Public/Mac/MacApplication.h
+++ b/Engine/Source/Runtime/ApplicationCore/Public/Mac/MacApplication.h
@@ -119,8 +119,11 @@ struct FMacScreen
 	NSRect VisibleFrame;
 	NSRect FramePixels;
 	NSRect VisibleFramePixels;
+	NSEdgeInsets SafeAreaInsets;
 
-	FMacScreen(NSScreen* InScreen) : Screen([InScreen retain]), Frame(InScreen.frame), VisibleFrame(InScreen.visibleFrame), FramePixels(InScreen.frame), VisibleFramePixels(InScreen.visibleFrame) {}
+	FMacScreen(NSScreen* InScreen) : Screen([InScreen retain]), Frame(InScreen.frame), VisibleFrame(InScreen.visibleFrame),
+									 FramePixels(InScreen.frame), VisibleFramePixels(InScreen.visibleFrame),
+									 SafeAreaInsets(InScreen.safeAreaInsets) {}
 	~FMacScreen() { [Screen release]; }
 };
 typedef TSharedRef<FMacScreen, ESPMode::ThreadSafe> FMacScreenRef;
