Make sure main Editor window won't bringToFront, align behavior with PC and other platforms

diff --git a/Engine/Source/Runtime/ApplicationCore/Private/Mac/MacWindow.cpp b/Engine/Source/Runtime/ApplicationCore/Private/Mac/MacWindow.cpp
--- a/Engine/Source/Runtime/ApplicationCore/Private/Mac/MacWindow.cpp
+++ b/Engine/Source/Runtime/ApplicationCore/Private/Mac/MacWindow.cpp
@@ -15,6 +15,7 @@ FMacWindow::FMacWindow()
 ,	bIsVisible(false)
 ,	bIsClosed(false)
 ,	bIsFirstTimeVisible(true)
+,   bIsMainEditorWindow(false)
 {
 }
 
@@ -35,6 +36,9 @@ void FMacWindow::Initialize( FMacApplication* const Application, const TSharedRe
 
 	OwningApplication = Application;
 	Definition = InDefinition;
+    
+    // The main Editor window, like Batman, doesn't have a valid parent
+    bIsMainEditorWindow = !InParent.IsValid();
 
 	// Finally, let's initialize the new native window object.  Calling this function will often cause OS
 	// window messages to be sent! (such as activation messages)
@@ -256,8 +260,8 @@ void FMacWindow::MoveWindowTo( int32 X, int32 Y )
 void FMacWindow::BringToFront( bool bForce )
 {
 	bIsVisible = (bIsVisible || bForce);
-
-	if (!bIsClosed && bIsVisible)
+    
+	if (!bIsClosed && bIsVisible && !bIsMainEditorWindow) // Don't try to bring editor window to front
 	{
 		SCOPED_AUTORELEASE_POOL;
 
@@ -438,12 +442,7 @@ bool FMacWindow::GetRestoredDimensions(int32& X, int32& Y, int32& Width, int32&
 
 void FMacWindow::SetWindowFocus()
 {
-	MainThreadCall(^{
-		SCOPED_AUTORELEASE_POOL;
-		[WindowHandle orderFrontAndMakeMain:true andKey:true];
-	}, UnrealShowEventMode, true);
-
-	MacApplication->OnWindowOrderedFront(SharedThis(this));
+    BringToFront(true);
 }
 
 void FMacWindow::SetOpacity( const float InOpacity )
diff --git a/Engine/Source/Runtime/ApplicationCore/Public/Mac/MacWindow.h b/Engine/Source/Runtime/ApplicationCore/Public/Mac/MacWindow.h
--- a/Engine/Source/Runtime/ApplicationCore/Public/Mac/MacWindow.h
+++ b/Engine/Source/Runtime/ApplicationCore/Public/Mac/MacWindow.h
@@ -121,4 +121,6 @@ private:
 
 	/** Whether the window is yet to have its first Show() call. This is set false after first Show(). */
 	bool bIsFirstTimeVisible : 1;
+    
+    bool bIsMainEditorWindow : 1;
 };
