Fix XBox and Razer wired controllers on Android 12+

#android

diff --git a/Engine/Source/Runtime/ApplicationCore/Private/Android/AndroidInputInterface.cpp b/Engine/Source/Runtime/ApplicationCore/Private/Android/AndroidInputInterface.cpp
--- a/Engine/Source/Runtime/ApplicationCore/Private/Android/AndroidInputInterface.cpp
+++ b/Engine/Source/Runtime/ApplicationCore/Private/Android/AndroidInputInterface.cpp
@@ -935,11 +935,16 @@ void FAndroidInputInterface::SendControllerEvents()
 							CurrentDevice.ControllerClass = ControllerClassType::XBoxWired;
 							CurrentDevice.bSupportsHat = true;
 							CurrentDevice.bTriggersUseThresholdForClick = true;
-							CurrentDevice.bRightStickZRZ = false;
-							CurrentDevice.bRightStickRXRY = true;
-							CurrentDevice.bMapZRZToTriggers = true;
-							CurrentDevice.LTAnalogRangeMinimum = -1.0f;
-							CurrentDevice.RTAnalogRangeMinimum = -1.0f;
+
+							// different mapping before Android 12
+							if (FAndroidMisc::GetAndroidBuildVersion() < 31)
+							{
+								CurrentDevice.bRightStickZRZ = false;
+								CurrentDevice.bRightStickRXRY = true;
+								CurrentDevice.bMapZRZToTriggers = true;
+								CurrentDevice.LTAnalogRangeMinimum = -1.0f;
+								CurrentDevice.RTAnalogRangeMinimum = -1.0f;
+							}
 						}
 						else if (CurrentDevice.DeviceInfo.Name.StartsWith(TEXT("Xbox Wired Controller")))
 						{
@@ -1012,11 +1017,16 @@ void FAndroidInputInterface::SendControllerEvents()
 								CurrentDevice.ControllerClass = ControllerClassType::XBoxWired;
 								CurrentDevice.bSupportsHat = true;
 								CurrentDevice.bTriggersUseThresholdForClick = true;
-								CurrentDevice.bRightStickZRZ = false;
-								CurrentDevice.bRightStickRXRY = true;
-								CurrentDevice.bMapZRZToTriggers = true;
-								CurrentDevice.LTAnalogRangeMinimum = -1.0f;
-								CurrentDevice.RTAnalogRangeMinimum = -1.0f;
+
+								// different mapping before Android 12
+								if (FAndroidMisc::GetAndroidBuildVersion() < 31)
+								{
+									CurrentDevice.bRightStickZRZ = false;
+									CurrentDevice.bRightStickRXRY = true;
+									CurrentDevice.bMapZRZToTriggers = true;
+									CurrentDevice.LTAnalogRangeMinimum = -1.0f;
+									CurrentDevice.RTAnalogRangeMinimum = -1.0f;
+								}
 							}
 							else if (CurrentDevice.DeviceInfo.Name.StartsWith(TEXT("Razer Kishi V2")))
 							{
