Enable Swappy and NonVSync framerates on Android by default

diff --git a/Engine/Source/Runtime/ApplicationCore/Private/Android/AndroidPlatformFramePacer.cpp b/Engine/Source/Runtime/ApplicationCore/Private/Android/AndroidPlatformFramePacer.cpp
--- a/Engine/Source/Runtime/ApplicationCore/Private/Android/AndroidPlatformFramePacer.cpp
+++ b/Engine/Source/Runtime/ApplicationCore/Private/Android/AndroidPlatformFramePacer.cpp
@@ -13,13 +13,13 @@
 
 TAutoConsoleVariable<int32> FAndroidPlatformRHIFramePacer::CVarUseSwappyForFramePacing(
 	TEXT("a.UseSwappyForFramePacing"),
-	0,
+	1,
 	TEXT("True to use Swappy for frame pacing."));
 
 TAutoConsoleVariable<int32> FAndroidPlatformRHIFramePacer::CVarSupportNonVSyncMultipleFrameRates(
 	TEXT("a.SupportNonVSyncMultipleFrameRates"),
-	0,
-	TEXT("Set to True to support frame rates we cannot vsync at. Requires a.UseSwappyForFramePacing=1 for OpenGL ES."));
+	1,
+	TEXT("Set to True to support frame rates we cannot vsync at. Requires a.UseSwappyForFramePacing=1"));
 
 TAutoConsoleVariable<int32> FAndroidPlatformRHIFramePacer::CVarAllowFrameTimestamps(
 	TEXT("a.AllowFrameTimestamps"),
