Update force feedback support on Android
- Use VibrationEffect for intensity control
- Forward feedback to controller(s) if Android.UseControllerFeedback non-zero (defaults to 1)
- Fix issue with fingerprint scanner false detection as controller
- Does not continue feedback when backgrounded

#android

diff --git a/Engine/Build/Android/Java/src/com/epicgames/unreal/GameActivity.java.template b/Engine/Build/Android/Java/src/com/epicgames/unreal/GameActivity.java.template
--- a/Engine/Build/Android/Java/src/com/epicgames/unreal/GameActivity.java.template
+++ b/Engine/Build/Android/Java/src/com/epicgames/unreal/GameActivity.java.template
@@ -81,7 +81,10 @@ import android.os.Bundle;
 import android.os.PersistableBundle;
 import android.util.Log;
 
+import android.hardware.input.InputManager;
 import android.os.Vibrator;
+import android.os.VibrationEffect;
+import android.os.VibratorManager;
 import android.os.SystemClock;
 import android.os.Looper;
 import android.os.Handler;
@@ -249,6 +252,9 @@ public class GameActivity extends $${gameActivitySuperClass}$$ implements Surfac
 	static boolean first_acceleration_sample = true;
 	static final float SampleDecayRate = 0.85f;
 
+	private InputManager inputManager;
+	private Vibrator deviceVibrator;
+
 	private Rect safezone = new Rect(0, 0, 0, 0);
 	private RectF safezoneF = new RectF(0.0f, 0.0f, 0.0f, 0.0f);
 
@@ -2341,6 +2347,17 @@ public class GameActivity extends $${gameActivitySuperClass}$$ implements Surfac
 		magnetometer = sensorManager.getDefaultSensor(Sensor.TYPE_MAGNETIC_FIELD);
 		gyroscope = sensorManager.getDefaultSensor(Sensor.TYPE_GYROSCOPE);
 
+		inputManager = (InputManager) getSystemService(Context.INPUT_SERVICE);
+		deviceVibrator = null;
+		if (android.os.Build.VERSION.SDK_INT >= android.os.Build.VERSION_CODES.S)
+		{
+			VibratorManager vibratorManager = (VibratorManager) getSystemService(Context.VIBRATOR_MANAGER_SERVICE);
+			deviceVibrator = vibratorManager.getDefaultVibrator();
+		}
+		else
+		{
+			deviceVibrator = (Vibrator) getSystemService(Context.VIBRATOR_SERVICE);
+		}
 
 		// disable autofill to avoid possible Android 10 crash (android.view.autofill.AutofillManager.ensureServiceClientAddedIfNeededLocked() SyncResultReceiver$TimeoutException)
 		if (Build.VERSION.SDK_INT >= Build.VERSION_CODES.O)
@@ -3933,31 +3950,32 @@ public class GameActivity extends $${gameActivitySuperClass}$$ implements Surfac
 	}
 
 	private class VibrateRunnable implements Runnable {
-		private int duration;
+		private int intensity;
+		private long duration;
 		private Vibrator vibrator;
 
-		VibrateRunnable(final int Duration, final Vibrator vibrator)
+		VibrateRunnable(final int Intensity, final long Duration, final Vibrator vibrator)
 		{
+			this.intensity = Intensity;
 			this.duration = Duration;
 			this.vibrator = vibrator;
 		}
 		public void run ()
 		{
-			if (duration < 1)
+			if (intensity < 1 || duration < 1)
 			{
 				vibrator.cancel();
 			} else {
-				vibrator.vibrate(duration);
+				vibrator.vibrate(VibrationEffect.createOneShot(duration, intensity));
 			}
 		}
 	}
 
-	public void AndroidThunkJava_Vibrate(int Duration)
+	public void AndroidThunkJava_Vibrate(int Intensity, int Duration)
 	{
-		Vibrator vibrator = (Vibrator)getSystemService(VIBRATOR_SERVICE);
-		if (vibrator != null)
+		if (deviceVibrator != null)
 		{
-			_activity.runOnUiThread(new VibrateRunnable(Duration, vibrator));
+			_activity.runOnUiThread(new VibrateRunnable(Intensity, (long)Duration, deviceVibrator));
 		}
 	}
 
@@ -5245,8 +5263,9 @@ public class GameActivity extends $${gameActivitySuperClass}$$ implements Surfac
 		public int controllerId;
 		public String name;
 		public String descriptor;
+		public int feedbackMotorCount;
 
-		InputDeviceInfo(int did, int vid, int pid, int cid, String inName, String inDescriptor)
+		InputDeviceInfo(int did, int vid, int pid, int cid, String inName, String inDescriptor, int fmcount)
 		{
 			deviceId = did;
 			vendorId = vid;
@@ -5254,6 +5273,7 @@ public class GameActivity extends $${gameActivitySuperClass}$$ implements Surfac
 			controllerId = cid;
 			name = inName;
 			descriptor = inDescriptor;
+			feedbackMotorCount = fmcount;
 		}
 	}
 
@@ -5268,9 +5288,20 @@ public class GameActivity extends $${gameActivitySuperClass}$$ implements Surfac
 				int vendorId = 0;
 				int productId = 0;
 				int controllerNumber = 0;
+				int feedbackMotorCount = 0;
 
 				// requires 4.1 (Jellybean)
 				String descriptor = ANDROID_BUILD_VERSION >= 16 ? inputDevice.getDescriptor() : Integer.toString(deviceId);
+				
+				// supported on 12.0 onward
+				if (ANDROID_BUILD_VERSION >= 31)
+				{
+					VibratorManager vibratorManager = (VibratorManager) inputDevice.getVibratorManager();
+					if (vibratorManager != null)
+					{
+						feedbackMotorCount = vibratorManager.getVibratorIds().length;
+					}
+				}
 
 				// supported on 4.4 (KitKat) onward
 				if (ANDROID_BUILD_VERSION >= 19)
@@ -5284,16 +5315,59 @@ public class GameActivity extends $${gameActivitySuperClass}$$ implements Surfac
 					{
 						if (deviceInfo.IsMatch(vendorId, productId))
 						{
-							return new InputDeviceInfo(deviceId, vendorId, productId, controllerNumber, deviceInfo.name, descriptor);
+							return new InputDeviceInfo(deviceId, vendorId, productId, controllerNumber, deviceInfo.name, descriptor, feedbackMotorCount);
 						}
 					}
 				}
 
 				// use device name as fallback (may be generic like "Bluetooth HID" so not always useful)
-				return new InputDeviceInfo(deviceId, vendorId, productId, controllerNumber, inputDevice.getName(), descriptor);
+				return new InputDeviceInfo(deviceId, vendorId, productId, controllerNumber, inputDevice.getName(), descriptor, feedbackMotorCount);
 			}
 		}
-		return new InputDeviceInfo(deviceId, 0, 0, -1, "Unknown", "Unknown");
+		return new InputDeviceInfo(deviceId, 0, 0, -1, "Unknown", "Unknown", 0);
+	}
+
+	private boolean UpdateVibrator(Vibrator vibrator, int intensity, long duration)
+	{
+		if (vibrator != null)
+		{
+			_activity.runOnUiThread(new VibrateRunnable(intensity, duration, vibrator));
+			return true;
+		}
+		return false;
+	}
+
+	public boolean AndroidThunkJava_SetInputDeviceVibrators(int deviceId, int leftIntensity, int leftDuration, int rightIntensity, int rightDuration)
+	{
+		boolean result = false;
+		
+		InputDevice inputDevice = inputManager.getInputDevice(deviceId);
+        if (inputDevice != null)
+		{
+			int intensity = leftIntensity > rightIntensity ? leftIntensity : rightIntensity;
+			long duration = leftDuration > rightDuration ? leftDuration : rightDuration;
+			
+			if (ANDROID_BUILD_VERSION >= 31)
+			{
+				VibratorManager vibratorManager = inputDevice.getVibratorManager();
+				int[] vibratorIds = vibratorManager.getVibratorIds();
+				int count = vibratorIds.length;
+				if (count == 1)
+				{
+					result |= UpdateVibrator(vibratorManager.getVibrator(vibratorIds[0]), intensity, duration);
+				}
+				else if (count > 1)
+				{
+					result |= UpdateVibrator(vibratorManager.getVibrator(vibratorIds[0]), leftIntensity, leftDuration);
+					result |= UpdateVibrator(vibratorManager.getVibrator(vibratorIds[1]), rightIntensity, rightDuration);
+				}
+            }
+			else
+			{
+				result |= UpdateVibrator(inputDevice.getVibrator(), intensity, duration);
+            }
+  		}
+		return result;
 	}
 
 	public void AndroidThunkJava_UseSurfaceViewWorkaround()
@@ -5364,7 +5438,12 @@ public class GameActivity extends $${gameActivitySuperClass}$$ implements Surfac
 			if (((inputDevice.getSources() & InputDevice.SOURCE_GAMEPAD) == InputDevice.SOURCE_GAMEPAD)
                 || ((inputDevice.getSources() & InputDevice.SOURCE_JOYSTICK) == InputDevice.SOURCE_JOYSTICK))
 			{
-				return true;
+				// make sure it isn't actually a keyboard or finger scanner
+				if ((inputDevice.getKeyboardType() != InputDevice.KEYBOARD_TYPE_ALPHABETIC)
+					&& (!inputDevice.getName().equals("uinput-fpc")))
+				{
+					return true;
+				}
 			}
 		}
 
diff --git a/Engine/Source/Runtime/ApplicationCore/Private/Android/AndroidInputInterface.cpp b/Engine/Source/Runtime/ApplicationCore/Private/Android/AndroidInputInterface.cpp
--- a/Engine/Source/Runtime/ApplicationCore/Private/Android/AndroidInputInterface.cpp
+++ b/Engine/Source/Runtime/ApplicationCore/Private/Android/AndroidInputInterface.cpp
@@ -2,6 +2,7 @@
 
 #include "Android/AndroidInputInterface.h"
 #if USE_ANDROID_INPUT
+#include "Android/AndroidEventManager.h"
 //#include "AndroidInputDeviceMappings.h"
 #include "Misc/ConfigCacheIni.h"
 #include "IInputDevice.h"
@@ -22,7 +23,7 @@ FCriticalSection FAndroidInputInterface::TouchInputCriticalSection;
 
 FAndroidGamepadDeviceMapping FAndroidInputInterface::DeviceMapping[MAX_NUM_CONTROLLERS];
 
-bool FAndroidInputInterface::VibeIsOn;
+int32 FAndroidInputInterface::CurrentVibeIntensity;
 int32 FAndroidInputInterface::MaxVibeTime = 1000;
 double FAndroidInputInterface::LastVibeUpdateTime = 0.0;
 FForceFeedbackValues FAndroidInputInterface::VibeValues;
@@ -33,6 +34,7 @@ bool FAndroidInputInterface::bControllersBlockDeviceFeedback = false;
 
 FAndroidControllerData FAndroidInputInterface::OldControllerData[MAX_NUM_CONTROLLERS];
 FAndroidControllerData FAndroidInputInterface::NewControllerData[MAX_NUM_CONTROLLERS];
+FAndroidControllerVibeState FAndroidInputInterface::ControllerVibeState[MAX_NUM_CONTROLLERS];
 
 FGamepadKeyNames::Type FAndroidInputInterface::ButtonMapping[MAX_NUM_CONTROLLER_BUTTONS];
 float FAndroidInputInterface::InitialButtonRepeatDelay;
@@ -48,6 +50,13 @@ TArray<FAndroidInputInterface::MotionData> FAndroidInputInterface::MotionDataSta
 TArray<FAndroidInputInterface::MouseData> FAndroidInputInterface::MouseDataStack
 	= TArray<FAndroidInputInterface::MouseData>();
 
+int32 GAndroidUseControllerFeedback = 1;
+static FAutoConsoleVariableRef CVarAndroidUseControllerFeedback(
+	TEXT("Android.UseControllerFeedback"),
+	GAndroidUseControllerFeedback,
+	TEXT("If set to non-zero, controllers with force feedback support will be active (Default: 1)"),
+	ECVF_Default);
+
 int32 GAndroidOldXBoxWirelessFirmware = 0;
 static FAutoConsoleVariableRef CVarAndroidOldXBoxWirelessFirmware(
 	TEXT("Android.OldXBoxWirelessFirmware"),
@@ -120,13 +129,13 @@ FAndroidInputInterface::FAndroidInputInterface(const TSharedRef< FGenericApplica
 	GConfig->GetFloat(TEXT("/Script/Engine.InputSettings"), TEXT("InitialButtonRepeatDelay"), InitialButtonRepeatDelay, GInputIni);
 	GConfig->GetFloat(TEXT("/Script/Engine.InputSettings"), TEXT("ButtonRepeatDelay"), ButtonRepeatDelay, GInputIni);
 
-	VibeIsOn = false;
-
-	for (int32 DeviceIndex = 0; DeviceIndex < MAX_NUM_CONTROLLERS; DeviceIndex++)
-	{
-		DeviceMapping[DeviceIndex].DeviceInfo.DeviceId = 0;
-		DeviceMapping[DeviceIndex].DeviceState = MappingState::Unassigned;
-	}
+	CurrentVibeIntensity = 0;
+	FMemory::Memset(VibeValues, 0);
+	
+	FMemory::Memset(DeviceMapping, 0);
+	FMemory::Memset(OldControllerData, 0);
+	FMemory::Memset(NewControllerData, 0);
+	FMemory::Memset(ControllerVibeState, 0);
 }
 
 void FAndroidInputInterface::ResetGamepadAssignments()
@@ -218,12 +227,6 @@ void FAndroidInputInterface::Tick(float DeltaTime)
 	{
 		(*DeviceIt)->Tick(DeltaTime);
 	}
-
-	if (VibeIsOn && 1000 * (FPlatformTime::Seconds() - LastVibeUpdateTime) > MaxVibeTime)
-	{
-		VibeIsOn = false;
-		UpdateVibeMotors();
-	}
 }
 
 void FAndroidInputInterface::SetLightColor(int32 ControllerId, FColor Color)
@@ -254,11 +257,40 @@ void FAndroidInputInterface::SetForceFeedbackChannelValue(int32 ControllerId, FF
 		}
 	}
 
+	// If didn't already assign feedback and active controller has feedback support use it, if enabled
+	if (!bDidFeedback && IsControllerAssignedToGamepad(ControllerId) && GAndroidUseControllerFeedback != 0 && DeviceMapping[ControllerId].DeviceInfo.FeedbackMotorCount > 0)
+	{
+		switch (ChannelType)
+		{
+			case FForceFeedbackChannelType::LEFT_LARGE:
+				ControllerVibeState[ControllerId].VibeValues.LeftLarge = Value;
+				break;
+
+			case FForceFeedbackChannelType::LEFT_SMALL:
+				ControllerVibeState[ControllerId].VibeValues.LeftSmall = Value;
+				break;
+
+			case FForceFeedbackChannelType::RIGHT_LARGE:
+				ControllerVibeState[ControllerId].VibeValues.RightLarge = Value;
+				break;
+
+			case FForceFeedbackChannelType::RIGHT_SMALL:
+				ControllerVibeState[ControllerId].VibeValues.RightSmall = Value;
+				break;
+
+			default:
+				// Unknown channel, so ignore it
+			break;
+		}
+		bDidFeedback = true;
+	}
+
 	bDidFeedback |= IsGamepadAttached() && bControllersBlockDeviceFeedback;
 
 	// If controller handled force feedback don't do it on the phone
 	if (bDidFeedback)
 	{
+		VibeValues.LeftLarge = VibeValues.RightLarge = VibeValues.LeftSmall = VibeValues.RightSmall = 0.0f;
 		return;
 	}
 
@@ -288,9 +320,6 @@ void FAndroidInputInterface::SetForceFeedbackChannelValue(int32 ControllerId, FF
 			// Unknown channel, so ignore it
 			break;
 	}
-
-	// Update with the latest values (wait for SendControllerEvents later?)
-	UpdateVibeMotors();
 }
 
 void FAndroidInputInterface::SetForceFeedbackChannelValues(int32 ControllerId, const FForceFeedbackValues &Values)
@@ -305,6 +334,13 @@ void FAndroidInputInterface::SetForceFeedbackChannelValues(int32 ControllerId, c
 		}
 	}
 
+	// If didn't already assign feedback and active controller has feedback support use it, if enabled
+	if (!bDidFeedback && IsControllerAssignedToGamepad(ControllerId) && GAndroidUseControllerFeedback != 0 && DeviceMapping[ControllerId].DeviceInfo.FeedbackMotorCount > 0)
+	{
+		ControllerVibeState[ControllerId].VibeValues = Values;
+		bDidFeedback = true;
+	}
+
 	bDidFeedback |= IsGamepadAttached() && bControllersBlockDeviceFeedback;
 
 	// If controller handled force feedback don't do it on the phone
@@ -316,12 +352,6 @@ void FAndroidInputInterface::SetForceFeedbackChannelValues(int32 ControllerId, c
 	{
 		VibeValues = Values;
 	}
-
-	// Note: only one motor on Android at the moment, but remember all the settings
-	// update will look at combination of all values to pick state
-
-	// Update with the latest values (wait for SendControllerEvents later?)
-	UpdateVibeMotors();
 }
 
 void FAndroidInputInterface::SetHapticFeedbackValues(int32 ControllerId, int32 Hand, const FHapticFeedbackValues& Values)
@@ -363,32 +393,112 @@ bool FAndroidInputInterface::IsGamepadAttached() const
 	return AndroidThunkCpp_IsGamepadAttached();
 }
 
-extern void AndroidThunkCpp_Vibrate(int32 Duration);
+static FORCEINLINE int32 ConvertToByte(float Value)
+{
+	int32 Setting = (int32)(Value * 255.0f);
+	return Setting < 0 ? 0 : (Setting < 255 ? Setting : 255);
+}
+
+extern void AndroidThunkCpp_Vibrate(int32 Intensity, int32 Duration);
+extern bool AndroidThunkCpp_SetInputDeviceVibrators(int32 deviceId, int32 leftIntensity, int32 leftDuration, int32 rightIntensity, int32 rightDuration);
 
 void FAndroidInputInterface::UpdateVibeMotors()
 {
+	// Turn off vibe if not in focus
+	bool bActive = CurrentVibeIntensity > 0;
+	if (!FAppEventManager::GetInstance()->IsGameInFocus())
+	{
+		if (bActive)
+		{
+			AndroidThunkCpp_Vibrate(0, MaxVibeTime);
+			CurrentVibeIntensity = 0;
+		}
+		return;
+	}
+
 	// Use largest vibration state as value
 	const float MaxLeft = VibeValues.LeftLarge > VibeValues.LeftSmall ? VibeValues.LeftLarge : VibeValues.LeftSmall;
 	const float MaxRight = VibeValues.RightLarge > VibeValues.RightSmall ? VibeValues.RightLarge : VibeValues.RightSmall;
 	const float Value = MaxLeft > MaxRight ? MaxLeft : MaxRight;
 
-	if (VibeIsOn)
+	int32 Intensity = ConvertToByte(Value);
+
+	// if previously active and overtime, current state is off
+	double CurrentTime = FPlatformTime::Seconds();
+	bool bOvertime = 1000 * (CurrentTime - LastVibeUpdateTime) >= MaxVibeTime;
+	if (bActive && bOvertime)
 	{
-		// Turn it off if below threshold
-		if (Value < 0.3f)
-		{
-			AndroidThunkCpp_Vibrate(0);
-			VibeIsOn = false;
-		}
+		CurrentVibeIntensity = 0;
 	}
-	else {
-		if (Value >= 0.3f)
+
+	// update if not already active at same level
+	if (CurrentVibeIntensity != Intensity)
+	{
+		AndroidThunkCpp_Vibrate(Intensity, MaxVibeTime);
+		CurrentVibeIntensity = Intensity;
+		LastVibeUpdateTime = CurrentTime;
+		//FPlatformMisc::LowLevelOutputDebugStringf(TEXT("VibDevice %f: %d"), (float)LastVibeUpdateTime, Intensity);
+	}
+}
+
+void FAndroidInputInterface::UpdateControllerVibeMotors(int32 ControllerId)
+{
+	FAndroidControllerVibeState& State = ControllerVibeState[ControllerId];
+
+	// Turn off vibe if not in focus
+	bool bActive = State.LeftIntensity > 0 || State.RightIntensity > 0;
+	if (!FAppEventManager::GetInstance()->IsGameInFocus())
+	{
+		if (bActive)
 		{
-			// Turn it on for 10 seconds (or until below threshold)
-			AndroidThunkCpp_Vibrate(MaxVibeTime);
-			LastVibeUpdateTime = FPlatformTime::Seconds();
-			VibeIsOn = true;
+			AndroidThunkCpp_SetInputDeviceVibrators(DeviceMapping[ControllerId].DeviceInfo.DeviceId, 0, MaxVibeTime, 0, MaxVibeTime);
+			State.LeftIntensity = 0;
+			State.RightIntensity = 0;
 		}
+		return;
+	}
+
+	float MaxLeft;
+	float MaxRight;
+
+	// Use largest vibration state as value for controller type
+	switch (DeviceMapping[ControllerId].ControllerClass)
+	{
+		case ControllerClassType::PlaystationWireless:
+//			DS4 maybe should use this?  PS5 seems correct with generic
+//			MaxLeft = (State.VibeValues.LeftLarge > State.VibeValues.RightLarge ? State.VibeValues.LeftLarge : State.VibeValues.RightLarge);
+//			MaxRight = (State.VibeValues.LeftSmall > State.VibeValues.RightSmall ? State.VibeValues.LeftSmall : State.VibeValues.RightSmall);
+//			break;
+
+		case ControllerClassType::Generic:
+		case ControllerClassType::XBoxWired:
+		case ControllerClassType::XBoxWireless:
+		default:
+			MaxLeft = (State.VibeValues.LeftLarge > State.VibeValues.LeftSmall ? State.VibeValues.LeftLarge : State.VibeValues.LeftSmall);
+			MaxRight = (State.VibeValues.RightLarge > State.VibeValues.RightSmall ? State.VibeValues.RightLarge : State.VibeValues.RightSmall);
+			break;
+	}
+
+	int32 LeftIntensity = ConvertToByte(MaxLeft);
+	int32 RightIntensity = ConvertToByte(MaxRight);
+
+	// if previously active and overtime, current state is off
+	double CurrentTime = FPlatformTime::Seconds();
+	bool bOvertime = 1000 * (CurrentTime - ControllerVibeState[ControllerId].LastVibeUpdateTime) >= MaxVibeTime;
+	if (bActive && bOvertime)
+	{
+		State.LeftIntensity = 0;
+		State.RightIntensity = 0;
+	}
+
+	// update if not already active at same level
+	if (State.LeftIntensity != LeftIntensity || State.RightIntensity != RightIntensity)
+	{
+		AndroidThunkCpp_SetInputDeviceVibrators(DeviceMapping[ControllerId].DeviceInfo.DeviceId, LeftIntensity, MaxVibeTime, RightIntensity, MaxVibeTime);
+		State.LeftIntensity = LeftIntensity;
+		State.RightIntensity = RightIntensity;
+		State.LastVibeUpdateTime = CurrentTime;
+		//FPlatformMisc::LowLevelOutputDebugStringf(TEXT("VibController %f: %d, %d"), (float)State.LastVibeUpdateTime, LeftIntensity, RightIntensity);
 	}
 }
 
@@ -855,6 +965,9 @@ void FAndroidInputInterface::SendControllerEvents()
 
 	FScopeLock Lock(&TouchInputCriticalSection);
 
+	// Update device vibe motor with latest values (only one motor so look at combination of all values to pick state)
+	UpdateVibeMotors();
+
 	// Check for gamepads needing validation if enabled
 	if (bAllowControllers)
 	{
@@ -1212,6 +1325,15 @@ void FAndroidInputInterface::SendControllerEvents()
 				}
 			}
 
+			// send controller force feedback updates if enabled
+			if (GAndroidUseControllerFeedback != 0)
+			{
+				if (DeviceMapping[ControllerIndex].DeviceInfo.FeedbackMotorCount > 0)
+				{
+					UpdateControllerVibeMotors(ControllerIndex);
+				}
+			}
+
 			// Update the state for next time
 			OldControllerState = NewControllerState;
 		}
diff --git a/Engine/Source/Runtime/ApplicationCore/Public/Android/AndroidInputInterface.h b/Engine/Source/Runtime/ApplicationCore/Public/Android/AndroidInputInterface.h
--- a/Engine/Source/Runtime/ApplicationCore/Public/Android/AndroidInputInterface.h
+++ b/Engine/Source/Runtime/ApplicationCore/Public/Android/AndroidInputInterface.h
@@ -83,7 +83,7 @@ enum TouchType
 
 enum MappingState
 {
-	Unassigned,
+	Unassigned = 0,
 	ToValidate,
 	Valid
 };
@@ -112,6 +112,7 @@ struct FAndroidInputDeviceInfo {
 	int32 ControllerId;
 	FString Name;
 	FString Descriptor;
+	int32 FeedbackMotorCount;
 };
 
 struct FAndroidGamepadDeviceMapping
@@ -187,6 +188,14 @@ struct FAndroidControllerData
 	float RTAnalog;
 };
 
+struct FAndroidControllerVibeState
+{
+	FForceFeedbackValues VibeValues;
+	int32 LeftIntensity;
+	int32 RightIntensity;
+	double LastVibeUpdateTime;
+};
+
 enum FAndroidMessageType
 {
 	MessageType_KeyDown,
@@ -290,9 +299,12 @@ private:
 	/** Find controller index corresponding to validated deviceId (returns -1 if not found) */
 	static int32 FindExistingDevice(int32 deviceId);
 
-	/** Push Vibration changes to the controllers */
+	/** Push Vibration changes to the main device */
 	void UpdateVibeMotors();
 
+	/** Push Vibration changes to the controller */
+	void UpdateControllerVibeMotors(int32 ControllerId);
+
 	struct MotionData
 	{
 		FVector Tilt;
@@ -326,7 +338,7 @@ private:
 	static TArray<TouchInput> TouchInputStack;
 
 	/** Vibration settings */
-	static bool VibeIsOn;
+	static int32 CurrentVibeIntensity;
 	// Maximum time vibration will be triggered without an update
 	static int32 MaxVibeTime;
 	static double LastVibeUpdateTime;
@@ -345,6 +357,7 @@ private:
 
 	static FAndroidControllerData OldControllerData[MAX_NUM_CONTROLLERS];
 	static FAndroidControllerData NewControllerData[MAX_NUM_CONTROLLERS];
+	static FAndroidControllerVibeState ControllerVibeState[MAX_NUM_CONTROLLERS];
 
 	static FGamepadKeyNames::Type ButtonMapping[MAX_NUM_CONTROLLER_BUTTONS];
 
diff --git a/Engine/Source/Runtime/Core/Private/Android/AndroidPlatformMisc.cpp b/Engine/Source/Runtime/Core/Private/Android/AndroidPlatformMisc.cpp
--- a/Engine/Source/Runtime/Core/Private/Android/AndroidPlatformMisc.cpp
+++ b/Engine/Source/Runtime/Core/Private/Android/AndroidPlatformMisc.cpp
@@ -1639,9 +1639,11 @@ void FAndroidMisc::PrepareMobileHaptics(EMobileHapticsType Type)
 void FAndroidMisc::TriggerMobileHaptics()
 {
 #if USE_ANDROID_JNI
-	extern void AndroidThunkCpp_Vibrate(int32 Duration);
-	// tiny little vibration
-	AndroidThunkCpp_Vibrate(10);
+	extern void AndroidThunkCpp_Vibrate(int32 Intensity, int32 Duration);
+	// directly play a small vibration one-shot
+	// note: this will do nothing if device is already playing force feedback (non-zero intensity)
+	// but will play and not be cancelled by force feedback since it only sends updates when not already above zero
+	AndroidThunkCpp_Vibrate(255, 10);
 #endif
 }
 
diff --git a/Engine/Source/Runtime/Launch/Private/Android/AndroidJNI.cpp b/Engine/Source/Runtime/Launch/Private/Android/AndroidJNI.cpp
--- a/Engine/Source/Runtime/Launch/Private/Android/AndroidJNI.cpp
+++ b/Engine/Source/Runtime/Launch/Private/Android/AndroidJNI.cpp
@@ -87,7 +87,7 @@ void FJavaWrapper::FindClassesAndMethods(JNIEnv* Env)
     AndroidThunkJava_ClipboardPaste = FindMethod(Env, GameActivityClassID, "AndroidThunkJava_ClipboardPaste", "()Ljava/lang/String;", bIsOptional);
 	AndroidThunkJava_ForceQuit = FindMethod(Env, GameActivityClassID, "AndroidThunkJava_ForceQuit", "()V", bIsOptional);
 	AndroidThunkJava_GetFontDirectory = FindStaticMethod(Env, GameActivityClassID, "AndroidThunkJava_GetFontDirectory", "()Ljava/lang/String;", bIsOptional);
-	AndroidThunkJava_Vibrate = FindMethod(Env, GameActivityClassID, "AndroidThunkJava_Vibrate", "(I)V", bIsOptional);
+	AndroidThunkJava_Vibrate = FindMethod(Env, GameActivityClassID, "AndroidThunkJava_Vibrate", "(II)V", bIsOptional);
 	AndroidThunkJava_IsMusicActive = FindMethod(Env, GameActivityClassID, "AndroidThunkJava_IsMusicActive", "()Z", bIsOptional);
 	AndroidThunkJava_IsScreensaverEnabled = FindMethod(Env, GameActivityClassID, "AndroidThunkJava_IsScreensaverEnabled", "()Z", bIsOptional);
 	AndroidThunkJava_KeepScreenOn = FindMethod(Env, GameActivityClassID, "AndroidThunkJava_KeepScreenOn", "(Z)V", bIsOptional);
@@ -96,6 +96,7 @@ void FJavaWrapper::FindClassesAndMethods(JNIEnv* Env)
 	AndroidThunkJava_ShowProgressDialog = FindMethod(Env, GameActivityClassID, "AndroidThunkJava_ShowProgressDialog", "(ZLjava/lang/String;ZI)V", bIsOptional);
 	AndroidThunkJava_UpdateProgressDialog = FindMethod(Env, GameActivityClassID, "AndroidThunkJava_UpdateProgressDialog", "(I)V", bIsOptional);
 	AndroidThunkJava_GetInputDeviceInfo = FindMethod(Env, GameActivityClassID, "AndroidThunkJava_GetInputDeviceInfo", "(I)Lcom/epicgames/unreal/GameActivity$InputDeviceInfo;", bIsOptional);
+	AndroidThunkJava_SetInputDeviceVibrators = FindMethod(Env, GameActivityClassID, "AndroidThunkJava_SetInputDeviceVibrators", "(IIIII)Z", bIsOptional);
 	AndroidThunkJava_IsGamepadAttached = FindMethod(Env, GameActivityClassID, "AndroidThunkJava_IsGamepadAttached", "()Z", bIsOptional);
 	AndroidThunkJava_HasMetaDataKey = FindMethod(Env, GameActivityClassID, "AndroidThunkJava_HasMetaDataKey", "(Ljava/lang/String;)Z", bIsOptional);
 	AndroidThunkJava_GetMetaDataBoolean = FindMethod(Env, GameActivityClassID, "AndroidThunkJava_GetMetaDataBoolean", "(Ljava/lang/String;)Z", bIsOptional);
@@ -143,6 +144,7 @@ void FJavaWrapper::FindClassesAndMethods(JNIEnv* Env)
 	InputDeviceInfo_ControllerId = FJavaWrapper::FindField(Env, InputDeviceInfoClass, "controllerId", "I", bIsOptional);
 	InputDeviceInfo_Name = FJavaWrapper::FindField(Env, InputDeviceInfoClass, "name", "Ljava/lang/String;", bIsOptional);
 	InputDeviceInfo_Descriptor = FJavaWrapper::FindField(Env, InputDeviceInfoClass, "descriptor", "Ljava/lang/String;", bIsOptional);
+	InputDeviceInfo_FeedbackMotorCount = FJavaWrapper::FindField(Env, InputDeviceInfoClass, "feedbackMotorCount", "I", bIsOptional);
 
 	/** GooglePlay services */
 	FindGooglePlayMethods(Env);
@@ -386,6 +388,7 @@ jmethodID FJavaWrapper::AndroidThunkJava_DismissSplashScreen;
 jmethodID FJavaWrapper::AndroidThunkJava_ShowProgressDialog;
 jmethodID FJavaWrapper::AndroidThunkJava_UpdateProgressDialog;
 jmethodID FJavaWrapper::AndroidThunkJava_GetInputDeviceInfo;
+jmethodID FJavaWrapper::AndroidThunkJava_SetInputDeviceVibrators;
 jmethodID FJavaWrapper::AndroidThunkJava_IsGamepadAttached;
 jmethodID FJavaWrapper::AndroidThunkJava_HasMetaDataKey;
 jmethodID FJavaWrapper::AndroidThunkJava_GetMetaDataBoolean;
@@ -425,6 +428,7 @@ jfieldID FJavaWrapper::InputDeviceInfo_ProductId;
 jfieldID FJavaWrapper::InputDeviceInfo_ControllerId;
 jfieldID FJavaWrapper::InputDeviceInfo_Name;
 jfieldID FJavaWrapper::InputDeviceInfo_Descriptor;
+jfieldID FJavaWrapper::InputDeviceInfo_FeedbackMotorCount;
 
 jclass FJavaWrapper::GoogleServicesClassID;
 jobject FJavaWrapper::GoogleServicesThis;
@@ -517,12 +521,12 @@ void AndroidThunkCpp_KeepScreenOn(bool Enable)
 	}
 }
 
-void AndroidThunkCpp_Vibrate(int32 Duration)
+void AndroidThunkCpp_Vibrate(int32 Intensity, int32 Duration)
 {
 	if (JNIEnv* Env = FAndroidApplication::GetJavaEnv())
 	{
 		// call the java side
-		FJavaWrapper::CallVoidMethod(Env, FJavaWrapper::GameActivityThis, FJavaWrapper::AndroidThunkJava_Vibrate, Duration);
+		FJavaWrapper::CallVoidMethod(Env, FJavaWrapper::GameActivityThis, FJavaWrapper::AndroidThunkJava_Vibrate, Intensity, Duration);
 	}
 }
 
@@ -580,7 +584,9 @@ bool AndroidThunkCpp_GetInputDeviceInfo(int32 deviceId, FAndroidInputDeviceInfo
 
 			results.Name = FJavaHelper::FStringFromLocalRef(Env, (jstring)Env->GetObjectField(*deviceInfo, FJavaWrapper::InputDeviceInfo_Name));
 			results.Descriptor = FJavaHelper::FStringFromLocalRef(Env, (jstring)Env->GetObjectField(*deviceInfo, FJavaWrapper::InputDeviceInfo_Descriptor));
-			
+
+			results.FeedbackMotorCount = (int32)Env->GetIntField(*deviceInfo, FJavaWrapper::InputDeviceInfo_FeedbackMotorCount);
+
 			return true;
 		}
 	}
@@ -592,6 +598,16 @@ bool AndroidThunkCpp_GetInputDeviceInfo(int32 deviceId, FAndroidInputDeviceInfo
 	results.ControllerId = -1;
 	results.Name = FString("Unknown");
 	results.Descriptor = FString("Unknown");
+	results.FeedbackMotorCount = 0;
+	return false;
+}
+
+bool AndroidThunkCpp_SetInputDeviceVibrators(int32 deviceId, int32 leftIntensity, int32 leftDuration, int32 rightIntensity, int32 rightDuration)
+{
+	if (JNIEnv* Env = FAndroidApplication::GetJavaEnv())
+	{
+		return FJavaWrapper::CallBooleanMethod(Env, FJavaWrapper::GameActivityThis, FJavaWrapper::AndroidThunkJava_SetInputDeviceVibrators, deviceId, leftIntensity, leftDuration, rightIntensity, rightDuration);
+	}
 	return false;
 }
 
diff --git a/Engine/Source/Runtime/Launch/Public/Android/AndroidJNI.h b/Engine/Source/Runtime/Launch/Public/Android/AndroidJNI.h
--- a/Engine/Source/Runtime/Launch/Public/Android/AndroidJNI.h
+++ b/Engine/Source/Runtime/Launch/Public/Android/AndroidJNI.h
@@ -45,6 +45,7 @@ public:
 	static jmethodID AndroidThunkJava_ShowProgressDialog;
 	static jmethodID AndroidThunkJava_UpdateProgressDialog;
 	static jmethodID AndroidThunkJava_GetInputDeviceInfo;
+	static jmethodID AndroidThunkJava_SetInputDeviceVibrators;
 	static jmethodID AndroidThunkJava_IsGamepadAttached;
 	static jmethodID AndroidThunkJava_HasMetaDataKey;
 	static jmethodID AndroidThunkJava_GetMetaDataBoolean;
@@ -88,6 +89,7 @@ public:
 	static jfieldID InputDeviceInfo_ControllerId;
 	static jfieldID InputDeviceInfo_Name;
 	static jfieldID InputDeviceInfo_Descriptor;
+	static jfieldID InputDeviceInfo_FeedbackMotorCount;
 
 	// IDs related to google play services
 	static jclass GoogleServicesClassID;
