[Virtual Keyboard] Allow the use of a virtual keyboard on any platforms
#rnx

diff --git a/Engine/Source/Runtime/ApplicationCore/Private/GenericPlatform/GenericPlatformApplicationMisc.cpp b/Engine/Source/Runtime/ApplicationCore/Private/GenericPlatform/GenericPlatformApplicationMisc.cpp
--- a/Engine/Source/Runtime/ApplicationCore/Private/GenericPlatform/GenericPlatformApplicationMisc.cpp
+++ b/Engine/Source/Runtime/ApplicationCore/Private/GenericPlatform/GenericPlatformApplicationMisc.cpp
@@ -27,6 +27,15 @@ FAutoConsoleVariableRef FGenericPlatformApplicationMisc::CVarEnableHighDPIAwaren
 	ECVF_ReadOnly
 );
 
+static bool bAllowVirtualKeyboard  = false;
+
+FAutoConsoleVariableRef FGenericPlatformApplicationMisc::CVarAllowVirtualKeyboard(
+	TEXT("AllowVirtualKeyboard"),
+	bAllowVirtualKeyboard,
+	TEXT("Allow the use of a virtual keyboard despite platform main screen being non-touch"),
+	ECVF_ReadOnly
+);
+
 void FGenericPlatformApplicationMisc::PreInit()
 {
 }
@@ -82,6 +91,11 @@ bool FGenericPlatformApplicationMisc::IsThisApplicationForeground()
 	return false;
 }
 
+bool FGenericPlatformApplicationMisc::RequiresVirtualKeyboard()
+{
+	return PLATFORM_HAS_TOUCH_MAIN_SCREEN || bAllowVirtualKeyboard;
+}
+
 FLinearColor FGenericPlatformApplicationMisc::GetScreenPixelColor(const FVector2D& InScreenPos, float InGamma)
 { 
 	return FLinearColor::Black;
diff --git a/Engine/Source/Runtime/ApplicationCore/Public/GenericPlatform/GenericPlatformApplicationMisc.h b/Engine/Source/Runtime/ApplicationCore/Public/GenericPlatform/GenericPlatformApplicationMisc.h
--- a/Engine/Source/Runtime/ApplicationCore/Public/GenericPlatform/GenericPlatformApplicationMisc.h
+++ b/Engine/Source/Runtime/ApplicationCore/Public/GenericPlatform/GenericPlatformApplicationMisc.h
@@ -124,10 +124,7 @@ struct APPLICATIONCORE_API FGenericPlatformApplicationMisc
 	/**
 	* Returns whether the platform wants to use a touch screen for a virtual keyboard.
 	*/
-	static bool RequiresVirtualKeyboard()
-	{
-		return PLATFORM_HAS_TOUCH_MAIN_SCREEN;
-	}
+	static bool RequiresVirtualKeyboard();
 
 	/**
 	 *	Pumps Windows messages.
@@ -389,5 +386,6 @@ protected:
 	static EScreenPhysicalAccuracy CachedPhysicalScreenAccuracy;
 	static int32 CachedPhysicalScreenDensity;
 	static FAutoConsoleVariableRef CVarEnableHighDPIAwareness;
+	static FAutoConsoleVariableRef CVarAllowVirtualKeyboard;
 
 };
