Update WindowsPlatformSplash to allow setting custom text before calling Show()

diff --git a/Engine/Source/Runtime/ApplicationCore/Private/Windows/WindowsPlatformSplash.cpp b/Engine/Source/Runtime/ApplicationCore/Private/Windows/WindowsPlatformSplash.cpp
--- a/Engine/Source/Runtime/ApplicationCore/Private/Windows/WindowsPlatformSplash.cpp
+++ b/Engine/Source/Runtime/ApplicationCore/Private/Windows/WindowsPlatformSplash.cpp
@@ -606,8 +606,12 @@ uint32 WINAPI StartSplashScreenThread( LPVOID unused )
  */
 static void StartSetSplashText( const SplashTextType::Type InType, const TCHAR* InText )
 {
-	// Only allow copyright text displayed while loading the game.  Editor displays all.
-	GSplashScreenText[InType] = FText::FromString(InText);
+	// If we've already been set don't init the string with the default
+	if (GSplashScreenText[InType].IsEmpty())
+	{
+		// Only allow copyright text displayed while loading the game.  Editor displays all.
+		GSplashScreenText[InType] = FText::FromString(InText);
+	}
 }
 
 void FWindowsPlatformSplash::Show()
@@ -780,6 +784,14 @@ void FWindowsPlatformSplash::SetSplashText( const SplashTextType::Type InType, c
 			}
 		}
 	}
+	else // We haven't started drawing yet we can set text ahead of time to show
+	{
+		// Update splash text
+		if (FCString::Strcmp(InText, *GSplashScreenText[InType].ToString()) != 0)
+		{
+			GSplashScreenText[InType] = FText::FromString(InText);
+		}
+	}
 }
 
 APPLICATIONCORE_API HWND GetSplashScreenWindowHandle()
