FExec: Converted Engine Exec overrides to Exec_Runtime/Exec_Dev where:
- No additional routing was being done in Exec.
- No default parameter was being provided through the Exec override.

This is so we can turn off Exec overrides in certain configurations.

diff --git a/Engine/Source/Runtime/ApplicationCore/Private/IOS/IOSInputInterface.cpp b/Engine/Source/Runtime/ApplicationCore/Private/IOS/IOSInputInterface.cpp
--- a/Engine/Source/Runtime/ApplicationCore/Private/IOS/IOSInputInterface.cpp
+++ b/Engine/Source/Runtime/ApplicationCore/Private/IOS/IOSInputInterface.cpp
@@ -706,7 +706,7 @@ void FIOSInputInterface::CalibrateMotion(uint32 PlayerIndex)
     }
 }
 
-bool FIOSInputInterface::Exec(UWorld* InWorld, const TCHAR* Cmd, FOutputDevice& Ar)
+bool FIOSInputInterface::Exec_Runtime(UWorld* InWorld, const TCHAR* Cmd, FOutputDevice& Ar)
 {
     // Keep track whether the command was handled or not.
     bool bHandledCommand = false;
diff --git a/Engine/Source/Runtime/ApplicationCore/Public/IOS/IOSInputInterface.h b/Engine/Source/Runtime/ApplicationCore/Public/IOS/IOSInputInterface.h
--- a/Engine/Source/Runtime/ApplicationCore/Public/IOS/IOSInputInterface.h
+++ b/Engine/Source/Runtime/ApplicationCore/Public/IOS/IOSInputInterface.h
@@ -88,10 +88,6 @@ public:
 	static APPLICATIONCORE_API void QueueTouchInput(const TArray<TouchInput>& InTouchEvents);
 	static void QueueKeyInput(int32 Key, int32 Char);
 
-	//~ Begin Exec Interface
-	virtual bool Exec(UWorld* InWorld, const TCHAR* Cmd, FOutputDevice& Ar) override;
-	//~ End Exec Interface
-
 	void SetGamepadsAllowed(bool bAllowed) { bAllowControllers = bAllowed; }
 	void SetGamepadsBlockDeviceFeedback(bool bBlock) { bControllersBlockDeviceFeedback = bBlock; }
 
@@ -102,7 +98,13 @@ public:
     static bool IsKeyboardInhibited() { return bKeyboardInhibited; }
     
     NSData* GetGamepadGlyphRawData(const FGamepadKeyNames::Type& ButtonKey, uint32 ControllerIndex);
-    
+
+protected:
+
+	//~ Begin Exec Interface
+	virtual bool Exec_Runtime(UWorld* InWorld, const TCHAR* Cmd, FOutputDevice& Ar) override;
+	//~ End Exec Interface
+
 private:
 
 	FIOSInputInterface( const TSharedRef< FGenericApplicationMessageHandler >& InMessageHandler );
diff --git a/Engine/Source/Runtime/Core/Private/Logging/LogSuppressionInterface.cpp b/Engine/Source/Runtime/Core/Private/Logging/LogSuppressionInterface.cpp
--- a/Engine/Source/Runtime/Core/Private/Logging/LogSuppressionInterface.cpp
+++ b/Engine/Source/Runtime/Core/Private/Logging/LogSuppressionInterface.cpp
@@ -558,7 +558,7 @@ class FLogSuppressionImplementation: public FLogSuppressionInterface, private FS
 	}
 
 	/** Console commands, see embeded usage statement **/
-	virtual bool Exec( UWorld* Inworld, const TCHAR* Cmd, FOutputDevice& Ar )
+	virtual bool Exec_Runtime( UWorld* Inworld, const TCHAR* Cmd, FOutputDevice& Ar )
 	{
 		if(FParse::Command(&Cmd,TEXT("LOG")))
 		{
diff --git a/Engine/Source/Runtime/Core/Private/Stats/StatsCommand.cpp b/Engine/Source/Runtime/Core/Private/Stats/StatsCommand.cpp
--- a/Engine/Source/Runtime/Core/Private/Stats/StatsCommand.cpp
+++ b/Engine/Source/Runtime/Core/Private/Stats/StatsCommand.cpp
@@ -2188,9 +2188,9 @@ static void StatCmd(FString InCmd, bool bStatCommand, FOutputDevice* Ar /*= null
 /** Exec used to execute core stats commands on the stats thread. */
 static class FStatCmdCore : private FSelfRegisteringExec
 {
-public:
+protected:
 	/** Console commands, see embeded usage statement **/
-	virtual bool Exec( UWorld*, const TCHAR* Cmd, FOutputDevice& Ar ) override
+	virtual bool Exec_Runtime( UWorld*, const TCHAR* Cmd, FOutputDevice& Ar ) override
 	{
 		// Block the thread as this affects external stat states now
 		return DirectStatsCommand(Cmd,true,&Ar);
diff --git a/Engine/Source/Runtime/CoreUObject/Private/UObject/ScriptCore.cpp b/Engine/Source/Runtime/CoreUObject/Private/UObject/ScriptCore.cpp
--- a/Engine/Source/Runtime/CoreUObject/Private/UObject/ScriptCore.cpp
+++ b/Engine/Source/Runtime/CoreUObject/Private/UObject/ScriptCore.cpp
@@ -1891,11 +1891,12 @@ static void OutputTotalBytecodeSize(FOutputDevice& Ar)
 struct FScriptAuditExec 
 	: public FSelfRegisteringExec
 {
+protected:
 	// FSelfRegisteringExec:
-	virtual bool Exec(UWorld* InWorld, const TCHAR* Cmd, FOutputDevice& Ar) override;
+	virtual bool Exec_Runtime(UWorld* InWorld, const TCHAR* Cmd, FOutputDevice& Ar) override;
 } ScriptAudit;
 
-bool FScriptAuditExec::Exec(UWorld* InWorld, const TCHAR* Cmd, FOutputDevice& Ar)
+bool FScriptAuditExec::Exec_Runtime(UWorld* InWorld, const TCHAR* Cmd, FOutputDevice& Ar)
 {
 	if (FParse::Command(&Cmd, TEXT("ScriptAudit")))
 	{
diff --git a/Engine/Source/Runtime/Engine/Private/DeviceProfiles/DeviceProfileManager.cpp b/Engine/Source/Runtime/Engine/Private/DeviceProfiles/DeviceProfileManager.cpp
--- a/Engine/Source/Runtime/Engine/Private/DeviceProfiles/DeviceProfileManager.cpp
+++ b/Engine/Source/Runtime/Engine/Private/DeviceProfiles/DeviceProfileManager.cpp
@@ -1400,10 +1400,10 @@ static bool GetCVarForDeviceProfile( FOutputDevice& Ar, FString DPName, FString
 
 class FPlatformCVarExec : public FSelfRegisteringExec
 {
-public:
+protected:
 
 	// FSelfRegisteringExec interface
-	virtual bool Exec(UWorld* InWorld, const TCHAR* Cmd, FOutputDevice& Ar) override
+	virtual bool Exec_Runtime(UWorld* InWorld, const TCHAR* Cmd, FOutputDevice& Ar) override
 	{
 		if (FParse::Command(&Cmd, TEXT("dpcvar")))
 		{
diff --git a/Engine/Source/Runtime/Engine/Private/GPUSkinCache.cpp b/Engine/Source/Runtime/Engine/Private/GPUSkinCache.cpp
--- a/Engine/Source/Runtime/Engine/Private/GPUSkinCache.cpp
+++ b/Engine/Source/Runtime/Engine/Private/GPUSkinCache.cpp
@@ -37,7 +37,7 @@ DEFINE_LOG_CATEGORY_STATIC(LogSkinCache, Log, All);
 /** Exec helper to handle GPU Skin Cache related commands. */
 class FSkinCacheExecHelper : public FSelfRegisteringExec
 {
-	virtual bool Exec(class UWorld* InWorld, const TCHAR* Cmd, FOutputDevice& Ar)
+	virtual bool Exec_Runtime(class UWorld* InWorld, const TCHAR* Cmd, FOutputDevice& Ar) override
 	{
 		/** Command to list all skeletal mesh lods which have the skin cache disabled. */
 		if (FParse::Command(&Cmd, TEXT("list skincacheusage")))
diff --git a/Engine/Source/Runtime/Engine/Private/InstancedStaticMesh.cpp b/Engine/Source/Runtime/Engine/Private/InstancedStaticMesh.cpp
--- a/Engine/Source/Runtime/Engine/Private/InstancedStaticMesh.cpp
+++ b/Engine/Source/Runtime/Engine/Private/InstancedStaticMesh.cpp
@@ -143,7 +143,7 @@ static TAutoConsoleVariable<int32> CVarEnableViewportSMInstanceSelection(
 
 class FISMExecHelper : public FSelfRegisteringExec
 {
-	virtual bool Exec(class UWorld* InWorld, const TCHAR* Cmd, FOutputDevice& Ar)
+	virtual bool Exec_Runtime(class UWorld* InWorld, const TCHAR* Cmd, FOutputDevice& Ar)
 	{
 		if (FParse::Command(&Cmd, TEXT("LIST ISM")))
 		{
diff --git a/Engine/Source/Runtime/Engine/Private/StatsRender2.cpp b/Engine/Source/Runtime/Engine/Private/StatsRender2.cpp
--- a/Engine/Source/Runtime/Engine/Private/StatsRender2.cpp
+++ b/Engine/Source/Runtime/Engine/Private/StatsRender2.cpp
@@ -294,9 +294,9 @@ static FString ShortenName( TCHAR const* LongName )
 /** Exec used to execute engine stats command on the game thread. */
 static class FStatCmdEngine : private FSelfRegisteringExec
 {
-public:
+protected:
 	/** Console commands. */
-	virtual bool Exec( UWorld*, const TCHAR* Cmd, FOutputDevice& Ar ) override
+	virtual bool Exec_Runtime( UWorld*, const TCHAR* Cmd, FOutputDevice& Ar ) override
 	{
 		if( FParse::Command( &Cmd, TEXT( "stat" ) ) )
 		{
diff --git a/Engine/Source/Runtime/LevelSequence/Private/LevelSequenceModule.cpp b/Engine/Source/Runtime/LevelSequence/Private/LevelSequenceModule.cpp
--- a/Engine/Source/Runtime/LevelSequence/Private/LevelSequenceModule.cpp
+++ b/Engine/Source/Runtime/LevelSequence/Private/LevelSequenceModule.cpp
@@ -253,7 +253,7 @@ bool HandleLevelSequenceSetClockSource(UWorld* InWorld, const TCHAR** InStr, FOu
 
 } // empty namespace
 
-bool FLevelSequenceModule::Exec(class UWorld* InWorld, const TCHAR* Cmd, FOutputDevice& Ar)
+bool FLevelSequenceModule::Exec_Runtime(class UWorld* InWorld, const TCHAR* Cmd, FOutputDevice& Ar)
 {
 	if (FParse::Command(&Cmd, TEXT("Sequencer")))
 	{
diff --git a/Engine/Source/Runtime/LevelSequence/Private/LevelSequenceModule.h b/Engine/Source/Runtime/LevelSequence/Private/LevelSequenceModule.h
--- a/Engine/Source/Runtime/LevelSequence/Private/LevelSequenceModule.h
+++ b/Engine/Source/Runtime/LevelSequence/Private/LevelSequenceModule.h
@@ -18,14 +18,15 @@ public:
 	virtual void StartupModule() override;
 	virtual void ShutdownModule() override;
 
-	// FSelfRegisteringExec interface
-	virtual bool Exec(class UWorld* InWorld, const TCHAR* Cmd, FOutputDevice& Ar) override;
-	
 	// ILevelSequenceModule interface
 	virtual FDelegateHandle RegisterObjectSpawner(FOnCreateMovieSceneObjectSpawner InOnCreateMovieSceneObjectSpawner) override;
 	virtual void GenerateObjectSpawners(TArray<TSharedRef<IMovieSceneObjectSpawner>>& OutSpawners) const override;
 	virtual void UnregisterObjectSpawner(FDelegateHandle InHandle) override;
 
+protected:
+	// FSelfRegisteringExec interface
+	virtual bool Exec_Runtime(class UWorld* InWorld, const TCHAR* Cmd, FOutputDevice& Ar) override;
+
 public:
 	/** List of object spawner delegates used to extend the spawn register */
 	TArray< FOnCreateMovieSceneObjectSpawner > OnCreateMovieSceneObjectSpawnerDelegates;
diff --git a/Engine/Source/Runtime/MediaAssets/Private/MediaAssetsModule.cpp b/Engine/Source/Runtime/MediaAssets/Private/MediaAssetsModule.cpp
--- a/Engine/Source/Runtime/MediaAssets/Private/MediaAssetsModule.cpp
+++ b/Engine/Source/Runtime/MediaAssets/Private/MediaAssetsModule.cpp
@@ -24,11 +24,11 @@ class FMediaAssetsModule
 	: public FSelfRegisteringExec
 	, public IMediaAssetsModule
 {
-public:
+protected:
 
 	//~ FSelfRegisteringExec interface
 
-	virtual bool Exec(UWorld* InWorld, const TCHAR* Cmd, FOutputDevice& Ar) override
+	virtual bool Exec_Runtime(UWorld* InWorld, const TCHAR* Cmd, FOutputDevice& Ar) override
 	{
 		if (FParse::Command(&Cmd, TEXT("MEDIA")))
 		{
diff --git a/Engine/Source/Runtime/Messaging/Private/MessagingModule.cpp b/Engine/Source/Runtime/Messaging/Private/MessagingModule.cpp
--- a/Engine/Source/Runtime/Messaging/Private/MessagingModule.cpp
+++ b/Engine/Source/Runtime/Messaging/Private/MessagingModule.cpp
@@ -26,11 +26,11 @@ class FMessagingModule
 	: public FSelfRegisteringExec
 	, public IMessagingModule
 {
-public:
+protected:
 
 	//~ FSelfRegisteringExec interface
 
-	virtual bool Exec(UWorld* InWorld, const TCHAR* Cmd, FOutputDevice& Ar) override
+	virtual bool Exec_Runtime(UWorld* InWorld, const TCHAR* Cmd, FOutputDevice& Ar) override
 	{
 		if (!FParse::Command(&Cmd, TEXT("MESSAGING")))
 		{
diff --git a/Engine/Source/Runtime/NavigationSystem/Private/NavigationSystem.cpp b/Engine/Source/Runtime/NavigationSystem/Private/NavigationSystem.cpp
--- a/Engine/Source/Runtime/NavigationSystem/Private/NavigationSystem.cpp
+++ b/Engine/Source/Runtime/NavigationSystem/Private/NavigationSystem.cpp
@@ -4542,7 +4542,7 @@ bool UNavigationSystemV1::HandleCountNavMemCommand()
 //----------------------------------------------------------------------//
 // Commands
 //----------------------------------------------------------------------//
-bool FNavigationSystemExec::Exec(UWorld* InWorld, const TCHAR* Cmd, FOutputDevice& Ar)
+bool FNavigationSystemExec::Exec_Runtime(UWorld* InWorld, const TCHAR* Cmd, FOutputDevice& Ar)
 {
 	UNavigationSystemV1*  NavSys = FNavigationSystem::GetCurrent<UNavigationSystemV1>(InWorld);
 
diff --git a/Engine/Source/Runtime/NavigationSystem/Public/NavigationSystem.h b/Engine/Source/Runtime/NavigationSystem/Public/NavigationSystem.h
--- a/Engine/Source/Runtime/NavigationSystem/Public/NavigationSystem.h
+++ b/Engine/Source/Runtime/NavigationSystem/Public/NavigationSystem.h
@@ -93,8 +93,9 @@ namespace FNavigationSystem
 
 struct FNavigationSystemExec: public FSelfRegisteringExec
 {
+protected:
 	//~ Begin FExec Interface
-	virtual bool Exec(UWorld* Inworld, const TCHAR* Cmd, FOutputDevice& Ar) override;
+	virtual bool Exec_Runtime(UWorld* Inworld, const TCHAR* Cmd, FOutputDevice& Ar) override;
 	//~ End FExec Interface
 };
 
diff --git a/Engine/Source/Runtime/NetworkFile/Private/NetworkPlatformFile.cpp b/Engine/Source/Runtime/NetworkFile/Private/NetworkPlatformFile.cpp
--- a/Engine/Source/Runtime/NetworkFile/Private/NetworkPlatformFile.cpp
+++ b/Engine/Source/Runtime/NetworkFile/Private/NetworkPlatformFile.cpp
@@ -1528,7 +1528,7 @@ void FNetworkPlatformFile::Tick()
 	}
 }
 
-bool FNetworkPlatformFile::Exec(class UWorld* InWorld, const TCHAR* Cmd, FOutputDevice& Ar )
+bool FNetworkPlatformFile::Exec_Runtime(class UWorld* InWorld, const TCHAR* Cmd, FOutputDevice& Ar )
 {
 	if (FParse::Command(&Cmd, TEXT("networkfile")))
 	{
diff --git a/Engine/Source/Runtime/NetworkFile/Public/NetworkPlatformFile.h b/Engine/Source/Runtime/NetworkFile/Public/NetworkPlatformFile.h
--- a/Engine/Source/Runtime/NetworkFile/Public/NetworkPlatformFile.h
+++ b/Engine/Source/Runtime/NetworkFile/Public/NetworkPlatformFile.h
@@ -176,15 +176,13 @@ public:
 
 	virtual FString GetVersionInfo() const;
 
-
+protected:
 
 	//////////////////////////////////////////////////////////////////////////
 	// FSelfRegisteringExec interface
-	virtual bool Exec(class UWorld* InWorld, const TCHAR* Cmd, FOutputDevice& Ar ) override;
+	virtual bool Exec_Runtime(class UWorld* InWorld, const TCHAR* Cmd, FOutputDevice& Ar ) override;
 
 
-protected:
-
 	/**
 	 * Send a heartbeat message to the file server. This will tell it we are alive, as well as
 	 * get back a list of files that have been updated on the server (so we can toss our copy)
diff --git a/Engine/Source/Runtime/NetworkFileSystem/Private/NetworkFileServerConnection.cpp b/Engine/Source/Runtime/NetworkFileSystem/Private/NetworkFileServerConnection.cpp
--- a/Engine/Source/Runtime/NetworkFileSystem/Private/NetworkFileServerConnection.cpp
+++ b/Engine/Source/Runtime/NetworkFileSystem/Private/NetworkFileServerConnection.cpp
@@ -1385,7 +1385,7 @@ FString FNetworkFileServerClientConnection::GetDescription() const
 }
 
 
-bool FNetworkFileServerClientConnection::Exec(class UWorld* InWorld, const TCHAR* Cmd, FOutputDevice& Ar ) 
+bool FNetworkFileServerClientConnection::Exec_Runtime(class UWorld* InWorld, const TCHAR* Cmd, FOutputDevice& Ar ) 
 {
 	if (FParse::Command(&Cmd, TEXT("networkserverconnection")))
 	{
diff --git a/Engine/Source/Runtime/NetworkFileSystem/Private/NetworkFileServerConnection.h b/Engine/Source/Runtime/NetworkFileSystem/Private/NetworkFileServerConnection.h
--- a/Engine/Source/Runtime/NetworkFileSystem/Private/NetworkFileServerConnection.h
+++ b/Engine/Source/Runtime/NetworkFileSystem/Private/NetworkFileServerConnection.h
@@ -165,7 +165,7 @@ protected:
 	void FileModifiedCallback( const FString& Filename );
 	
 	
-	virtual bool Exec(class UWorld* InWorld, const TCHAR* Cmd, FOutputDevice& Ar) override;
+	virtual bool Exec_Runtime(class UWorld* InWorld, const TCHAR* Cmd, FOutputDevice& Ar) override;
 
 
 	/**
diff --git a/Engine/Source/Runtime/Online/HTTP/Private/HttpModule.cpp b/Engine/Source/Runtime/Online/HTTP/Private/HttpModule.cpp
--- a/Engine/Source/Runtime/Online/HTTP/Private/HttpModule.cpp
+++ b/Engine/Source/Runtime/Online/HTTP/Private/HttpModule.cpp
@@ -217,7 +217,7 @@ bool FHttpModule::HandleHTTPCommand(const TCHAR* Cmd, FOutputDevice& Ar)
 	return true;
 }
 
-bool FHttpModule::Exec(UWorld* InWorld, const TCHAR* Cmd, FOutputDevice& Ar)
+bool FHttpModule::Exec_Runtime(UWorld* InWorld, const TCHAR* Cmd, FOutputDevice& Ar)
 {
 	// Ignore any execs that don't start with HTTP
 	if (FParse::Command(&Cmd, TEXT("HTTP")))
diff --git a/Engine/Source/Runtime/Online/HTTP/Public/HttpModule.h b/Engine/Source/Runtime/Online/HTTP/Public/HttpModule.h
--- a/Engine/Source/Runtime/Online/HTTP/Public/HttpModule.h
+++ b/Engine/Source/Runtime/Online/HTTP/Public/HttpModule.h
@@ -26,7 +26,7 @@ class FHttpModule :
 	public IModuleInterface, public FSelfRegisteringExec
 {
 
-public:
+protected:
 
 	// FSelfRegisteringExec
 
@@ -39,7 +39,9 @@ public:
 	 *
 	 * @return true if the handler consumed the input, false to continue searching handlers
 	 */
-	virtual bool Exec(UWorld* InWorld, const TCHAR* Cmd, FOutputDevice& Ar) override;
+	virtual bool Exec_Runtime(UWorld* InWorld, const TCHAR* Cmd, FOutputDevice& Ar) override;
+
+public:
 
 	/** 
 	 * Exec command handlers
diff --git a/Engine/Source/Runtime/Online/XMPP/Private/XmppModule.cpp b/Engine/Source/Runtime/Online/XMPP/Private/XmppModule.cpp
--- a/Engine/Source/Runtime/Online/XMPP/Private/XmppModule.cpp
+++ b/Engine/Source/Runtime/Online/XMPP/Private/XmppModule.cpp
@@ -663,7 +663,7 @@ bool FXmppModule::HandleXmppCommand( const TCHAR* Cmd, FOutputDevice& Ar )
 	return false;
 }
 
-bool FXmppModule::Exec(UWorld* InWorld, const TCHAR* Cmd, FOutputDevice& Ar)
+bool FXmppModule::Exec_Runtime(UWorld* InWorld, const TCHAR* Cmd, FOutputDevice& Ar)
 {
 	// Ignore any execs that don't start with Xmpp
 	if (FParse::Command(&Cmd, TEXT("XMPP")))
diff --git a/Engine/Source/Runtime/Online/XMPP/Public/XmppModule.h b/Engine/Source/Runtime/Online/XMPP/Public/XmppModule.h
--- a/Engine/Source/Runtime/Online/XMPP/Public/XmppModule.h
+++ b/Engine/Source/Runtime/Online/XMPP/Public/XmppModule.h
@@ -29,7 +29,7 @@ class XMPP_API FXmppModule :
 	public IModuleInterface, public FSelfRegisteringExec, public FTSTickerObjectBase
 {
 
-public:
+protected:
 
 	// FSelfRegisteringExec
 
@@ -42,7 +42,9 @@ public:
 	 *
 	 * @return true if the handler consumed the input, false to continue searching handlers
 	 */
-	virtual bool Exec(UWorld* InWorld, const TCHAR* Cmd, FOutputDevice& Ar) override;
+	virtual bool Exec_Runtime(UWorld* InWorld, const TCHAR* Cmd, FOutputDevice& Ar) override;
+
+public:
 
 	/** 
 	 * Exec command handlers
diff --git a/Engine/Source/Runtime/PerfCounters/Private/PerfCounters.cpp b/Engine/Source/Runtime/PerfCounters/Private/PerfCounters.cpp
--- a/Engine/Source/Runtime/PerfCounters/Private/PerfCounters.cpp
+++ b/Engine/Source/Runtime/PerfCounters/Private/PerfCounters.cpp
@@ -216,7 +216,7 @@ void FPerfCounters::TickSystemCounters(float DeltaTime)
 	}
 }
 
-bool FPerfCounters::Exec(UWorld* InWorld, const TCHAR* Cmd, FOutputDevice& Ar)
+bool FPerfCounters::Exec_Runtime(UWorld* InWorld, const TCHAR* Cmd, FOutputDevice& Ar)
 {
 	// ignore everything that doesn't start with PerfCounters
 	if (!FParse::Command(&Cmd, TEXT("perfcounters")))
diff --git a/Engine/Source/Runtime/PerfCounters/Private/PerfCounters.h b/Engine/Source/Runtime/PerfCounters/Private/PerfCounters.h
--- a/Engine/Source/Runtime/PerfCounters/Private/PerfCounters.h
+++ b/Engine/Source/Runtime/PerfCounters/Private/PerfCounters.h
@@ -36,10 +36,6 @@ public:
 	/** FTSTickerObjectBase */
 	virtual bool Tick(float DeltaTime) override;
 
-	//~ Begin Exec Interface
-	virtual bool Exec(UWorld* InWorld, const TCHAR* Cmd, FOutputDevice& Ar);
-	//~ End Exec Interface
-
 	//~ Begin IPerfCounters Interface
 	const FString& GetInstanceName() const override { return UniqueInstanceId; }
 	virtual double GetNumber(const FString& Name, double DefaultValue = 0.0) override;
@@ -59,6 +55,11 @@ public:
 	virtual bool ReportUnplayableCondition(const FString& ConditionDescription);
 	//~ Begin IPerfCounters Interface end
 
+protected:
+	//~ Begin Exec Interface
+	virtual bool Exec_Runtime(UWorld* InWorld, const TCHAR* Cmd, FOutputDevice& Ar);
+	//~ End Exec Interface
+
 private:
 	
 	void TickZeroLoad(float DeltaTime);
diff --git a/Engine/Source/Runtime/SessionServices/Private/SessionServicesModule.cpp b/Engine/Source/Runtime/SessionServices/Private/SessionServicesModule.cpp
--- a/Engine/Source/Runtime/SessionServices/Private/SessionServicesModule.cpp
+++ b/Engine/Source/Runtime/SessionServices/Private/SessionServicesModule.cpp
@@ -29,11 +29,11 @@ public:
 		, SessionService(nullptr)
 	{ }
 
-public:
+protected:
 
 	//~ FSelfRegisteringExec interface
 
-	virtual bool Exec(UWorld* InWorld, const TCHAR* Cmd, FOutputDevice& Ar) override
+	virtual bool Exec_Runtime(UWorld* InWorld, const TCHAR* Cmd, FOutputDevice& Ar) override
 	{
 		if (!FParse::Command(&Cmd, TEXT("SESSION")))
 		{
