UE will now properly re-raise notification windows to the front after maximizing, restoring, and moving windows.

diff --git a/Engine/Source/Runtime/ApplicationCore/Private/Linux/LinuxApplication.cpp b/Engine/Source/Runtime/ApplicationCore/Private/Linux/LinuxApplication.cpp
--- a/Engine/Source/Runtime/ApplicationCore/Private/Linux/LinuxApplication.cpp
+++ b/Engine/Source/Runtime/ApplicationCore/Private/Linux/LinuxApplication.cpp
@@ -804,6 +804,16 @@ void FLinuxApplication::ProcessDeferredMessage( SDL_Event Event )
 						ClientScreenY += BorderSizeY;
 
 						MessageHandler->OnMovedWindow(CurrentEventWindow.ToSharedRef(), ClientScreenX, ClientScreenY);
+
+						if (bFirstFrameOfWindowMove)
+						{
+							bFirstFrameOfWindowMove = false;
+
+							if (NotificationWindows.Num() > 0)
+							{
+								RaiseNotificationWindows(CurrentEventWindow);
+							}
+						}
 					}
 					break;
 
@@ -811,6 +821,11 @@ void FLinuxApplication::ProcessDeferredMessage( SDL_Event Event )
 					{
 						MessageHandler->OnWindowAction(CurrentEventWindow.ToSharedRef(), EWindowAction::Maximize);
 						UE_LOG(LogLinuxWindowEvent, Verbose, TEXT("Window: '%d' got maximized"), CurrentEventWindow->GetID());
+
+						if (NotificationWindows.Num() > 0)
+						{
+							RaiseNotificationWindows(CurrentEventWindow);
+						}
 					}
 					break;
 
@@ -818,6 +833,11 @@ void FLinuxApplication::ProcessDeferredMessage( SDL_Event Event )
 					{
 						MessageHandler->OnWindowAction(CurrentEventWindow.ToSharedRef(), EWindowAction::Restore);
 						UE_LOG(LogLinuxWindowEvent, Verbose, TEXT("Window: '%d' got restored"), CurrentEventWindow->GetID());
+
+						if (NotificationWindows.Num() > 0)
+						{
+							RaiseNotificationWindows(CurrentEventWindow);
+						}
 					}
 					break;
 
@@ -880,6 +900,8 @@ void FLinuxApplication::ProcessDeferredMessage( SDL_Event Event )
 						{
 							RaiseNotificationWindows(CurrentEventWindow);
 						}
+
+						bFirstFrameOfWindowMove = true;
 					}
 					break;
 
diff --git a/Engine/Source/Runtime/ApplicationCore/Public/Linux/LinuxApplication.h b/Engine/Source/Runtime/ApplicationCore/Public/Linux/LinuxApplication.h
--- a/Engine/Source/Runtime/ApplicationCore/Public/Linux/LinuxApplication.h
+++ b/Engine/Source/Runtime/ApplicationCore/Public/Linux/LinuxApplication.h
@@ -317,6 +317,9 @@ private:
 	/** TODO: describe */
 	bool bIsMouseCaptureEnabled;
 
+	/** True after every SDL_WINDOWEVENT_HIT_TEST until a following SDL_WINDOWEVENT_MOVED */
+	bool bFirstFrameOfWindowMove = false;
+
 	/** Window that we think has been activated last. */
 	TSharedPtr< FLinuxWindow > CurrentlyActiveWindow;
 	TSharedPtr< FLinuxWindow > CurrentFocusWindow;
