Cause notification windows to be raised above all windows except for modal windows, instead of just their parent windows.

This partially rolls back some of the changes made by CL 2532869 for UE-14543. With the addition of the modal check in RaiseNotifications(), this rollback doesn't cause any known issues

diff --git a/Engine/Source/Runtime/ApplicationCore/Private/Linux/LinuxApplication.cpp b/Engine/Source/Runtime/ApplicationCore/Private/Linux/LinuxApplication.cpp
--- a/Engine/Source/Runtime/ApplicationCore/Private/Linux/LinuxApplication.cpp
+++ b/Engine/Source/Runtime/ApplicationCore/Private/Linux/LinuxApplication.cpp
@@ -1666,7 +1666,7 @@ void FLinuxApplication::RemoveNotificationWindow(SDL_HWindow HWnd)
 	for (int32 WindowIndex=0; WindowIndex < NotificationWindows.Num(); ++WindowIndex)
 	{
 		TSharedRef< FLinuxWindow > Window = NotificationWindows[ WindowIndex ];
-		
+
 		if ( Window->GetHWnd() == HWnd )
 		{
 			NotificationWindows.RemoveAt(WindowIndex);
@@ -1677,12 +1677,12 @@ void FLinuxApplication::RemoveNotificationWindow(SDL_HWindow HWnd)
 
 void FLinuxApplication::RaiseNotificationWindows(const TSharedPtr< FLinuxWindow >& ParentWindow)
 {
-	// Raise notification window only for the correct parent window.
-	// TODO Do we have to make this restriction?
+	// Raise notification windows above everything except for modal windows
 	for (int32 WindowIndex=0; WindowIndex < NotificationWindows.Num(); ++WindowIndex)
 	{
 		TSharedRef< FLinuxWindow > NotificationWindow = NotificationWindows[WindowIndex];
-		if(ParentWindow == NotificationWindow->GetParent())
+
+		if(!ParentWindow->IsModalWindow())
 		{
 			SDL_RaiseWindow(NotificationWindow->GetHWnd());
 		}
diff --git a/Engine/Source/Runtime/ApplicationCore/Private/Linux/LinuxWindow.cpp b/Engine/Source/Runtime/ApplicationCore/Private/Linux/LinuxWindow.cpp
--- a/Engine/Source/Runtime/ApplicationCore/Private/Linux/LinuxWindow.cpp
+++ b/Engine/Source/Runtime/ApplicationCore/Private/Linux/LinuxWindow.cpp
@@ -347,7 +347,7 @@ void FLinuxWindow::Initialize( FLinuxApplication* const Application, const TShar
 		they are some sort of independend but will be raised if the TopLevel Window gets focused or activated.
 	*/
 	// Make the Window modal for it's parent.
-	if (bIsUtilityWindow || bIsDialogWindow || bIsConsoleWindow || bIsDialogWindow)
+	if (bIsUtilityWindow || bIsDialogWindow || bIsConsoleWindow || bIsNotificationWindow)
 	{
 		SDL_SetWindowModalFor(HWnd, InParent->GetHWnd());
 	}
@@ -837,6 +837,11 @@ bool FLinuxWindow::IsTopLevelWindow() const
 	return bIsTopLevelWindow;
 }
 
+bool FLinuxWindow::IsModalWindow() const
+{
+	return Definition->IsModalWindow;
+}
+
 bool FLinuxWindow::IsDialogWindow() const
 {
 	return bIsDialogWindow;
diff --git a/Engine/Source/Runtime/ApplicationCore/Public/Linux/LinuxWindow.h b/Engine/Source/Runtime/ApplicationCore/Public/Linux/LinuxWindow.h
--- a/Engine/Source/Runtime/ApplicationCore/Public/Linux/LinuxWindow.h
+++ b/Engine/Source/Runtime/ApplicationCore/Public/Linux/LinuxWindow.h
@@ -54,22 +54,25 @@ public:
 
 	/** TODO: describe */
 	bool IsNotificationWindow() const;
-	
+
 	/** TODO: describe */
 	bool IsTopLevelWindow() const;
 
+	/** @return true if this is a modal window */
+	bool IsModalWindow() const;
+
 	/** TODO: describe */
 	bool IsDialogWindow() const;
 
 	/** TODO: describe */
 	bool IsDragAndDropWindow() const;
-	
+
 	/** TODO: describe */
 	bool IsUtilityWindow() const;
 
 	/** @return the window activation policy used when showing the window */
 	EWindowActivationPolicy GetActivationPolicy() const;
-	
+
 	/** TODO: describe */
 	bool IsFocusWhenFirstShown() const;
 
