Less frequent update to intensity for on device vibrator

#android

diff --git a/Engine/Source/Runtime/ApplicationCore/Private/Android/AndroidInputInterface.cpp b/Engine/Source/Runtime/ApplicationCore/Private/Android/AndroidInputInterface.cpp
--- a/Engine/Source/Runtime/ApplicationCore/Private/Android/AndroidInputInterface.cpp
+++ b/Engine/Source/Runtime/ApplicationCore/Private/Android/AndroidInputInterface.cpp
@@ -50,6 +50,13 @@ TArray<FAndroidInputInterface::MotionData> FAndroidInputInterface::MotionDataSta
 TArray<FAndroidInputInterface::MouseData> FAndroidInputInterface::MouseDataStack
 	= TArray<FAndroidInputInterface::MouseData>();
 
+float GAndroidVibrationThreshold = 0.3f;
+static FAutoConsoleVariableRef CVarAndroidVibrationThreshold(
+	TEXT("Android.VibrationThreshold"),
+	GAndroidVibrationThreshold,
+	TEXT("If set above 0.0 acts as on/off threshold for device vibrator (Default: 0.3)"),
+	ECVF_Default);
+
 int32 GAndroidUseControllerFeedback = 1;
 static FAutoConsoleVariableRef CVarAndroidUseControllerFeedback(
 	TEXT("Android.UseControllerFeedback"),
@@ -419,7 +426,13 @@ void FAndroidInputInterface::UpdateVibeMotors()
 	// Use largest vibration state as value
 	const float MaxLeft = VibeValues.LeftLarge > VibeValues.LeftSmall ? VibeValues.LeftLarge : VibeValues.LeftSmall;
 	const float MaxRight = VibeValues.RightLarge > VibeValues.RightSmall ? VibeValues.RightLarge : VibeValues.RightSmall;
-	const float Value = MaxLeft > MaxRight ? MaxLeft : MaxRight;
+	float Value = MaxLeft > MaxRight ? MaxLeft : MaxRight;
+
+	// apply optional threshold for old behavior
+	if (GAndroidVibrationThreshold > 0.0f)
+	{
+		Value = Value < GAndroidVibrationThreshold ? 0.0f : 1.0f;
+	}
 
 	int32 Intensity = ConvertToByte(Value);
 
