Ensure GetPlatformUserForUserIndex and GetUserIndexForPlatformUser are commutitive.
 - Prevents potential runaway recursion in ULocalPlayer::SetPlatformUserId

diff --git a/Engine/Source/Runtime/ApplicationCore/Private/GenericPlatform/GenericPlatformInputDeviceMapper.cpp b/Engine/Source/Runtime/ApplicationCore/Private/GenericPlatform/GenericPlatformInputDeviceMapper.cpp
--- a/Engine/Source/Runtime/ApplicationCore/Private/GenericPlatform/GenericPlatformInputDeviceMapper.cpp
+++ b/Engine/Source/Runtime/ApplicationCore/Private/GenericPlatform/GenericPlatformInputDeviceMapper.cpp
@@ -361,15 +361,13 @@ bool FGenericPlatformInputDeviceMapper::RemapControllerIdToPlatformUserAndDevice
 
 FPlatformUserId FGenericPlatformInputDeviceMapper::GetPlatformUserForUserIndex(int32 LocalUserIndex)
 {
-	FPlatformUserId OutUserId;
-	FInputDeviceId TempDeviceId;
-
 	// The platform user index is equivalent to ControllerId in most legacy code
-	if (RemapControllerIdToPlatformUserAndDevice(LocalUserIndex, OutUserId, TempDeviceId))
+	if (IsUsingControllerIdAsUserId())
 	{
-		return OutUserId;
+		return FPlatformUserId::CreateFromInternalId(LocalUserIndex);
 	}
 
+	checkNoEntry();
 	return PLATFORMUSERID_NONE;
 }
 
