Android DPI window resolution selection cvars
 - Android.WindowDPI, set a DPI to achieve when calculating the app's screen resolution.
 - Android.3DSceneMaxDesiredPixelCount, set a target pixel count for the 3D scene color resolution.
 - Android.3DSceneMinDPI, set a forced minimum DPI for the 3D scene color resolution, useful to ensure a minimum resolution quality on low DPI devices.

diff --git a/Engine/Build/Android/Java/src/com/epicgames/unreal/GameActivity.java.template b/Engine/Build/Android/Java/src/com/epicgames/unreal/GameActivity.java.template
--- a/Engine/Build/Android/Java/src/com/epicgames/unreal/GameActivity.java.template
+++ b/Engine/Build/Android/Java/src/com/epicgames/unreal/GameActivity.java.template
@@ -5721,63 +5721,69 @@ public class GameActivity extends $${gameActivitySuperClass}$$ implements Surfac
 
 	public int AndroidThunkJava_GetMetaDataInt(String key)
 	{
-		if (key.equals("unreal.http.proxy.proxyPort"))
+		if (key == null)
 		{
-			if (ANDROID_BUILD_VERSION >= 14)
+			return 0;
+		}
+
+		switch(key)
+		{
+			case "unreal.http.proxy.proxyPort":
 			{
-				String ProxyPort = System.getProperty("http.proxyPort");
-				return ProxyPort == null ? -1 : Integer.parseInt(ProxyPort);
+				if (ANDROID_BUILD_VERSION >= 14)
+				{
+					String ProxyPort = System.getProperty("http.proxyPort");
+					return ProxyPort == null ? -1 : Integer.parseInt(ProxyPort);
+				}
+				else
+				{
+					return android.net.Proxy.getPort(getApplicationContext());
+				}
 			}
-			else
+			case "android.hardware.vulkan.version":
 			{
-				return android.net.Proxy.getPort(getApplicationContext());
+				return VulkanVersion;
 			}
-		}
-		else
-		if (key.equals("android.hardware.vulkan.version"))
-		{
-			return VulkanVersion;
-		}
-		else
-		if (key.equals("android.hardware.vulkan.level"))
-		{
-			return VulkanLevel;
-		}
-		else
-		if (key.equals("unreal.getUsedMemory"))
-		{
-			int ProcessMemory = 0;
-			synchronized(_activity)
+			case "android.hardware.vulkan.level":
 			{
-				ProcessMemory = _activity.UsedMemory;
+				return VulkanLevel;
+			}
+			case "unreal.getUsedMemory":
+			{
+				int ProcessMemory = 0;
+				synchronized(_activity)
+				{
+					ProcessMemory = _activity.UsedMemory;
+				}
+				return ProcessMemory;
+			}
+			case "audiomanager.framesPerBuffer":
+			{
+				AudioManager am = (AudioManager) getSystemService(Context.AUDIO_SERVICE);
+				String framesPerBuffer = am.getProperty(AudioManager.PROPERTY_OUTPUT_FRAMES_PER_BUFFER);
+				int framesPerBufferInt = framesPerBuffer == null ? 0 : Integer.parseInt(framesPerBuffer);
+				if (framesPerBufferInt == 0) framesPerBufferInt = 256; // Use default
+				Log.debug("[JAVA] audiomanager.framesPerBuffer = " + framesPerBufferInt);
+				return framesPerBufferInt;
+			}
+			case "audiomanager.optimalSampleRate":
+			{
+				AudioManager am = (AudioManager) getSystemService(Context.AUDIO_SERVICE);
+				String sampleRateStr = am.getProperty(AudioManager.PROPERTY_OUTPUT_SAMPLE_RATE);
+				int sampleRateInt = sampleRateStr == null ? 0 : Integer.parseInt(sampleRateStr);
+				if (sampleRateInt == 0) sampleRateInt = 44100; // Use a default value if property not found
+				Log.debug("[JAVA] audiomanager.optimalSampleRate = " + sampleRateInt);
+				return sampleRateInt;
+			}
+			default:
+			{
+				if (_bundle == null)
+				{
+					return 0;
+				}
+				return _bundle.getInt(key);
 			}
-			return ProcessMemory;
-		}
-		else
-		if (key.equals("audiomanager.framesPerBuffer"))
-		{
-			AudioManager am = (AudioManager) getSystemService(Context.AUDIO_SERVICE);
-			String framesPerBuffer = am.getProperty(AudioManager.PROPERTY_OUTPUT_FRAMES_PER_BUFFER);
-			int framesPerBufferInt = framesPerBuffer == null ? 0 : Integer.parseInt(framesPerBuffer);
-			if (framesPerBufferInt == 0) framesPerBufferInt = 256; // Use default
-			Log.debug("[JAVA] audiomanager.framesPerBuffer = " + framesPerBufferInt);
-			return framesPerBufferInt;
-		}
-		else
-		if (key.equals("audiomanager.optimalSampleRate"))
-		{
-			AudioManager am = (AudioManager) getSystemService(Context.AUDIO_SERVICE);
-			String sampleRateStr = am.getProperty(AudioManager.PROPERTY_OUTPUT_SAMPLE_RATE);
-			int sampleRateInt = sampleRateStr == null ? 0 : Integer.parseInt(sampleRateStr);
-			if (sampleRateInt == 0) sampleRateInt = 44100; // Use a default value if property not found
-			Log.debug("[JAVA] audiomanager.optimalSampleRate = " + sampleRateInt);
-			return sampleRateInt;
-		}
-		if (_bundle == null || key == null)
-		{
-			return 0;
 		}
-		return _bundle.getInt(key);
 	}
 
 	public long AndroidThunkJava_GetMetaDataLong(String key)
@@ -5819,56 +5825,183 @@ public class GameActivity extends $${gameActivitySuperClass}$$ implements Surfac
 		return 0.0f;
 	}
 
-	public String AndroidThunkJava_GetMetaDataString(String key)
+	private Display GetDefaultDisplay()
 	{
-		if (key.startsWith("getprop:"))
+		if (android.os.Build.VERSION.SDK_INT >= 31)
 		{
-			return getProp(key.substring(8));
+			return getDisplay();
 		}
 		else
-		if (key.equals("unreal.http.proxy.proxyHost"))
 		{
-			if (ANDROID_BUILD_VERSION >= 14)
-			{
-				return System.getProperty("http.proxyHost");
-			}
-			else
-			{
-				return android.net.Proxy.getHost(getApplicationContext());
-			}
+			return getWindowManager().getDefaultDisplay();
+		}
+	}
+
+	private DisplayMetrics GetDefaultDisplayMetrics()
+	{
+		DisplayMetrics metrics = new DisplayMetrics();
+		if (android.os.Build.VERSION.SDK_INT >= 31)
+		{
+			metrics = getResources().getDisplayMetrics();
 		}
 		else
-		if (key.equals("unreal.displaymetrics.dpi"))
+		if (android.os.Build.VERSION.SDK_INT >= 17)
 		{
-			DisplayMetrics metrics = new DisplayMetrics();
-			if (android.os.Build.VERSION.SDK_INT >= 31)
-			{
-				metrics = getResources().getDisplayMetrics();
-			}
-			else
-			if (android.os.Build.VERSION.SDK_INT >= 17)
+			getWindowManager().getDefaultDisplay().getRealMetrics(metrics);
+		}
+		else
+		{
+			// note: not available so get what we can
+			getWindowManager().getDefaultDisplay().getMetrics(metrics);
+		}
+		return metrics;
+	}
+
+	private boolean IsPrimitiveOrWrapperOrString(Class<?> type)
+	{
+		return type.isPrimitive()
+				|| type == Byte.class
+				|| type == Short.class
+				|| type == Integer.class
+				|| type == Long.class
+				|| type == Float.class
+				|| type == Double.class
+				|| type == Boolean.class
+				|| type == Character.class
+				|| type == String.class;
+	}
+
+	private String ExtractFieldsToKVString(Object obj, String Delim, boolean invokeParameterlessGetMethods)
+	{
+		String ret = "";
+		Class<?> ObjClass = obj.getClass();
+		try
+		{
+			Field[] fields = ObjClass.getDeclaredFields();
+			boolean bFirst = true;
+			for( Field field : fields)
 			{
-				getWindowManager().getDefaultDisplay().getRealMetrics(metrics);
+				// examine only public fields, ignore static and final to avoid gathering constants.
+				int fieldmodifiers = field.getModifiers();
+				if(!java.lang.reflect.Modifier.isPublic(fieldmodifiers)
+					|| java.lang.reflect.Modifier.isStatic(fieldmodifiers) 
+					|| java.lang.reflect.Modifier.isFinal(fieldmodifiers) 
+				)
+				{
+					continue;
+				}
+
+				String K,V;
+				try
+				{
+					K = field.getName();
+					V = field.get(obj).toString().replace(Delim, "");
+					ret += bFirst ? "" : Delim;
+					bFirst = false;
+					ret +=  K + Delim + V;
+				}
+				catch(Exception e)
+				{
+				}
 			}
-			else
+			if(invokeParameterlessGetMethods)
 			{
-				// note: not available so get what we can
-				getWindowManager().getDefaultDisplay().getMetrics(metrics);
+				Method[] methods = ObjClass.getDeclaredMethods();
+				for(Method method : methods)
+				{
+					String K,V;
+					int methodmodifier = method.getModifiers();
+					if( !java.lang.reflect.Modifier.isPublic(methodmodifier) || method.getParameterTypes().length>0 || !(method.getName().startsWith("get")||method.getName().startsWith("is")) || !IsPrimitiveOrWrapperOrString(method.getReturnType()) )
+					{
+						continue;
+					}
+					try
+					{
+						K = method.getName();
+						V = method.invoke(obj).toString().replace(Delim, "");
+						ret += bFirst ? "" : Delim;
+						bFirst = false;
+						ret +=  K + Delim + V;
+					}
+					catch(Exception e)
+					{
+					}
+				}
 			}
+		}
+		catch(SecurityException e)
+		{
+			System.out.println("failed to find fields for obj ("+ObjClass.toString()+"):\n"+ e.toString());
+		}
 
-			DecimalFormatSymbols symbols = new DecimalFormatSymbols();
-			symbols.setDecimalSeparator('.');
-			DecimalFormat decimalFormat = new DecimalFormat("###.##", symbols);
+		return ret;
+	}
 
-			String screenDPI = decimalFormat.format(metrics.xdpi) + "," + decimalFormat.format(metrics.ydpi) + "," + decimalFormat.format(metrics.densityDpi);
-			return screenDPI;
-		}
-		else
-		if (_bundle == null || key == null)
+	public String AndroidThunkJava_GetMetaDataString(String key)
+	{
+		if (key == null)
 		{
 			return null;
 		}
-		return _bundle.getString(key);
+
+		switch(key)
+		{
+			case "unreal.http.proxy.proxyHost":
+			{
+				if (ANDROID_BUILD_VERSION >= 14)
+				{
+					return System.getProperty("http.proxyHost");
+				}
+				else
+				{
+					return android.net.Proxy.getHost(getApplicationContext());
+				}
+			}
+			case "unreal.displaymetrics.dpi":
+			{
+				DisplayMetrics metrics = GetDefaultDisplayMetrics();
+
+				DecimalFormatSymbols symbols = new DecimalFormatSymbols();
+				symbols.setDecimalSeparator('.');
+				DecimalFormat decimalFormat = new DecimalFormat("###.##", symbols);
+
+				Display displayobj = GetDefaultDisplay();
+				String Test = ExtractFieldsToKVString(displayobj, ";", true);
+
+
+				String screenDPI = decimalFormat.format(metrics.xdpi) + "," + decimalFormat.format(metrics.ydpi) + "," + decimalFormat.format(metrics.densityDpi);
+				return screenDPI;
+			}
+			case "unreal.displaymetrics.metrics":
+			{
+				DisplayMetrics metrics = GetDefaultDisplayMetrics();
+				return ExtractFieldsToKVString(metrics, ";", false);
+			}
+			case "unreal.display":
+			{
+				Display displayobj = GetDefaultDisplay();
+				String Info = ExtractFieldsToKVString(displayobj, ";", true);
+				if(!Info.isEmpty())
+				{
+					android.graphics.Point displaySize = new android.graphics.Point();
+					displayobj.getRealSize(displaySize);
+					Info += ";realWidth;"+displaySize.x+";realHeight;"+displaySize.y;
+				}
+				return Info;
+			}
+			default:
+			{
+				if (key.startsWith("getprop:"))
+				{
+					return getProp(key.substring(8));
+				}				
+				if (_bundle != null)
+				{
+					return _bundle.getString(key);
+				}
+				return null;
+			}
+		}		
 	}
 
 	public boolean AndroidThunkJava_HasIntentExtrasKey(String key)
diff --git a/Engine/Source/Runtime/ApplicationCore/Private/Android/AndroidWindow.cpp b/Engine/Source/Runtime/ApplicationCore/Private/Android/AndroidWindow.cpp
--- a/Engine/Source/Runtime/ApplicationCore/Private/Android/AndroidWindow.cpp
+++ b/Engine/Source/Runtime/ApplicationCore/Private/Android/AndroidWindow.cpp
@@ -14,32 +14,72 @@
 #include "Misc/CommandLine.h"
 #include "HAL/PlatformStackWalk.h"
 
+
+int32 GAndroidWindowDPI = 0;
+static FAutoConsoleVariableRef CVarAndroidWindowDPI(
+	TEXT("Android.WindowDPI"),
+	GAndroidWindowDPI,
+	TEXT("Values > 0 will set the system window resolution (i.e. swap buffer) to achieve the DPI requested.\n")
+	TEXT("default: 0"),
+	ECVF_ReadOnly);
+
+int32 GAndroid3DSceneMaxDesiredPixelCount = 0;
+static FAutoConsoleVariableRef CVarAndroid3DSceneMaxDesiredPixelCount(
+	TEXT("Android.3DSceneMaxDesiredPixelCount"),
+	GAndroid3DSceneMaxDesiredPixelCount,
+	TEXT("Works in conjunction with Android.WindowDPI, this specifies a maximum pixel count for the 3D scene.\n")
+	TEXT("Values >0 will be used to scale down the 3D scene (equivalent to setting r.screenpercentage)\n")
+	TEXT("such that the 3d scene will be no more than 3DSceneMaxDesiredPixelCount pixels.\n")
+	TEXT("This is only applied if the 3d scene pixel at the 'Android.WindowDPI' DPI goes above the specified value.\n")
+	TEXT("default: 0"),
+	ECVF_ReadOnly);
+
+int32 GAndroid3DSceneMinDPI = 0;
+static FAutoConsoleVariableRef CVarAndroid3DSceneMinDPI(
+	TEXT("Android.3DSceneMinDPI"),
+	GAndroid3DSceneMinDPI,
+	TEXT("Works in conjunction with Android.3DSceneMaxDesiredPixelCount, and specifies a minimum DPI level for the 3D scene.\n")
+	TEXT("Values >0 specify an absolute minimum 3D scene DPI, if Android.3DSceneMaxDesiredPixelCount causes the 3d resolution to be below\n")
+	TEXT("Android.3DSceneMinDPI then the scale factor will be set to achieve the minimum dpi.\n")
+	TEXT("Useful to achieve a minimum quality level on low DPI devices at the expense of GPU performance.\n")
+	TEXT("default: 0"),
+	ECVF_ReadOnly);
+
 int GAndroidPropagateAlpha = 0;
 
-struct FCachedWindowRect
+struct FAndroidCachedWindowRectParams
 {
-	FCachedWindowRect() : WindowWidth(-1), WindowHeight(-1), WindowInit(false), ContentScaleFactor(-1.0f), MobileResX(-1), MobileResY(-1), Window_EventThread(nullptr)
-	{
-	}
+	int32 WindowWidth = -1;
+	int32 WindowHeight = -1;
+
+	float ContentScaleFactor = -1.0f;
+	int32 MobileResX = -1;
+	int32 MobileResY = -1;
+
+	int32 WindowDPI = -1;
+	int32 SceneMinDPI = -1;
+	int32 SceneMaxDesiredPixelCount = -1;
 
-	int32 WindowWidth;
-	int32 WindowHeight;
-	bool WindowInit;
-	float ContentScaleFactor;
-	int32 MobileResX;
-	int32 MobileResY;
-	ANativeWindow* Window_EventThread;
+	ANativeWindow* Window_EventThread = nullptr;
+
+	bool operator ==(const FAndroidCachedWindowRectParams& Rhs) const { return FMemory::Memcmp(this, &Rhs, sizeof(Rhs)) == 0; }
+	bool operator !=(const FAndroidCachedWindowRectParams& Rhs) const { return FMemory::Memcmp(this, &Rhs, sizeof(Rhs)) != 0; }
 };
 
 
 // Cached calculated screen resolution
-static FCachedWindowRect CachedWindowRect;
-static FCachedWindowRect CachedWindowRect_EventThread;
+static FAndroidCachedWindowRectParams CachedWindowRect;
+static FAndroidCachedWindowRectParams CachedWindowRect_EventThread;
+
+const FAndroidCachedWindowRectParams& GetCachedRect(bool bUseEventThreadWindow)
+{
+	return bUseEventThreadWindow ? CachedWindowRect_EventThread : CachedWindowRect;
+}
 
 static void ClearCachedWindowRects()
 {
-	CachedWindowRect.WindowInit = false;
-	CachedWindowRect_EventThread.WindowInit = false;
+	CachedWindowRect = FAndroidCachedWindowRectParams();
+	CachedWindowRect_EventThread = FAndroidCachedWindowRectParams();
 }
 
 static int32 GSurfaceViewWidth = -1;
@@ -280,41 +320,54 @@ void* FAndroidWindow::WaitForHardwareWindow()
 extern bool AndroidThunkCpp_IsOculusMobileApplication();
 #endif
 
-bool FAndroidWindow::IsCachedRectValid(bool bUseEventThreadWindow, const float RequestedContentScaleFactor, const int32 RequestedMobileResX, const int32 RequestedMobileResY, ANativeWindow* Window)
+static bool IsCachedRectValid(bool bUseEventThreadWindow, const FAndroidCachedWindowRectParams& TestRect)
 {
 	// window must be valid when bUseEventThreadWindow and null when !bUseEventThreadWindow.
-	check((Window != nullptr) == bUseEventThreadWindow);
+	check((TestRect.Window_EventThread != nullptr) == bUseEventThreadWindow);
+
+	const FAndroidCachedWindowRectParams& CachedRect = GetCachedRect(bUseEventThreadWindow);
 
-	const FCachedWindowRect& CachedRect = bUseEventThreadWindow ? CachedWindowRect_EventThread : CachedWindowRect;
+	bool bValidCache = true;
 
-	if (!CachedRect.WindowInit)
+	if (CachedRect.ContentScaleFactor != TestRect.ContentScaleFactor)
 	{
-		return false;
+		FPlatformMisc::LowLevelOutputDebugStringf(TEXT("***** RequestedContentScaleFactor different %f != %f, not using res cache (%d)"), TestRect.ContentScaleFactor, CachedWindowRect.ContentScaleFactor, (int32)bUseEventThreadWindow);
+		bValidCache = false;
 	}
 
-	bool bValidCache = true;
+	if (CachedRect.MobileResX != TestRect.MobileResX)
+	{
+		FPlatformMisc::LowLevelOutputDebugStringf(TEXT("***** RequestedMobileResX different %d != %d, not using res cache (%d)"), TestRect.MobileResX, CachedWindowRect.MobileResX, (int32)bUseEventThreadWindow);
+		bValidCache = false;
+	}
 
-	if (CachedRect.ContentScaleFactor != RequestedContentScaleFactor )
+	if (CachedRect.MobileResY != TestRect.MobileResY)
 	{
-		FPlatformMisc::LowLevelOutputDebugStringf(TEXT("***** RequestedContentScaleFactor different %f != %f, not using res cache (%d)"), RequestedContentScaleFactor, CachedWindowRect.ContentScaleFactor, (int32)bUseEventThreadWindow);
+		FPlatformMisc::LowLevelOutputDebugStringf(TEXT("***** RequestedMobileResY different %d != %d, not using res cache (%d)"), TestRect.MobileResY, CachedWindowRect.MobileResY, (int32)bUseEventThreadWindow);
 		bValidCache = false;
 	}
 
-	if (CachedRect.MobileResX != RequestedMobileResX)
+	if (CachedRect.Window_EventThread != TestRect.Window_EventThread)
 	{
-		FPlatformMisc::LowLevelOutputDebugStringf(TEXT("***** RequestedMobileResX different %d != %d, not using res cache (%d)"), RequestedMobileResX, CachedWindowRect.MobileResX, (int32)bUseEventThreadWindow);
+		FPlatformMisc::LowLevelOutputDebugStringf(TEXT("***** Window different, not using res cache (%d)"), (int32)bUseEventThreadWindow);
 		bValidCache = false;
 	}
 
-	if (CachedRect.MobileResY != RequestedMobileResY)
+	if (CachedRect.WindowDPI != TestRect.WindowDPI)
 	{
-		FPlatformMisc::LowLevelOutputDebugStringf(TEXT("***** RequestedMobileResY different %d != %d, not using res cache (%d)"), RequestedMobileResY, CachedWindowRect.MobileResY, (int32)bUseEventThreadWindow);
+		FPlatformMisc::LowLevelOutputDebugStringf(TEXT("***** WindowDPI is %d, not using res cache (%d)"), TestRect.WindowDPI, (int32)bUseEventThreadWindow);
 		bValidCache = false;
 	}
 
-	if (CachedRect.Window_EventThread != Window)
+	if (CachedRect.SceneMaxDesiredPixelCount != TestRect.SceneMaxDesiredPixelCount)
 	{
-		FPlatformMisc::LowLevelOutputDebugStringf(TEXT("***** Window different, not using res cache (%d)"), (int32)bUseEventThreadWindow);
+		FPlatformMisc::LowLevelOutputDebugStringf(TEXT("***** SceneMaxDesiredPixelCount is %d, not using res cache (%d)"), TestRect.SceneMaxDesiredPixelCount, (int32)bUseEventThreadWindow);
+		bValidCache = false;
+	}
+
+	if (CachedRect.SceneMinDPI != TestRect.SceneMinDPI)
+	{
+		FPlatformMisc::LowLevelOutputDebugStringf(TEXT("***** SceneMinDPI is %d, not using res cache (%d)"), TestRect.SceneMinDPI, (int32)bUseEventThreadWindow);
 		bValidCache = false;
 	}
 
@@ -327,21 +380,23 @@ bool FAndroidWindow::IsCachedRectValid(bool bUseEventThreadWindow, const float R
 	return bValidCache;
 }
 
-void FAndroidWindow::CacheRect(bool bUseEventThreadWindow, const int32 Width, const int32 Height, const float RequestedContentScaleFactor, const int32 RequestedMobileResX, const int32 RequestedMobileResY, ANativeWindow* Window)
+void CacheRect(bool bUseEventThreadWindow, const FAndroidCachedWindowRectParams& NewValues)
 {
-	check(Window != nullptr || !bUseEventThreadWindow);
+	check(NewValues.Window_EventThread != nullptr || !bUseEventThreadWindow);
 
-	FCachedWindowRect& CachedRect = bUseEventThreadWindow ? CachedWindowRect_EventThread : CachedWindowRect;
+	(bUseEventThreadWindow ? CachedWindowRect_EventThread : CachedWindowRect) = NewValues;
 
-	CachedRect.WindowWidth = Width;
-	CachedRect.WindowHeight = Height;
-	CachedRect.WindowInit = true;
-	CachedRect.ContentScaleFactor = RequestedContentScaleFactor;
-	CachedRect.MobileResX = RequestedMobileResX;
-	CachedRect.MobileResY = RequestedMobileResY;
-	CachedRect.Window_EventThread = Window;
+	UE_LOG(LogAndroid, Log, TEXT("***** Cached WindowRect %d, %d (%d)"), NewValues.WindowWidth, NewValues.WindowHeight, (int32)bUseEventThreadWindow);
 }
 
+struct FAndroidDisplayInfo
+{
+	FIntVector2 WindowDims;
+	float SceneScaleFactor;
+};
+
+static FAndroidDisplayInfo GetAndroidDisplayInfoFromDPITargets(int32 TargetDPI, int32 SceneMaxDesiredPixelCount, int32 LowerLimit3DDPI);
+
 FPlatformRect FAndroidWindow::GetScreenRect(bool bUseEventThreadWindow)
 {
 	int32 OverrideResX, OverrideResY;
@@ -393,7 +448,16 @@ FPlatformRect FAndroidWindow::GetScreenRect(bool bUseEventThreadWindow)
 
 	// since orientation won't change on Android, use cached results if still valid. Different cache is maintained for event_thread flavor.
 	ANativeWindow* Window = bUseEventThreadWindow ? (ANativeWindow*)FAndroidWindow::GetHardwareWindow_EventThread() : nullptr;
-	bool bComputeRect = !IsCachedRectValid(bUseEventThreadWindow, RequestedContentScaleFactor, RequestedResX, RequestedResY, Window);
+	FAndroidCachedWindowRectParams CurrentParams;
+	CurrentParams.ContentScaleFactor = RequestedContentScaleFactor;
+	CurrentParams.MobileResX = RequestedResX;
+	CurrentParams.MobileResY = RequestedResX;
+	CurrentParams.Window_EventThread = Window;
+	CurrentParams.WindowDPI = GAndroidWindowDPI;
+	CurrentParams.SceneMinDPI = FMath::Min(GAndroid3DSceneMinDPI, GAndroidWindowDPI);
+	CurrentParams.SceneMaxDesiredPixelCount = GAndroid3DSceneMaxDesiredPixelCount;
+
+	bool bComputeRect = !IsCachedRectValid(bUseEventThreadWindow, CurrentParams);
 	if (bComputeRect)
 	{
 		// currently hardcoding resolution
@@ -408,14 +472,27 @@ FPlatformRect FAndroidWindow::GetScreenRect(bool bUseEventThreadWindow)
 
 		if (!bIsOculusMobileApp)
 		{
-			AndroidWindowUtils::ApplyContentScaleFactor(ScreenWidth, ScreenHeight);
+			if (CurrentParams.WindowDPI && RequestedResX == 0 && RequestedResY == 0)
+			{
+				FAndroidDisplayInfo Info = GetAndroidDisplayInfoFromDPITargets(CurrentParams.WindowDPI, CurrentParams.SceneMaxDesiredPixelCount, CurrentParams.SceneMinDPI);
+				ScreenWidth = Info.WindowDims.X;
+				ScreenHeight = Info.WindowDims.Y;
+				static IConsoleVariable* CVarSSP = IConsoleManager::Get().FindConsoleVariable(TEXT("r.SecondaryScreenPercentage.GameViewport"));
+				CVarSSP->Set((float)Info.SceneScaleFactor*100.0f);
+			}
+			else
+			{
+				AndroidWindowUtils::ApplyContentScaleFactor(ScreenWidth, ScreenHeight);
+			}
+			CurrentParams.WindowWidth = ScreenWidth;
+			CurrentParams.WindowHeight = ScreenHeight;
 		}
 
 		// save for future calls
-		CacheRect(bUseEventThreadWindow, ScreenWidth, ScreenHeight, RequestedContentScaleFactor, RequestedResX, RequestedResY, Window);
+		CacheRect(bUseEventThreadWindow, CurrentParams);
 	}
 
-	const FCachedWindowRect& CachedRect = bUseEventThreadWindow ? CachedWindowRect_EventThread : CachedWindowRect;
+	const FAndroidCachedWindowRectParams& CachedRect = GetCachedRect(bUseEventThreadWindow);
 
 	// create rect and return
 	FPlatformRect ScreenRect;
@@ -487,3 +564,98 @@ bool FAndroidWindow::OnWindowOrientationChanged(bool bIsPortrait)
 	}
 	return false;
 }
+
+extern FString AndroidThunkCpp_GetMetaDataString(const FString& Key);
+
+// Sets the Android screen size to 
+static FAndroidDisplayInfo GetAndroidDisplayInfoFromDPITargets(int32 TargetDPI, int32 SceneMaxDesiredPixelCount, int32 LowerLimit3DDPI)
+{
+	check(IsInGameThread());
+
+	auto StringToMap = [](const FString& str)
+	{
+		TMap<FString, FString> mapret;
+		TArray<FString> OutArray;
+		str.ParseIntoArray(OutArray, TEXT(";"));
+		if (OutArray.Num() % 2 == 0)
+		{
+			for (int i = 0; i < OutArray.Num(); i += 2)
+			{
+				mapret.Add(OutArray[i], OutArray[i + 1]);
+			}
+		}
+		return mapret;
+	};
+
+	FString JNIMetrics = AndroidThunkCpp_GetMetaDataString(FString(TEXT("unreal.displaymetrics.metrics")));
+	FString JNIDisplay = AndroidThunkCpp_GetMetaDataString(FString(TEXT("unreal.display")));
+
+	TMap<FString, FString> metricsParams = StringToMap(JNIMetrics);
+	TMap<FString, FString> displayParams = StringToMap(JNIDisplay);
+	static const FString DensityDpiKey(TEXT("densityDpi"));
+	static const FString WidthPixelsKey(TEXT("realWidth"));
+	static const FString HeightPixelsKey(TEXT("realHeight"));
+	int32 NativeScreenDensityDPI = metricsParams.Contains(TEXT("densityDpi")) ? FCString::Atoi(*metricsParams.FindChecked(DensityDpiKey)) : 0;
+	FIntVector2 NativeScreenPixelDims(
+		displayParams.Contains(WidthPixelsKey) ? FCString::Atoi(*displayParams.FindChecked(WidthPixelsKey)) : 0
+		, displayParams.Contains(HeightPixelsKey) ? FCString::Atoi(*displayParams.FindChecked(HeightPixelsKey)) : 0
+	);
+
+	FAndroidDisplayInfo Info;
+	Info.SceneScaleFactor = 1.0f;
+
+	// NativeScreenDensityDPI is approx.
+	const float ApproxSystemToDesiredDPIScale = FMath::Min( TargetDPI / (float)NativeScreenDensityDPI, 1.0f);
+	Info.WindowDims = FIntVector2((int32)FMath::RoundFromZero(NativeScreenPixelDims.X * ApproxSystemToDesiredDPIScale), (int32)FMath::RoundFromZero(NativeScreenPixelDims.Y * ApproxSystemToDesiredDPIScale));	
+	FIntVector2 Sanitized = AndroidWindowUtils::SanitizeAndroidScreenSize(MoveTemp(NativeScreenPixelDims), FIntVector2(Info.WindowDims));
+
+	// ignore minor differences from sanitization.
+	//float ScaleFromSanitizing = FMath::Sqrt((float)(Sanitized.X * Sanitized.Y) / (float)(Info.WindowDims.X * Info.WindowDims.Y));
+	Info.WindowDims = Sanitized;
+	//TargetDPI = ScaleFromSanitizing;
+
+	UE_LOG(LogAndroid, Display, TEXT("AndroidDisplayInfoFromDPITargets : Native screen dpi %d, res %d x %d"), NativeScreenDensityDPI, NativeScreenPixelDims.X, NativeScreenPixelDims.Y);
+	UE_CLOG(TargetDPI <= NativeScreenDensityDPI, LogAndroid, Display, TEXT("AndroidDisplayInfoFromDPITargets : New DPI target %d, window dims %d, %d"), TargetDPI, Info.WindowDims.X, Info.WindowDims.Y);
+	UE_CLOG(TargetDPI > NativeScreenDensityDPI, LogAndroid, Display, TEXT("AndroidDisplayInfoFromDPITargets : TargetDPI too high, using native screen DPI %d, window dims %d, %d"), NativeScreenDensityDPI, Info.WindowDims.X, Info.WindowDims.Y);
+	TargetDPI = FMath::Min(TargetDPI, NativeScreenDensityDPI);
+
+	int DesiredPixelCount = Info.WindowDims.X * Info.WindowDims.Y;
+	if (SceneMaxDesiredPixelCount && DesiredPixelCount > SceneMaxDesiredPixelCount)
+	{
+		// if we're going to be pushing too many pixels, scale back 3d scene size to get us to SceneMaxDesiredPixelCount
+		Info.SceneScaleFactor = FMath::Sqrt(SceneMaxDesiredPixelCount / (float)DesiredPixelCount);
+		UE_LOG(LogAndroid, Warning, TEXT("AndroidDisplayInfoFromDPITargets : DPI %d has a %d pixels, this exceeds the pixels limit of %d by %d%%. 3d scene target is reduced to %d x %d"),
+			TargetDPI,
+			DesiredPixelCount, 
+			SceneMaxDesiredPixelCount, 
+			(uint32)((DesiredPixelCount/(float)SceneMaxDesiredPixelCount)*100),
+			(uint32)FMath::RoundFromZero(Info.WindowDims.X * Info.SceneScaleFactor),
+			(uint32)FMath::RoundFromZero(Info.WindowDims.Y * Info.SceneScaleFactor)
+			);
+	}
+	
+	// if a min dpi was specified then clamp to that and accept a perf hit for res quality.
+	if (LowerLimit3DDPI && (TargetDPI * Info.SceneScaleFactor) < LowerLimit3DDPI)
+	{
+		float DPILimitScale = LowerLimit3DDPI / (float)TargetDPI;
+
+		UE_LOG(LogAndroid, Warning, TEXT("AndroidDisplayInfoFromDPITargets : 3d scene target of %d DPI is lower than specified limit of %d DPI, increasing scene target dims %dx%d -> %dx%d"),
+			(uint32)FMath::RoundFromZero(TargetDPI * Info.SceneScaleFactor),
+			LowerLimit3DDPI,
+			(uint32)FMath::RoundFromZero(Info.WindowDims.X * Info.SceneScaleFactor),
+			(uint32)FMath::RoundFromZero(Info.WindowDims.Y * Info.SceneScaleFactor),
+			(uint32)FMath::RoundFromZero(Info.WindowDims.X * DPILimitScale),
+			(uint32)FMath::RoundFromZero(Info.WindowDims.Y * DPILimitScale)
+			);
+		Info.SceneScaleFactor = LowerLimit3DDPI / (float)TargetDPI;
+	}
+
+	if (SceneMaxDesiredPixelCount)
+	{
+		int FinalSceneTargetPixelCount = DesiredPixelCount * Info.SceneScaleFactor * Info.SceneScaleFactor;
+		UE_LOG(LogAndroid, Display, TEXT("AndroidDisplayInfoFromDPITargets : SceneTarget Pixel count %d%% of limit."), (uint32)((FinalSceneTargetPixelCount / (float)SceneMaxDesiredPixelCount) * 100));
+	}
+	
+	return Info;
+}
+
diff --git a/Engine/Source/Runtime/ApplicationCore/Public/Android/AndroidWindow.h b/Engine/Source/Runtime/ApplicationCore/Public/Android/AndroidWindow.h
--- a/Engine/Source/Runtime/ApplicationCore/Public/Android/AndroidWindow.h
+++ b/Engine/Source/Runtime/ApplicationCore/Public/Android/AndroidWindow.h
@@ -74,12 +74,6 @@ protected:
 	virtual EWindowMode::Type GetWindowMode() const override { return EWindowMode::Fullscreen; }
 
 private:
-	/** called from GetScreenRect function */
-	/** test cached values from the latest computations stored by CacheRect to decide their validity with the provided arguments */
-	static bool IsCachedRectValid(bool bUseEventThreadWindow, const float RequestedContentScaleFactor, const int32 RequestedMobileResX, const int32 RequestedMobileResY, ANativeWindow* Window);
-	/** caches some values used to compute the size of the window by GetScreenRect function */
-	static void CacheRect(bool bUseEventThreadWindow, const int32 Width, const int32 Height, const float RequestedContentScaleFactor, const int32 RequestedMobileResX, const int32 RequestedMobileResY, ANativeWindow* Window);
-
 	/**
 	 * Protect the constructor; only TSharedRefs of this class can be made.
 	 */
diff --git a/Engine/Source/Runtime/ApplicationCore/Public/Android/AndroidWindowUtils.h b/Engine/Source/Runtime/ApplicationCore/Public/Android/AndroidWindowUtils.h
--- a/Engine/Source/Runtime/ApplicationCore/Public/Android/AndroidWindowUtils.h
+++ b/Engine/Source/Runtime/ApplicationCore/Public/Android/AndroidWindowUtils.h
@@ -19,10 +19,23 @@ DEFINE_LOG_CATEGORY_STATIC(LogAndroidWindowUtils, Log, All)
 
 namespace AndroidWindowUtils
 {
-	static void ApplyContentScaleFactor(int32& InOutScreenWidth, int32& InOutScreenHeight)
+	FIntVector2 SanitizeAndroidScreenSize(const FIntVector2&& MaxScreenDims, const FIntVector2&& RequestedScreenDims)
+	{
+		FIntVector2 SanitizedDims;
+		// ensure Width and Height is multiple of 8
+		SanitizedDims.X = (RequestedScreenDims.X / 8) * 8;
+		SanitizedDims.Y = (RequestedScreenDims.Y / 8) * 8;
+
+		// clamp to native resolution
+		SanitizedDims.X = FPlatformMath::Min(SanitizedDims.X, MaxScreenDims.X);
+		SanitizedDims.Y = FPlatformMath::Min(SanitizedDims.Y, MaxScreenDims.Y);
+		return SanitizedDims;
+	}
+
+	void ApplyContentScaleFactor(int32& InOutScreenWidth, int32& InOutScreenHeight)
 	{
 		// CSF is a multiplier to 1280x720
-		static IConsoleVariable* CVarScale = IConsoleManager::Get().FindConsoleVariable(TEXT("r.MobileContentScaleFactor"));		
+		static IConsoleVariable* CVarScale = IConsoleManager::Get().FindConsoleVariable(TEXT("r.MobileContentScaleFactor"));
 		float RequestedContentScaleFactor = CVarScale->GetFloat();
 
 		static IConsoleVariable* CVarResX = IConsoleManager::Get().FindConsoleVariable(TEXT("r.Mobile.DesiredResX"));
@@ -102,16 +115,11 @@ namespace AndroidWindowUtils
 				Width = FMath::TruncToInt32((float)Height * AspectRatio + 0.5f);
 			}
 
-			// ensure Width and Height is multiple of 8
-			Width = (Width / 8) * 8;
-			Height = (Height / 8) * 8;
-
-			// clamp to native resolution
-			InOutScreenWidth = FPlatformMath::Min(Width, InOutScreenWidth);
-			InOutScreenHeight = FPlatformMath::Min(Height, InOutScreenHeight);
+			const FIntVector2 Sanitized = SanitizeAndroidScreenSize(FIntVector2(InOutScreenWidth, InOutScreenHeight), FIntVector2(Width, Height));
+			InOutScreenWidth = Sanitized.X;
+			InOutScreenHeight = Sanitized.Y;
 
 			UE_LOG(LogAndroidWindowUtils, Log, TEXT("Setting Width=%d and Height=%d (requested scale = %f)"), InOutScreenWidth, InOutScreenHeight, RequestedContentScaleFactor);
 		}
 	}
-
-} // end AndroidWindowUtils
\ No newline at end of file
+} // end AndroidWindowUtils
diff --git a/Engine/Source/Runtime/OpenGLDrv/Private/Android/AndroidEGL.cpp b/Engine/Source/Runtime/OpenGLDrv/Private/Android/AndroidEGL.cpp
--- a/Engine/Source/Runtime/OpenGLDrv/Private/Android/AndroidEGL.cpp
+++ b/Engine/Source/Runtime/OpenGLDrv/Private/Android/AndroidEGL.cpp
@@ -682,7 +682,9 @@ void AndroidEGL::InitRenderSurface(bool bUseSmallSurface, bool bCreateWndSurface
 		AndroidThunkCpp_SetDesiredViewSize(Width, Height);
 	}
 
-	FPlatformMisc::LowLevelOutputDebugStringf(TEXT("AndroidEGL::InitRenderSurface, wnd: %p, width: %d, height %d "), PImplData->Window, Width, Height);
+	FIntVector2 OriginalWindowSize(ANativeWindow_getWidth(PImplData->Window), ANativeWindow_getHeight(PImplData->Window));
+
+	FPlatformMisc::LowLevelOutputDebugStringf(TEXT("AndroidEGL::InitRenderSurface, setting wnd: %p, width: %d->%d, height %d->%d "), PImplData->Window, OriginalWindowSize.X, Width, OriginalWindowSize.Y, Height);
 	ANativeWindow_setBuffersGeometry(PImplData->Window, Width, Height, PImplData->NativeVisualID);
 	CreateEGLRenderSurface(PImplData->Window, bCreateWndSurface);
 
