Fix downstream cis

diff --git a/Engine/Source/Runtime/ApplicationCore/Private/Android/AndroidWindow.cpp b/Engine/Source/Runtime/ApplicationCore/Private/Android/AndroidWindow.cpp
--- a/Engine/Source/Runtime/ApplicationCore/Private/Android/AndroidWindow.cpp
+++ b/Engine/Source/Runtime/ApplicationCore/Private/Android/AndroidWindow.cpp
@@ -605,8 +605,8 @@ static FAndroidDisplayInfo GetAndroidDisplayInfoFromDPITargets(int32 TargetDPI,
 	Info.SceneScaleFactor = 1.0f;
 
 	// NativeScreenDensityDPI is approx.
-	const float ApproxSystemToDesiredDPIScale = FMath::Min( TargetDPI / (float)NativeScreenDensityDPI, 1.0f);
-	Info.WindowDims = FIntVector2((int32)FMath::RoundFromZero(NativeScreenPixelDims.X * ApproxSystemToDesiredDPIScale), (int32)FMath::RoundFromZero(NativeScreenPixelDims.Y * ApproxSystemToDesiredDPIScale));	
+	const float ApproxSystemToDesiredDPIScale = FMath::Min((float)TargetDPI / (float)NativeScreenDensityDPI, 1.0f);
+	Info.WindowDims = FIntVector2((int32)FMath::RoundFromZero((float)NativeScreenPixelDims.X * ApproxSystemToDesiredDPIScale), (int32)FMath::RoundFromZero((float)NativeScreenPixelDims.Y * ApproxSystemToDesiredDPIScale));
 	FIntVector2 Sanitized = AndroidWindowUtils::SanitizeAndroidScreenSize(MoveTemp(NativeScreenPixelDims), FIntVector2(Info.WindowDims));
 
 	// ignore minor differences from sanitization.
@@ -623,37 +623,37 @@ static FAndroidDisplayInfo GetAndroidDisplayInfoFromDPITargets(int32 TargetDPI,
 	if (SceneMaxDesiredPixelCount && DesiredPixelCount > SceneMaxDesiredPixelCount)
 	{
 		// if we're going to be pushing too many pixels, scale back 3d scene size to get us to SceneMaxDesiredPixelCount
-		Info.SceneScaleFactor = FMath::Sqrt(SceneMaxDesiredPixelCount / (float)DesiredPixelCount);
+		Info.SceneScaleFactor = FMath::Sqrt((float)SceneMaxDesiredPixelCount / (float)DesiredPixelCount);
 		UE_LOG(LogAndroid, Warning, TEXT("AndroidDisplayInfoFromDPITargets : DPI %d has a %d pixels, this exceeds the pixels limit of %d by %d%%. 3d scene target is reduced to %d x %d"),
 			TargetDPI,
 			DesiredPixelCount, 
 			SceneMaxDesiredPixelCount, 
-			(uint32)((DesiredPixelCount/(float)SceneMaxDesiredPixelCount)*100),
-			(uint32)FMath::RoundFromZero(Info.WindowDims.X * Info.SceneScaleFactor),
-			(uint32)FMath::RoundFromZero(Info.WindowDims.Y * Info.SceneScaleFactor)
+			(uint32)(((float)DesiredPixelCount/(float)SceneMaxDesiredPixelCount)*100),
+			(uint32)FMath::RoundFromZero((float)Info.WindowDims.X * Info.SceneScaleFactor),
+			(uint32)FMath::RoundFromZero((float)Info.WindowDims.Y * Info.SceneScaleFactor)
 			);
 	}
 	
 	// if a min dpi was specified then clamp to that and accept a perf hit for res quality.
-	if (LowerLimit3DDPI && (TargetDPI * Info.SceneScaleFactor) < LowerLimit3DDPI)
+	if (LowerLimit3DDPI && (int32)((float)TargetDPI * Info.SceneScaleFactor) < LowerLimit3DDPI)
 	{
-		float DPILimitScale = LowerLimit3DDPI / (float)TargetDPI;
+		float DPILimitScale = (float)LowerLimit3DDPI / (float)TargetDPI;
 
 		UE_LOG(LogAndroid, Warning, TEXT("AndroidDisplayInfoFromDPITargets : 3d scene target of %d DPI is lower than specified limit of %d DPI, increasing scene target dims %dx%d -> %dx%d"),
-			(uint32)FMath::RoundFromZero(TargetDPI * Info.SceneScaleFactor),
+			(uint32)FMath::RoundFromZero((float)TargetDPI * Info.SceneScaleFactor),
 			LowerLimit3DDPI,
-			(uint32)FMath::RoundFromZero(Info.WindowDims.X * Info.SceneScaleFactor),
-			(uint32)FMath::RoundFromZero(Info.WindowDims.Y * Info.SceneScaleFactor),
-			(uint32)FMath::RoundFromZero(Info.WindowDims.X * DPILimitScale),
-			(uint32)FMath::RoundFromZero(Info.WindowDims.Y * DPILimitScale)
+			(uint32)FMath::RoundFromZero((float)Info.WindowDims.X * Info.SceneScaleFactor),
+			(uint32)FMath::RoundFromZero((float)Info.WindowDims.Y * Info.SceneScaleFactor),
+			(uint32)FMath::RoundFromZero((float)Info.WindowDims.X * DPILimitScale),
+			(uint32)FMath::RoundFromZero((float)Info.WindowDims.Y * DPILimitScale)
 			);
-		Info.SceneScaleFactor = LowerLimit3DDPI / (float)TargetDPI;
+		Info.SceneScaleFactor = (float)LowerLimit3DDPI / (float)TargetDPI;
 	}
 
 	if (SceneMaxDesiredPixelCount)
 	{
-		int FinalSceneTargetPixelCount = DesiredPixelCount * Info.SceneScaleFactor * Info.SceneScaleFactor;
-		UE_LOG(LogAndroid, Display, TEXT("AndroidDisplayInfoFromDPITargets : SceneTarget Pixel count %d%% of limit."), (uint32)((FinalSceneTargetPixelCount / (float)SceneMaxDesiredPixelCount) * 100));
+		int FinalSceneTargetPixelCount = (float)DesiredPixelCount * Info.SceneScaleFactor * Info.SceneScaleFactor;
+		UE_LOG(LogAndroid, Display, TEXT("AndroidDisplayInfoFromDPITargets : SceneTarget Pixel count %d%% of limit."), (uint32)(((float)FinalSceneTargetPixelCount / (float)SceneMaxDesiredPixelCount) * 100));
 	}
 	
 	return Info;
