Fix downstream cis

diff --git a/Engine/Source/Runtime/ApplicationCore/Private/Android/AndroidWindow.cpp b/Engine/Source/Runtime/ApplicationCore/Private/Android/AndroidWindow.cpp
--- a/Engine/Source/Runtime/ApplicationCore/Private/Android/AndroidWindow.cpp
+++ b/Engine/Source/Runtime/ApplicationCore/Private/Android/AndroidWindow.cpp
@@ -652,7 +652,7 @@ static FAndroidDisplayInfo GetAndroidDisplayInfoFromDPITargets(int32 TargetDPI,
 
 	if (SceneMaxDesiredPixelCount)
 	{
-		int FinalSceneTargetPixelCount = (float)DesiredPixelCount * Info.SceneScaleFactor * Info.SceneScaleFactor;
+		int FinalSceneTargetPixelCount = (int)((float)DesiredPixelCount * Info.SceneScaleFactor * Info.SceneScaleFactor);
 		UE_LOG(LogAndroid, Display, TEXT("AndroidDisplayInfoFromDPITargets : SceneTarget Pixel count %d%% of limit."), (uint32)(((float)FinalSceneTargetPixelCount / (float)SceneMaxDesiredPixelCount) * 100));
 	}
 	
