Fixed -NoPCH -DisableUnity error for Linux.

diff --git a/Engine/Source/Runtime/ApplicationCore/Private/Linux/LinuxCursor.cpp b/Engine/Source/Runtime/ApplicationCore/Private/Linux/LinuxCursor.cpp
--- a/Engine/Source/Runtime/ApplicationCore/Private/Linux/LinuxCursor.cpp
+++ b/Engine/Source/Runtime/ApplicationCore/Private/Linux/LinuxCursor.cpp
@@ -3,6 +3,7 @@
 #include "Linux/LinuxCursor.h"
 
 #include "Misc/App.h"
+#include "Misc/Paths.h"
 #include "Linux/LinuxApplication.h"
 #include "Linux/LinuxPlatformApplicationMisc.h"
 
