Stop tooltips from popping up during window drags

diff --git a/Engine/Source/Runtime/ApplicationCore/Private/Linux/LinuxApplication.cpp b/Engine/Source/Runtime/ApplicationCore/Private/Linux/LinuxApplication.cpp
--- a/Engine/Source/Runtime/ApplicationCore/Private/Linux/LinuxApplication.cpp
+++ b/Engine/Source/Runtime/ApplicationCore/Private/Linux/LinuxApplication.cpp
@@ -814,6 +814,10 @@ void FLinuxApplication::ProcessDeferredMessage( SDL_Event Event )
 								RaiseNotificationWindows(CurrentEventWindow);
 							}
 						}
+
+						// Mouse dragging is the likely cause of this window move, so invalidate the cached cursor position
+						FLinuxCursor *LinuxCursor = (FLinuxCursor*)Cursor.Get();
+						LinuxCursor->InvalidateCaches();
 					}
 					break;
 
