Change the cast style used in CL 24732346 to static_cast

diff --git a/Engine/Source/Runtime/ApplicationCore/Private/Linux/LinuxApplication.cpp b/Engine/Source/Runtime/ApplicationCore/Private/Linux/LinuxApplication.cpp
--- a/Engine/Source/Runtime/ApplicationCore/Private/Linux/LinuxApplication.cpp
+++ b/Engine/Source/Runtime/ApplicationCore/Private/Linux/LinuxApplication.cpp
@@ -816,7 +816,7 @@ void FLinuxApplication::ProcessDeferredMessage( SDL_Event Event )
 						}
 
 						// Mouse dragging is the likely cause of this window move, so invalidate the cached cursor position
-						FLinuxCursor *LinuxCursor = (FLinuxCursor*)Cursor.Get();
+						FLinuxCursor *LinuxCursor = static_cast<FLinuxCursor*>(Cursor.Get());
 						LinuxCursor->InvalidateCaches();
 					}
 					break;
