Fix assert when Android 3d scene DPI is clamped by Android.3DSceneMinDPI.
 - set secondary screen percent cvar from GT only.

diff --git a/Engine/Source/Runtime/ApplicationCore/Private/Android/AndroidWindow.cpp b/Engine/Source/Runtime/ApplicationCore/Private/Android/AndroidWindow.cpp
--- a/Engine/Source/Runtime/ApplicationCore/Private/Android/AndroidWindow.cpp
+++ b/Engine/Source/Runtime/ApplicationCore/Private/Android/AndroidWindow.cpp
@@ -477,8 +477,11 @@ FPlatformRect FAndroidWindow::GetScreenRect(bool bUseEventThreadWindow)
 				FAndroidDisplayInfo Info = GetAndroidDisplayInfoFromDPITargets(CurrentParams.WindowDPI, CurrentParams.SceneMaxDesiredPixelCount, CurrentParams.SceneMinDPI);
 				ScreenWidth = Info.WindowDims.X;
 				ScreenHeight = Info.WindowDims.Y;
-				static IConsoleVariable* CVarSSP = IConsoleManager::Get().FindConsoleVariable(TEXT("r.SecondaryScreenPercentage.GameViewport"));
-				CVarSSP->Set((float)Info.SceneScaleFactor*100.0f);
+				if(IsInGameThread())
+				{
+					static IConsoleVariable* CVarSSP = IConsoleManager::Get().FindConsoleVariable(TEXT("r.SecondaryScreenPercentage.GameViewport"));
+					CVarSSP->Set((float)Info.SceneScaleFactor * 100.0f);
+				}
 			}
 			else
 			{
@@ -570,8 +573,6 @@ extern FString AndroidThunkCpp_GetMetaDataString(const FString& Key);
 // Sets the Android screen size to 
 static FAndroidDisplayInfo GetAndroidDisplayInfoFromDPITargets(int32 TargetDPI, int32 SceneMaxDesiredPixelCount, int32 LowerLimit3DDPI)
 {
-	check(IsInGameThread());
-
 	auto StringToMap = [](const FString& str)
 	{
 		TMap<FString, FString> mapret;
