Add a null platform for use when rendering off screen.

diff --git a/Engine/Source/Runtime/ApplicationCore/Private/Linux/LinuxApplication.cpp b/Engine/Source/Runtime/ApplicationCore/Private/Linux/LinuxApplication.cpp
--- a/Engine/Source/Runtime/ApplicationCore/Private/Linux/LinuxApplication.cpp
+++ b/Engine/Source/Runtime/ApplicationCore/Private/Linux/LinuxApplication.cpp
@@ -1,13 +1,14 @@
 // Copyright Epic Games, Inc. All Rights Reserved.
 
 #include "Linux/LinuxApplication.h"
-
+#include "Null/NullApplication.h"
 #include "HAL/PlatformTime.h"
 #include "Misc/StringUtility.h"
 #include "Misc/ConfigCacheIni.h"
 #include "Misc/App.h"
 #include "Features/IModularFeatures.h"
 #include "Linux/LinuxPlatformApplicationMisc.h"
+#include "Null/NullPlatformApplicationMisc.h"
 #include "IInputDeviceModule.h"
 #include "IHapticDevice.h"
 #include "GenericPlatform/GenericPlatformInputDeviceMapper.h"
@@ -1560,104 +1561,111 @@ FPlatformRect FLinuxApplication::GetWorkArea( const FPlatformRect& CurrentWindow
 
 void FDisplayMetrics::RebuildDisplayMetrics(FDisplayMetrics& OutDisplayMetrics)
 {
-	int NumDisplays = 0;
-
-	if (LIKELY(FApp::CanEverRender()))
+	if (FNullPlatformApplicationMisc::IsUsingNullApplication())
 	{
-		if (FLinuxPlatformApplicationMisc::InitSDL()) //	will not initialize more than once
-		{
-			NumDisplays = SDL_GetNumVideoDisplays();
-		}
-		else
-		{
-			UE_LOG(LogInit, Warning, TEXT("FDisplayMetrics::GetDisplayMetrics: InitSDL() failed, cannot get display metrics"));
-		}
+		FNullPlatformDisplayMetrics::RebuildDisplayMetrics(OutDisplayMetrics);
 	}
-
-	OutDisplayMetrics.MonitorInfo.Empty();
-
-	if (NumDisplays <= 0)
+	else
 	{
-		if (IsRunningDedicatedServer())
-		{
-			// dedicated servers has always been exiting early
-			OutDisplayMetrics.PrimaryDisplayWorkAreaRect = FPlatformRect(0, 0, 0, 0);
-			OutDisplayMetrics.VirtualDisplayRect = OutDisplayMetrics.PrimaryDisplayWorkAreaRect;
-			OutDisplayMetrics.PrimaryDisplayWidth = 0;
-			OutDisplayMetrics.PrimaryDisplayHeight = 0;
+		int NumDisplays = 0;
 
-			return;
+		if (LIKELY(FApp::CanEverRender()))
+		{
+			if (FLinuxPlatformApplicationMisc::InitSDL()) //	will not initialize more than once
+			{
+				NumDisplays = SDL_GetNumVideoDisplays();
+			}
+			else
+			{
+				UE_LOG(LogInit, Warning, TEXT("FDisplayMetrics::GetDisplayMetrics: InitSDL() failed, cannot get display metrics"));
+			}
 		}
-		else
+
+		OutDisplayMetrics.MonitorInfo.Empty();
+
+		if (NumDisplays <= 0)
 		{
-			// headless clients need some plausible values because high level logic depends on viewport sizes not being 0 (see e.g. UnrealClient.cpp)
-			int32 Width = 1920;
-			int32 Height = 1080;
+			if (IsRunningDedicatedServer())
+			{
+				// dedicated servers has always been exiting early
+				OutDisplayMetrics.PrimaryDisplayWorkAreaRect = FPlatformRect(0, 0, 0, 0);
+				OutDisplayMetrics.VirtualDisplayRect = OutDisplayMetrics.PrimaryDisplayWorkAreaRect;
+				OutDisplayMetrics.PrimaryDisplayWidth = 0;
+				OutDisplayMetrics.PrimaryDisplayHeight = 0;
 
-			FMonitorInfo Display;
-			Display.bIsPrimary = true;
-			if (FPlatformApplicationMisc::IsHighDPIAwarenessEnabled())
+				return;
+			}
+			else
 			{
-				Display.DPI = 96;
+				// headless clients need some plausible values because high level logic depends on viewport sizes not being 0 (see e.g. UnrealClient.cpp)
+				int32 Width = 1920;
+				int32 Height = 1080;
+
+				FMonitorInfo Display;
+				Display.bIsPrimary = true;
+				if (FPlatformApplicationMisc::IsHighDPIAwarenessEnabled())
+				{
+					Display.DPI = 96;
+				}
+				Display.ID = TEXT("fakedisplay");
+				Display.NativeWidth = Width;
+				Display.NativeHeight = Height;
+				Display.MaxResolution = FIntPoint(Width, Height);
+				Display.DisplayRect = FPlatformRect(0, 0, Width, Height);
+				Display.WorkArea = FPlatformRect(0, 0, Width, Height);
+
+				OutDisplayMetrics.PrimaryDisplayWorkAreaRect = Display.WorkArea;
+				OutDisplayMetrics.VirtualDisplayRect = OutDisplayMetrics.PrimaryDisplayWorkAreaRect;
+				OutDisplayMetrics.PrimaryDisplayWidth = Display.NativeWidth;
+				OutDisplayMetrics.PrimaryDisplayHeight = Display.NativeHeight;
+				OutDisplayMetrics.MonitorInfo.Add(Display);
 			}
-			Display.ID = TEXT("fakedisplay");
-			Display.NativeWidth = Width;
-			Display.NativeHeight = Height;
-			Display.MaxResolution = FIntPoint(Width, Height);
-			Display.DisplayRect = FPlatformRect(0, 0, Width, Height);
-			Display.WorkArea = FPlatformRect(0, 0, Width, Height);
-
-			OutDisplayMetrics.PrimaryDisplayWorkAreaRect = Display.WorkArea;
-			OutDisplayMetrics.VirtualDisplayRect = OutDisplayMetrics.PrimaryDisplayWorkAreaRect;
-			OutDisplayMetrics.PrimaryDisplayWidth = Display.NativeWidth;
-			OutDisplayMetrics.PrimaryDisplayHeight = Display.NativeHeight;
-			OutDisplayMetrics.MonitorInfo.Add(Display);
 		}
-	}
 
-	for (int32 DisplayIdx = 0; DisplayIdx < NumDisplays; ++DisplayIdx)
-	{
-		SDL_Rect DisplayBounds, UsableBounds;
-		FMonitorInfo Display;
-		SDL_GetDisplayBounds(DisplayIdx, &DisplayBounds);
-		SDL_GetDisplayUsableBounds(DisplayIdx, &UsableBounds);
+		for (int32 DisplayIdx = 0; DisplayIdx < NumDisplays; ++DisplayIdx)
+		{
+			SDL_Rect DisplayBounds, UsableBounds;
+			FMonitorInfo Display;
+			SDL_GetDisplayBounds(DisplayIdx, &DisplayBounds);
+			SDL_GetDisplayUsableBounds(DisplayIdx, &UsableBounds);
 
-		Display.Name = UTF8_TO_TCHAR(SDL_GetDisplayName(DisplayIdx));
-		Display.ID = FString::Printf(TEXT("display%d"), DisplayIdx);
-		Display.NativeWidth = DisplayBounds.w;
-		Display.NativeHeight = DisplayBounds.h;
-		Display.MaxResolution = FIntPoint(DisplayBounds.w, DisplayBounds.h);
-		Display.DisplayRect = FPlatformRect(DisplayBounds.x, DisplayBounds.y, DisplayBounds.x + DisplayBounds.w, DisplayBounds.y + DisplayBounds.h);
-		Display.WorkArea = FPlatformRect(UsableBounds.x, UsableBounds.y, UsableBounds.x + UsableBounds.w, UsableBounds.y + UsableBounds.h);
-		Display.bIsPrimary = DisplayIdx == 0;
+			Display.Name = UTF8_TO_TCHAR(SDL_GetDisplayName(DisplayIdx));
+			Display.ID = FString::Printf(TEXT("display%d"), DisplayIdx);
+			Display.NativeWidth = DisplayBounds.w;
+			Display.NativeHeight = DisplayBounds.h;
+			Display.MaxResolution = FIntPoint(DisplayBounds.w, DisplayBounds.h);
+			Display.DisplayRect = FPlatformRect(DisplayBounds.x, DisplayBounds.y, DisplayBounds.x + DisplayBounds.w, DisplayBounds.y + DisplayBounds.h);
+			Display.WorkArea = FPlatformRect(UsableBounds.x, UsableBounds.y, UsableBounds.x + UsableBounds.w, UsableBounds.y + UsableBounds.h);
+			Display.bIsPrimary = DisplayIdx == 0;
 
-		if (FPlatformApplicationMisc::IsHighDPIAwarenessEnabled())
-		{
-			float HorzDPI = 0.0f, VertDPI = 0.0f;
-			if (SDL_GetDisplayDPI(DisplayIdx, nullptr, &HorzDPI, &VertDPI) == 0)
+			if (FPlatformApplicationMisc::IsHighDPIAwarenessEnabled())
 			{
-				Display.DPI = FMath::FloorToInt((HorzDPI + VertDPI) / 2.0f);
+				float HorzDPI = 0.0f, VertDPI = 0.0f;
+				if (SDL_GetDisplayDPI(DisplayIdx, nullptr, &HorzDPI, &VertDPI) == 0)
+				{
+					Display.DPI = FMath::FloorToInt((HorzDPI + VertDPI) / 2.0f);
+				}
 			}
-		}
 
-		OutDisplayMetrics.MonitorInfo.Add(Display);
+			OutDisplayMetrics.MonitorInfo.Add(Display);
 
-		if (Display.bIsPrimary)
-		{
-			OutDisplayMetrics.PrimaryDisplayWorkAreaRect = FPlatformRect(UsableBounds.x, UsableBounds.y, UsableBounds.x + UsableBounds.w, UsableBounds.y + UsableBounds.h);
+			if (Display.bIsPrimary)
+			{
+				OutDisplayMetrics.PrimaryDisplayWorkAreaRect = FPlatformRect(UsableBounds.x, UsableBounds.y, UsableBounds.x + UsableBounds.w, UsableBounds.y + UsableBounds.h);
 
-			OutDisplayMetrics.PrimaryDisplayWidth = DisplayBounds.w;
-			OutDisplayMetrics.PrimaryDisplayHeight = DisplayBounds.h;
+				OutDisplayMetrics.PrimaryDisplayWidth = DisplayBounds.w;
+				OutDisplayMetrics.PrimaryDisplayHeight = DisplayBounds.h;
 
-			OutDisplayMetrics.VirtualDisplayRect = OutDisplayMetrics.PrimaryDisplayWorkAreaRect;
-		}
-		else
-		{
-			// accumulate the total bound rect
-			OutDisplayMetrics.VirtualDisplayRect.Left = FMath::Min(DisplayBounds.x, OutDisplayMetrics.VirtualDisplayRect.Left);
-			OutDisplayMetrics.VirtualDisplayRect.Right = FMath::Max(OutDisplayMetrics.VirtualDisplayRect.Right, DisplayBounds.x + DisplayBounds.w);
-			OutDisplayMetrics.VirtualDisplayRect.Top = FMath::Min(DisplayBounds.y, OutDisplayMetrics.VirtualDisplayRect.Top);
-			OutDisplayMetrics.VirtualDisplayRect.Bottom = FMath::Max(OutDisplayMetrics.VirtualDisplayRect.Bottom, DisplayBounds.y + DisplayBounds.h);
+				OutDisplayMetrics.VirtualDisplayRect = OutDisplayMetrics.PrimaryDisplayWorkAreaRect;
+			}
+			else
+			{
+				// accumulate the total bound rect
+				OutDisplayMetrics.VirtualDisplayRect.Left = FMath::Min(DisplayBounds.x, OutDisplayMetrics.VirtualDisplayRect.Left);
+				OutDisplayMetrics.VirtualDisplayRect.Right = FMath::Max(OutDisplayMetrics.VirtualDisplayRect.Right, DisplayBounds.x + DisplayBounds.w);
+				OutDisplayMetrics.VirtualDisplayRect.Top = FMath::Min(DisplayBounds.y, OutDisplayMetrics.VirtualDisplayRect.Top);
+				OutDisplayMetrics.VirtualDisplayRect.Bottom = FMath::Max(OutDisplayMetrics.VirtualDisplayRect.Bottom, DisplayBounds.y + DisplayBounds.h);
+			}
 		}
 	}
 
diff --git a/Engine/Source/Runtime/ApplicationCore/Private/Linux/LinuxPlatformApplicationMisc.cpp b/Engine/Source/Runtime/ApplicationCore/Private/Linux/LinuxPlatformApplicationMisc.cpp
--- a/Engine/Source/Runtime/ApplicationCore/Private/Linux/LinuxPlatformApplicationMisc.cpp
+++ b/Engine/Source/Runtime/ApplicationCore/Private/Linux/LinuxPlatformApplicationMisc.cpp
@@ -1,6 +1,7 @@
 // Copyright Epic Games, Inc. All Rights Reserved.
 
 #include "Linux/LinuxPlatformApplicationMisc.h"
+#include "Null/NullPlatformApplicationMisc.h"
 #include "Misc/CommandLine.h"
 #include "Misc/App.h"
 #include "HAL/ThreadHeartBeat.h"
@@ -413,6 +414,11 @@ class FFeedbackContext* FLinuxPlatformApplicationMisc::GetFeedbackContext()
 
 GenericApplication* FLinuxPlatformApplicationMisc::CreateApplication()
 {
+	if (FParse::Param(FCommandLine::Get(), TEXT("RenderOffScreen")))
+	{
+		return FNullPlatformApplicationMisc::CreateApplication();
+	}
+
 	return FLinuxApplication::CreateLinuxApplication();
 }
 
diff --git a/Engine/Source/Runtime/ApplicationCore/Private/Null/NullApplication.cpp b/Engine/Source/Runtime/ApplicationCore/Private/Null/NullApplication.cpp
new file mode 100644
--- /dev/null
+++ b/Engine/Source/Runtime/ApplicationCore/Private/Null/NullApplication.cpp
@@ -0,0 +1,445 @@
+// Copyright Epic Games, Inc. All Rights Reserved.
+
+#include "Null/NullApplication.h"
+
+#include "HAL/PlatformTime.h"
+#include "Misc/StringUtility.h"
+#include "Misc/ConfigCacheIni.h"
+#include "Misc/App.h"
+#include "Features/IModularFeatures.h"
+#include "Null/NullPlatformApplicationMisc.h"
+#include "IInputDeviceModule.h"
+#include "IHapticDevice.h"
+#include "GenericPlatform/GenericPlatformInputDeviceMapper.h"
+
+FNullApplication* NullApplication = NULL;
+
+FNullApplication* FNullApplication::CreateNullApplication()
+{
+	if (!NullApplication)
+	{
+		NullApplication = new FNullApplication();
+	}
+
+	return NullApplication;
+}
+
+void FNullApplication::MoveWindowTo(FGenericWindow* Window, const int32 X, const int32 Y)
+{
+	if (FNullApplication* Application = CreateNullApplication())
+	{
+		TSharedPtr<FNullWindow> NullWindow = Application->FindWindowByPtr(Window);
+		Application->GetMessageHandler()->OnMovedWindow(NullWindow.ToSharedRef(), X, Y);
+	}
+}
+
+void FNullApplication::OnSizeChanged(FGenericWindow* Window, const int32 Width, const int32 Height)
+{
+	if (FNullApplication* Application = CreateNullApplication())
+	{
+		TSharedPtr<FNullWindow> NullWindow = Application->FindWindowByPtr(Window);
+		Application->GetMessageHandler()->OnSizeChanged(NullWindow.ToSharedRef(), Width, Height, false);
+	}
+}
+
+void FNullApplication::GetFullscreenInfo(int32& X, int32& Y, int32& Width, int32& Height)
+{
+	if (FNullApplication* Application = CreateNullApplication())
+	{
+		FPlatformRect WorkArea;
+		WorkArea = Application->GetWorkArea(WorkArea);
+
+		X = WorkArea.Left;
+		Y = WorkArea.Top;
+		Width = WorkArea.Right - WorkArea.Left;
+		Height = WorkArea.Bottom - WorkArea.Top;
+	}
+}
+
+void FNullApplication::ShowWindow(FGenericWindow* Window)
+{
+	if (FNullApplication* Application = CreateNullApplication())
+	{
+		TSharedPtr<FNullWindow> NullWindow = Application->FindWindowByPtr(Window);
+		Application->ActivateWindow(NullWindow);
+	}
+}
+
+void FNullApplication::HideWindow(FGenericWindow* Window)
+{
+	if (FNullApplication* Application = CreateNullApplication())
+	{
+		TSharedPtr<FNullWindow> NullWindow = Application->FindWindowByPtr(Window);
+		// Commented as we currently don't want to hide windows.
+		// Application->GetMessageHandler()->OnWindowActivationChanged(NullWindow.ToSharedRef(), EWindowActivation::Deactivate);
+	}
+}
+
+void FNullApplication::DestroyWindow(FGenericWindow* Window)
+{
+	if (FNullApplication* Application = CreateNullApplication())
+	{
+		TSharedPtr<FNullWindow> NullWindow = Application->FindWindowByPtr(Window);
+		if (NullWindow)
+		{
+			Application->DestroyWindow(NullWindow.ToSharedRef());
+		}
+	}
+}
+
+FNullApplication::FNullApplication()
+	: GenericApplication(MakeShareable(new FNullCursor()))
+	, bHasLoadedInputPlugins(false)
+{
+	// TODO (william.belcher): There's no way of grabbing a work area from an OS that doesn't exist.
+	// Therefore, we have theoretically uncapped minimum and maximum. Let's just say you've got a
+	// 4096x4096 square as that's the maximum supported resolution for NvEnc
+	WorkArea.Left = 0;
+	WorkArea.Top = 0;
+	WorkArea.Right = 4096;
+	WorkArea.Bottom = 4096;
+
+	uint32 ResolutionX, ResolutionY = 0;
+	FString ResolutionStr;
+	bool bSuccess = FParse::Value(FCommandLine::Get(), TEXT("NullApplicationWorkArea="), ResolutionStr);
+	if (bSuccess)
+	{
+		bSuccess = ParseResolution(*ResolutionStr, ResolutionX, ResolutionY);
+	}
+	else
+	{
+		bool UserSpecifiedWidth = FParse::Value(FCommandLine::Get(), TEXT("NullApplicationWorkAreaX="), ResolutionX);
+		bool UserSpecifiedHeight = FParse::Value(FCommandLine::Get(), TEXT("NullApplicationWorkAreaY="), ResolutionY);
+		bSuccess = UserSpecifiedWidth | UserSpecifiedHeight;
+
+		const float AspectRatio = 16.f / 9.f;
+		if (UserSpecifiedWidth && !UserSpecifiedHeight)
+		{
+			ResolutionY = FMath::FloorToInt((float)ResolutionX / AspectRatio);
+		}
+		else if (UserSpecifiedHeight && !UserSpecifiedWidth)
+		{
+			ResolutionX = FMath::FloorToInt((float)ResolutionY * AspectRatio);
+		}
+	}
+
+	if (bSuccess)
+	{
+		WorkArea.Right = ResolutionX;
+		WorkArea.Bottom = ResolutionY;
+	}
+}
+
+FNullApplication::~FNullApplication()
+{
+}
+
+void FNullApplication::DestroyApplication()
+{
+}
+
+TSharedRef<FGenericWindow> FNullApplication::MakeWindow()
+{
+	return FNullWindow::Make();
+}
+
+void FNullApplication::InitializeWindow(const TSharedRef<FGenericWindow>& InWindow,
+	const TSharedRef<FGenericWindowDefinition>& InDefinition,
+	const TSharedPtr<FGenericWindow>& InParent,
+	const bool bShowImmediately)
+{
+	const TSharedRef<FNullWindow> Window = StaticCastSharedRef<FNullWindow>(InWindow);
+	const TSharedPtr<FNullWindow> ParentWindow = StaticCastSharedPtr<FNullWindow>(InParent);
+
+	Window->Initialize(this, InDefinition, ParentWindow, bShowImmediately);
+	Windows.Add(Window);
+}
+
+void FNullApplication::DestroyWindow(TSharedRef<FNullWindow> WindowToRemove)
+{
+	Windows.Remove(WindowToRemove);
+}
+
+TSharedPtr<FNullWindow> FNullApplication::FindWindowByPtr(FGenericWindow* WindowToFind)
+{
+	for (int32 WindowIndex = 0; WindowIndex < Windows.Num(); ++WindowIndex)
+	{
+		TSharedRef<FNullWindow> Window = Windows[WindowIndex];
+		if ((&Window.Get()) == static_cast<FNullWindow*>(WindowToFind))
+		{
+			return Window;
+		}
+	}
+
+	return TSharedPtr<FNullWindow>(nullptr);
+}
+
+void FNullApplication::ActivateWindow(const TSharedPtr<FNullWindow>& Window)
+{
+	const FGenericWindowDefinition Definition = Window->GetDefinition();
+	bool bIsTooltipWindow = (Definition.Type == EWindowType::ToolTip);
+	bool bIsNotificationWindow = (Definition.Type == EWindowType::Notification);
+
+	if (bIsTooltipWindow || bIsNotificationWindow)
+	{
+		// We don't want to force window activation changed with tooltips or notifications
+		// as doing so will close any context menus that are open
+		return;
+	}
+
+	if (CurrentlyActiveWindow != Window)
+	{
+		PreviousActiveWindow = CurrentlyActiveWindow;
+		CurrentlyActiveWindow = Window;
+		if (PreviousActiveWindow.IsValid())
+		{
+			MessageHandler->OnWindowActivationChanged(PreviousActiveWindow.ToSharedRef(), EWindowActivation::Deactivate);
+		}
+		MessageHandler->OnWindowActivationChanged(CurrentlyActiveWindow.ToSharedRef(), EWindowActivation::Activate);
+	}
+}
+
+void FNullApplication::SetMessageHandler(const TSharedRef<FGenericApplicationMessageHandler>& InMessageHandler)
+{
+	GenericApplication::SetMessageHandler(InMessageHandler);
+}
+
+void FNullApplication::PumpMessages(const float TimeDelta)
+{
+}
+
+bool FNullApplication::IsCursorDirectlyOverSlateWindow() const
+{
+	return true;
+}
+
+TSharedPtr<FGenericWindow> FNullApplication::GetWindowUnderCursor()
+{
+	return nullptr;
+}
+
+void FNullApplication::ProcessDeferredEvents(const float TimeDelta)
+{
+}
+
+void FNullApplication::PollGameDeviceState(const float TimeDelta)
+{
+	IPlatformInputDeviceMapper& Mapper = IPlatformInputDeviceMapper::Get();
+	// initialize any externally-implemented input devices (we delay load initialize the array so any plugins have had time to load)
+	if (!bHasLoadedInputPlugins && GIsRunning)
+	{
+		TArray<IInputDeviceModule*> PluginImplementations = IModularFeatures::Get().GetModularFeatureImplementations<IInputDeviceModule>(IInputDeviceModule::GetModularFeatureName());
+		for (auto InputPluginIt = PluginImplementations.CreateIterator(); InputPluginIt; ++InputPluginIt)
+		{
+			TSharedPtr<IInputDevice> Device = (*InputPluginIt)->CreateInputDevice(MessageHandler);
+			if (Device.IsValid())
+			{
+				UE_LOG(LogInit, Log, TEXT("Adding external input plugin."));
+				ExternalInputDevices.Add(Device);
+			}
+		}
+
+		bHasLoadedInputPlugins = true;
+	}
+
+	// Poll externally-implemented devices
+	for (auto DeviceIt = ExternalInputDevices.CreateIterator(); DeviceIt; ++DeviceIt)
+	{
+		(*DeviceIt)->Tick(TimeDelta);
+		(*DeviceIt)->SendControllerEvents();
+	}
+}
+
+FModifierKeysState FNullApplication::GetModifierKeys() const
+{
+	return FModifierKeysState();
+}
+
+void FNullApplication::SetCapture(const TSharedPtr<FGenericWindow>& InWindow)
+{
+}
+
+void* FNullApplication::GetCapture(void) const
+{
+	return NULL;
+}
+
+bool FNullApplication::IsGamepadAttached() const
+{
+	return false;
+}
+
+void FNullApplication::SetHighPrecisionMouseMode(const bool Enable, const TSharedPtr<FGenericWindow>& InWindow)
+{
+	MessageHandler->OnCursorSet();
+	bUsingHighPrecisionMouseInput = Enable;
+}
+
+FPlatformRect FNullApplication::GetWorkArea(const FPlatformRect& CurrentWindow) const
+{
+	return WorkArea;
+}
+
+void FNullApplication::SetWorkArea(const FPlatformRect& NewWorkArea)
+{
+	WorkArea = NewWorkArea;
+}
+
+bool FNullApplication::Exec(UWorld* InWorld, const TCHAR* Cmd, FOutputDevice& Ar)
+{
+	// Ignore any execs that doesn't start with NullApp
+	if (!FParse::Command(&Cmd, TEXT("NullApp")))
+	{
+		return false;
+	}
+
+	if (FParse::Command(&Cmd, TEXT("Cursor")))
+	{
+		return HandleCursorCommand(Cmd, Ar);
+	}
+	else if (FParse::Command(&Cmd, TEXT("Window")))
+	{
+		return HandleWindowCommand(Cmd, Ar);
+	}
+
+	return false;
+}
+
+bool FNullApplication::HandleCursorCommand(const TCHAR* Cmd, FOutputDevice& Ar)
+{
+	if (FParse::Command(&Cmd, TEXT("Status")))
+	{
+		FNullCursor* NullCursor = static_cast<FNullCursor*>(Cursor.Get());
+		FVector2D CurrentPosition = NullCursor->GetPosition();
+
+		Ar.Logf(TEXT("Cursor status:"));
+		Ar.Logf(TEXT("Position: (%f, %f)"), CurrentPosition.X, CurrentPosition.Y);
+		return true;
+	}
+
+	return false;
+}
+
+bool FNullApplication::HandleWindowCommand(const TCHAR* Cmd, FOutputDevice& Ar)
+{
+	if (FParse::Command(&Cmd, TEXT("List")))
+	{
+		Ar.Logf(TEXT("Window list:"));
+		for (int WindowIdx = 0, NumWindows = Windows.Num(); WindowIdx < NumWindows; ++WindowIdx)
+		{
+			Ar.Logf(TEXT("%d"), WindowIdx);
+		}
+
+		return true;
+	}
+
+	return false;
+}
+
+void FNullApplication::SetForceFeedbackChannelValue(int32 ControllerId, FForceFeedbackChannelType ChannelType, float Value)
+{
+	if (FApp::UseVRFocus() && !FApp::HasVRFocus())
+	{
+		return; // do not proceed if the app uses VR focus but doesn't have it
+	}
+	// send vibration to externally-implemented devices
+	for (auto DeviceIt = ExternalInputDevices.CreateIterator(); DeviceIt; ++DeviceIt)
+	{
+		(*DeviceIt)->SetChannelValue(ControllerId, ChannelType, Value);
+	}
+}
+void FNullApplication::SetForceFeedbackChannelValues(int32 ControllerId, const FForceFeedbackValues& Values)
+{
+	if (FApp::UseVRFocus() && !FApp::HasVRFocus())
+	{
+		return; // do not proceed if the app uses VR focus but doesn't have it
+	}
+
+	// send vibration to externally-implemented devices
+	for (auto DeviceIt = ExternalInputDevices.CreateIterator(); DeviceIt; ++DeviceIt)
+	{
+		// *N.B 06/20/2016*: Ideally, we would want to use GetHapticDevice instead
+		// but they're not implemented for SteamController and SteamVRController
+		if ((*DeviceIt)->IsGamepadAttached())
+		{
+			(*DeviceIt)->SetChannelValues(ControllerId, Values);
+		}
+	}
+}
+
+void FNullApplication::SetHapticFeedbackValues(int32 ControllerId, int32 Hand, const FHapticFeedbackValues& Values)
+{
+	if (FApp::UseVRFocus() && !FApp::HasVRFocus())
+	{
+		return; // do not proceed if the app uses VR focus but doesn't have it
+	}
+
+	for (auto DeviceIt = ExternalInputDevices.CreateIterator(); DeviceIt; ++DeviceIt)
+	{
+		IHapticDevice* HapticDevice = (*DeviceIt)->GetHapticDevice();
+		if (HapticDevice)
+		{
+			HapticDevice->SetHapticFeedbackValues(ControllerId, Hand, Values);
+		}
+	}
+}
+
+bool FNullApplication::ParseResolution(const TCHAR* InResolution, uint32& OutX, uint32& OutY)
+{
+	if (*InResolution)
+	{
+		FString CmdString(InResolution);
+		CmdString = CmdString.TrimStartAndEnd().ToLower();
+
+		// Retrieve the X dimensional value
+		const uint32 X = FMath::Max(FCString::Atoi(*CmdString), 0);
+
+		// Determine whether the user has entered a resolution and extract the Y dimension.
+		FString YString;
+
+		// Find separator between values (Example of expected format: 1280x768)
+		const TCHAR* YValue = NULL;
+		if (FCString::Strchr(*CmdString, 'x'))
+		{
+			YValue = const_cast<TCHAR*>(FCString::Strchr(*CmdString, 'x') + 1);
+			YString = YValue;
+			// Remove any whitespace from the end of the string
+			YString = YString.TrimStartAndEnd();
+		}
+
+		// If the Y dimensional value exists then setup to use the specified resolution.
+		uint32 Y = 0;
+		if (YValue && YString.Len() > 0)
+		{
+			if (YString.IsNumeric())
+			{
+				Y = FMath::Max(FCString::Atoi(YValue), 0);
+				OutX = X;
+				OutY = Y;
+				return true;
+			}
+		}
+	}
+	return false;
+}
+
+void FNullPlatformDisplayMetrics::RebuildDisplayMetrics(struct FDisplayMetrics& OutDisplayMetrics)
+{
+	int32 X, Y, Width, Height;
+	FNullApplication::GetFullscreenInfo(X, Y, Width, Height);
+	// Total screen size of the primary monitor
+	OutDisplayMetrics.PrimaryDisplayWidth = Width;
+	OutDisplayMetrics.PrimaryDisplayHeight = Height;
+
+
+	OutDisplayMetrics.PrimaryDisplayWorkAreaRect.Left = X;
+	OutDisplayMetrics.PrimaryDisplayWorkAreaRect.Top = Y;
+	OutDisplayMetrics.PrimaryDisplayWorkAreaRect.Right = X + Width;
+	OutDisplayMetrics.PrimaryDisplayWorkAreaRect.Bottom = Y + Height;
+	
+	// Virtual desktop area
+	OutDisplayMetrics.VirtualDisplayRect.Left = X;
+	OutDisplayMetrics.VirtualDisplayRect.Top = Y;
+	OutDisplayMetrics.VirtualDisplayRect.Right = X + Width;
+	OutDisplayMetrics.VirtualDisplayRect.Bottom = Y + Height;
+}
\ No newline at end of file
diff --git a/Engine/Source/Runtime/ApplicationCore/Private/Null/NullCursor.cpp b/Engine/Source/Runtime/ApplicationCore/Private/Null/NullCursor.cpp
new file mode 100644
--- /dev/null
+++ b/Engine/Source/Runtime/ApplicationCore/Private/Null/NullCursor.cpp
@@ -0,0 +1,59 @@
+// Copyright Epic Games, Inc. All Rights Reserved.
+
+#include "Null/NullCursor.h"
+
+#include "Misc/App.h"
+#include "Null/NullApplication.h"
+#include "Null/NullPlatformApplicationMisc.h"
+
+FNullCursor::FNullCursor()
+{
+	if (!FApp::CanEverRender())
+	{
+		// assume that non-rendering application will be fine with token cursor
+		UE_LOG(LogInit, Log, TEXT("Not creating cursor resources due to headless application."));
+		return;
+	}
+}
+
+FNullCursor::~FNullCursor()
+{
+}
+
+FVector2D FNullCursor::GetPosition() const
+{
+	return FVector2D(CachedGlobalXPosition, CachedGlobalYPosition);
+}
+
+void FNullCursor::SetPosition(const int32 X, const int32 Y)
+{
+	CachedGlobalXPosition = X;
+	CachedGlobalYPosition = Y;
+}
+
+void FNullCursor::SetType(const EMouseCursor::Type InNewCursor)
+{
+}
+
+EMouseCursor::Type FNullCursor::GetType() const
+{
+	return EMouseCursor::Type::Default;
+}
+
+void FNullCursor::GetSize(int32& Width, int32& Height) const
+{
+	Width = 16;
+	Height = 16;
+}
+
+void FNullCursor::Show(bool bShow)
+{
+}
+
+void FNullCursor::Lock(const RECT* const Bounds)
+{
+}
+
+void FNullCursor::SetTypeShape(EMouseCursor::Type InCursorType, void* InCursorHandle)
+{
+}
diff --git a/Engine/Source/Runtime/ApplicationCore/Private/Null/NullPlatformApplicationMisc.cpp b/Engine/Source/Runtime/ApplicationCore/Private/Null/NullPlatformApplicationMisc.cpp
new file mode 100644
--- /dev/null
+++ b/Engine/Source/Runtime/ApplicationCore/Private/Null/NullPlatformApplicationMisc.cpp
@@ -0,0 +1,23 @@
+// Copyright Epic Games, Inc. All Rights Reserved.
+
+#include "Null/NullPlatformApplicationMisc.h"
+#include "Null/NullApplication.h"
+#include "Misc/CommandLine.h"
+#include "Misc/Parse.h"
+
+GenericApplication* FNullPlatformApplicationMisc::CreateApplication()
+{
+	return FNullApplication::CreateNullApplication();
+}
+
+bool FNullPlatformApplicationMisc::IsUsingNullApplication()
+{
+	// NullPlatform Application only supported on Linux and Windows
+	// Theoretically it's supported on all platforms but we need to use the Pixel Streaming plugin to view
+	// the output, so we don't want to enable the null application on platforms that Pixel Streaming doesn't support
+#if !(PLATFORM_WINDOWS || PLATFORM_LINUX)
+	static bool bIsRenderingOffScreen = FParse::Param(FCommandLine::Get(), TEXT("RenderOffScreen"));
+	return bIsRenderingOffScreen;
+#endif
+	return false;
+}
\ No newline at end of file
diff --git a/Engine/Source/Runtime/ApplicationCore/Private/Null/NullWindow.cpp b/Engine/Source/Runtime/ApplicationCore/Private/Null/NullWindow.cpp
new file mode 100644
--- /dev/null
+++ b/Engine/Source/Runtime/ApplicationCore/Private/Null/NullWindow.cpp
@@ -0,0 +1,238 @@
+// Copyright Epic Games, Inc. All Rights Reserved.
+
+#include "Null/NullWindow.h"
+#include "Misc/CommandLine.h"
+#include "Misc/Parse.h"
+#include "Null/NullApplication.h"
+
+DEFINE_LOG_CATEGORY_STATIC(LogNullPlatformWindow, Log, All);
+
+TSharedRef<FNullWindow> FNullWindow::Make()
+{
+	return MakeShareable(new FNullWindow());
+}
+
+FNullWindow::FNullWindow()
+	: DPIScaleFactor(1.f)
+	, ScreenPosition(FIntPoint::ZeroValue)
+	, SizeInScreen(FIntPoint::ZeroValue)
+	, bIsVisible(false)
+{
+}
+
+FNullWindow::~FNullWindow()
+{
+}
+
+void FNullWindow::Initialize(class FNullApplication* const Application, const TSharedRef<FGenericWindowDefinition>& InDefinition, const TSharedPtr<FNullWindow>& InParent, const bool bShowImmediately)
+{
+	Definition = InDefinition;
+}
+
+void FNullWindow::ReshapeWindow(int32 X, int32 Y, int32 Width, int32 Height)
+{
+	MoveWindowTo(X, Y);
+
+	SizeInScreen.X = Width;
+	SizeInScreen.Y = Height;
+	FNullApplication::OnSizeChanged(this, Width, Height);
+}
+
+bool FNullWindow::GetFullScreenInfo(int32& X, int32& Y, int32& Width, int32& Height) const
+{
+	FNullApplication::GetFullscreenInfo(X, Y, Width, Height);
+	return true;
+}
+
+void FNullWindow::MoveWindowTo(int32 X, int32 Y)
+{
+	ScreenPosition.X = X;
+	ScreenPosition.Y = Y;
+
+	FNullApplication::MoveWindowTo(this, X, Y);
+}
+
+void FNullWindow::BringToFront(bool bForce)
+{
+	// empty default functionality
+}
+
+void FNullWindow::HACK_ForceToFront()
+{
+	// empty default functionality
+}
+
+void FNullWindow::Destroy()
+{
+	FNullApplication::DestroyWindow(this);
+}
+
+void FNullWindow::Minimize()
+{
+	// Empty functionality as we don't want to minimize because we'll never be able to get it back
+}
+
+void FNullWindow::Maximize()
+{
+	int32 X, Y, Width, Height;
+	if (GetFullScreenInfo(X, Y, Width, Height))
+	{
+		ReshapeWindow(X, Y, Width, Height);
+	}
+}
+
+void FNullWindow::Restore()
+{
+	int32 X, Y, Width, Height;
+	if (GetFullScreenInfo(X, Y, Width, Height))
+	{
+		int32 NewWidth = Width / 2;
+		int32 NewHeight = Height / 2;
+
+		ReshapeWindow(NewWidth / 2, NewHeight / 2, NewWidth, NewHeight);
+	}
+}
+
+void FNullWindow::Show()
+{
+	if (!bIsVisible)
+	{
+		FNullApplication::ShowWindow(this);
+		bIsVisible = true;
+	}
+}
+
+void FNullWindow::Hide()
+{
+	if (bIsVisible)
+	{
+		FNullApplication::HideWindow(this);
+		bIsVisible = false;
+	}
+}
+
+void FNullWindow::SetWindowMode(EWindowMode::Type InNewWindowMode)
+{
+	// empty default functionality
+}
+
+EWindowMode::Type FNullWindow::GetWindowMode() const
+{
+	// default functionality
+	return EWindowMode::Fullscreen;
+}
+
+bool FNullWindow::IsMaximized() const
+{
+	int32 X, Y, Width, Height;
+	GetFullScreenInfo(X, Y, Width, Height);
+	return (SizeInScreen.X == Width && SizeInScreen.Y == Height);
+}
+
+bool FNullWindow::IsMinimized() const
+{
+	return false;
+}
+
+bool FNullWindow::IsVisible() const
+{
+	// empty default functionality
+	return bIsVisible;
+}
+
+bool FNullWindow::GetRestoredDimensions(int32& X, int32& Y, int32& Width, int32& Height)
+{
+	return false;
+}
+
+void FNullWindow::SetWindowFocus()
+{
+	// empty default functionality
+}
+
+void FNullWindow::SetOpacity(const float InOpacity)
+{
+	// empty default functionality
+}
+
+void FNullWindow::Enable(bool bEnable)
+{
+	// empty default functionality
+}
+
+bool FNullWindow::IsPointInWindow(int32 X, int32 Y) const
+{
+	// empty default functionality
+	return true;
+}
+
+int32 FNullWindow::GetWindowBorderSize() const
+{
+	// empty default functionality
+	return 0;
+}
+
+int32 FNullWindow::GetWindowTitleBarSize() const
+{
+	// empty default functionality
+	return 0;
+}
+
+void* FNullWindow::GetOSWindowHandle() const
+{
+	return nullptr;
+}
+
+bool FNullWindow::IsForegroundWindow() const
+{
+	// empty default functionality
+	return true;
+}
+
+bool FNullWindow::IsFullscreenSupported() const
+{
+	// empty default functionality
+	return true;
+}
+
+void FNullWindow::SetText(const TCHAR* const Text)
+{
+	// empty default functionality
+}
+
+const FGenericWindowDefinition& FNullWindow::GetDefinition() const
+{
+	return *Definition.Get();
+}
+
+bool FNullWindow::IsDefinitionValid() const
+{
+	return Definition.IsValid();
+}
+
+void FNullWindow::AdjustCachedSize(FVector2D& Size) const
+{
+	Size.X = ScreenPosition.X;
+	Size.Y = ScreenPosition.Y;
+}
+
+bool FNullWindow::IsManualManageDPIChanges() const
+{
+	// returns false by default so the application can auto-manage the size of its windows, in response to DPI variations
+	return false;
+}
+
+void FNullWindow::SetManualManageDPIChanges(const bool bAutoHandle)
+{
+	// empty default functionality
+}
+
+void FNullWindow::DrawAttention(const FWindowDrawAttentionParameters& Parameters)
+{
+	// empty default functionality
+}
+
+void FNullWindow::SetNativeWindowButtonsVisibility(bool bVisible)
+{
+	// empty default functionality
+}
diff --git a/Engine/Source/Runtime/ApplicationCore/Private/Windows/WindowsApplication.cpp b/Engine/Source/Runtime/ApplicationCore/Private/Windows/WindowsApplication.cpp
--- a/Engine/Source/Runtime/ApplicationCore/Private/Windows/WindowsApplication.cpp
+++ b/Engine/Source/Runtime/ApplicationCore/Private/Windows/WindowsApplication.cpp
@@ -1,7 +1,7 @@
 // Copyright Epic Games, Inc. All Rights Reserved.
 
 #include "Windows/WindowsApplication.h"
-
+#include "Null/NullApplication.h"
 #include "Containers/StringConv.h"
 #include "CoreGlobals.h"
 #include "Internationalization/Text.h"
@@ -18,6 +18,7 @@
 #include "HAL/ThreadHeartBeat.h"
 #include "Templates/UniquePtr.h"
 #include "Windows/WindowsPlatformApplicationMisc.h"
+#include "Null/NullPlatformApplicationMisc.h"
 #include "Stats/Stats.h"
 #include "HAL/IConsoleManager.h"
 
@@ -844,30 +845,37 @@ static void GetMonitorsInfo(TArray<FMonitorInfo>& OutMonitorInfo)
 
 void FDisplayMetrics::RebuildDisplayMetrics(struct FDisplayMetrics& OutDisplayMetrics)
 {
-	// Total screen size of the primary monitor
-	OutDisplayMetrics.PrimaryDisplayWidth = ::GetSystemMetrics( SM_CXSCREEN );
-	OutDisplayMetrics.PrimaryDisplayHeight = ::GetSystemMetrics( SM_CYSCREEN );
-
-	// Get the screen rect of the primary monitor, excluding taskbar etc.
-	RECT WorkAreaRect;
-	if(!SystemParametersInfo(SPI_GETWORKAREA, 0, &WorkAreaRect, 0))
+	if (FNullPlatformApplicationMisc::IsUsingNullApplication())
 	{
-		WorkAreaRect.top = WorkAreaRect.bottom = WorkAreaRect.left = WorkAreaRect.right = 0;
+		FNullPlatformDisplayMetrics::RebuildDisplayMetrics(OutDisplayMetrics);
 	}
+	else
+	{
+		// Total screen size of the primary monitor
+		OutDisplayMetrics.PrimaryDisplayWidth = ::GetSystemMetrics( SM_CXSCREEN );
+		OutDisplayMetrics.PrimaryDisplayHeight = ::GetSystemMetrics( SM_CYSCREEN );
 
-	OutDisplayMetrics.PrimaryDisplayWorkAreaRect.Left = WorkAreaRect.left;
-	OutDisplayMetrics.PrimaryDisplayWorkAreaRect.Top = WorkAreaRect.top;
-	OutDisplayMetrics.PrimaryDisplayWorkAreaRect.Right = WorkAreaRect.right;
-	OutDisplayMetrics.PrimaryDisplayWorkAreaRect.Bottom = WorkAreaRect.bottom;
-	
-	// Virtual desktop area
-	OutDisplayMetrics.VirtualDisplayRect.Left = ::GetSystemMetrics( SM_XVIRTUALSCREEN );
-	OutDisplayMetrics.VirtualDisplayRect.Top = ::GetSystemMetrics( SM_YVIRTUALSCREEN );
-	OutDisplayMetrics.VirtualDisplayRect.Right = OutDisplayMetrics.VirtualDisplayRect.Left + ::GetSystemMetrics( SM_CXVIRTUALSCREEN );
-	OutDisplayMetrics.VirtualDisplayRect.Bottom = OutDisplayMetrics.VirtualDisplayRect.Top + ::GetSystemMetrics( SM_CYVIRTUALSCREEN );
-
-	// Get connected monitor information
-	GetMonitorsInfo(OutDisplayMetrics.MonitorInfo);
+		// Get the screen rect of the primary monitor, excluding taskbar etc.
+		RECT WorkAreaRect;
+		if(!SystemParametersInfo(SPI_GETWORKAREA, 0, &WorkAreaRect, 0))
+		{
+			WorkAreaRect.top = WorkAreaRect.bottom = WorkAreaRect.left = WorkAreaRect.right = 0;
+		}
+
+		OutDisplayMetrics.PrimaryDisplayWorkAreaRect.Left = WorkAreaRect.left;
+		OutDisplayMetrics.PrimaryDisplayWorkAreaRect.Top = WorkAreaRect.top;
+		OutDisplayMetrics.PrimaryDisplayWorkAreaRect.Right = WorkAreaRect.right;
+		OutDisplayMetrics.PrimaryDisplayWorkAreaRect.Bottom = WorkAreaRect.bottom;
+		
+		// Virtual desktop area
+		OutDisplayMetrics.VirtualDisplayRect.Left = ::GetSystemMetrics( SM_XVIRTUALSCREEN );
+		OutDisplayMetrics.VirtualDisplayRect.Top = ::GetSystemMetrics( SM_YVIRTUALSCREEN );
+		OutDisplayMetrics.VirtualDisplayRect.Right = OutDisplayMetrics.VirtualDisplayRect.Left + ::GetSystemMetrics( SM_CXVIRTUALSCREEN );
+		OutDisplayMetrics.VirtualDisplayRect.Bottom = OutDisplayMetrics.VirtualDisplayRect.Top + ::GetSystemMetrics( SM_CYVIRTUALSCREEN );
+
+		// Get connected monitor information
+		GetMonitorsInfo(OutDisplayMetrics.MonitorInfo);
+	}
 
 	// Apply the debug safe zones
 	OutDisplayMetrics.ApplyDefaultSafeZones();
diff --git a/Engine/Source/Runtime/ApplicationCore/Private/Windows/WindowsPlatformApplicationMisc.cpp b/Engine/Source/Runtime/ApplicationCore/Private/Windows/WindowsPlatformApplicationMisc.cpp
--- a/Engine/Source/Runtime/ApplicationCore/Private/Windows/WindowsPlatformApplicationMisc.cpp
+++ b/Engine/Source/Runtime/ApplicationCore/Private/Windows/WindowsPlatformApplicationMisc.cpp
@@ -16,6 +16,7 @@
 #include "Windows/WindowsPlatformOutputDevices.h"
 #include "GenericPlatform/GenericPlatformCrashContext.h"
 #include "Templates/RefCounting.h"
+#include "Null/NullPlatformApplicationMisc.h"
 
 // Resource includes.
 #include "Runtime/Launch/Resources/Windows/Resource.h"
@@ -76,6 +77,11 @@ class FFeedbackContext* FWindowsPlatformApplicationMisc::GetFeedbackContext()
 
 GenericApplication* FWindowsPlatformApplicationMisc::CreateApplication()
 {
+	if (FParse::Param(FCommandLine::Get(), TEXT("RenderOffScreen")))
+	{
+		return FNullPlatformApplicationMisc::CreateApplication();
+	}
+
 	HICON AppIconHandle = LoadIcon( hInstance, MAKEINTRESOURCE( GetAppIcon() ) );
 	if( AppIconHandle == NULL )
 	{
diff --git a/Engine/Source/Runtime/ApplicationCore/Public/Null/NullApplication.h b/Engine/Source/Runtime/ApplicationCore/Public/Null/NullApplication.h
new file mode 100644
--- /dev/null
+++ b/Engine/Source/Runtime/ApplicationCore/Public/Null/NullApplication.h
@@ -0,0 +1,148 @@
+// Copyright Epic Games, Inc. All Rights Reserved.
+
+#pragma once
+
+#include "CoreTypes.h"
+#include "HAL/UnrealMemory.h"
+#include "Containers/Array.h"
+#include "Containers/UnrealString.h"
+#include "Containers/Map.h"
+#include "Math/Vector2D.h"
+#include "Templates/SharedPointer.h"
+#include "Misc/CoreMisc.h"
+#include "GenericPlatform/GenericWindow.h"
+#include "GenericPlatform/GenericApplicationMessageHandler.h"
+#include "GenericPlatform/GenericWindowDefinition.h"
+#include "GenericPlatform/GenericApplication.h"
+#include "GenericPlatform/IInputInterface.h"
+#include "Null/NullWindow.h"
+#include "Null/NullCursor.h"
+
+class IInputDevice;
+
+struct FNullPlatformDisplayMetrics : public FDisplayMetrics
+{
+	APPLICATIONCORE_API static void RebuildDisplayMetrics(struct FDisplayMetrics& OutDisplayMetrics);
+};
+
+/**
+ * An implementation of GenericApplication specifically for use when rendering off screen.
+ * This application has no platform backing so instead keeps track of its associated NullWindows itself.
+ */
+class APPLICATIONCORE_API FNullApplication : public GenericApplication, public FSelfRegisteringExec, public IInputInterface
+{
+public:
+	static FNullApplication* CreateNullApplication();
+
+	static void MoveWindowTo(FGenericWindow* Window, const int32 X, const int32 Y);
+
+	static void OnSizeChanged(FGenericWindow* Window, const int32 Width, const int32 Height);
+
+	static void GetFullscreenInfo(int32& X, int32& Y, int32& Width, int32& Height);
+
+	static void ShowWindow(FGenericWindow* Window);
+
+	static void HideWindow(FGenericWindow* Window);
+
+	static void DestroyWindow(FGenericWindow* Window);
+
+public:
+	virtual ~FNullApplication();
+
+	virtual void DestroyApplication() override;
+
+	// FSelfRegisteringExec
+	virtual bool Exec(UWorld* InWorld, const TCHAR* Cmd, FOutputDevice& Ar) override;
+
+	TSharedPtr<FNullWindow> FindWindowByPtr(FGenericWindow* WindowToFind);
+
+	void ActivateWindow(const TSharedPtr<FNullWindow>& Window);
+
+public:
+	virtual void SetMessageHandler(const TSharedRef<class FGenericApplicationMessageHandler>& InMessageHandler) override;
+
+	virtual void PollGameDeviceState(const float TimeDelta) override;
+
+	virtual void PumpMessages(const float TimeDelta) override;
+
+	virtual void ProcessDeferredEvents(const float TimeDelta) override;
+
+	virtual TSharedRef<FGenericWindow> MakeWindow() override;
+
+	virtual void InitializeWindow(const TSharedRef<FGenericWindow>& Window, const TSharedRef<FGenericWindowDefinition>& InDefinition, const TSharedPtr<FGenericWindow>& InParent, const bool bShowImmediately) override;
+
+	void DestroyWindow(TSharedRef<FNullWindow> WindowToRemove);
+
+	virtual void SetCapture(const TSharedPtr<FGenericWindow>& InWindow) override;
+
+	virtual void* GetCapture(void) const override;
+
+	virtual void SetHighPrecisionMouseMode(const bool Enable, const TSharedPtr<FGenericWindow>& InWindow) override;
+
+	virtual bool IsUsingHighPrecisionMouseMode() const override { return bUsingHighPrecisionMouseInput; }
+
+	virtual bool IsGamepadAttached() const override;
+
+	virtual FModifierKeysState GetModifierKeys() const override;
+
+	virtual FPlatformRect GetWorkArea(const FPlatformRect& CurrentWindow) const override;
+
+	void SetWorkArea(const FPlatformRect& NewWorkArea);
+
+	virtual EWindowTransparency GetWindowTransparencySupport() const override
+	{
+		return EWindowTransparency::PerWindow;
+	}
+
+	virtual bool IsCursorDirectlyOverSlateWindow() const override;
+
+	virtual TSharedPtr<FGenericWindow> GetWindowUnderCursor() override;
+
+	virtual bool IsMouseAttached() const override { return true; }
+
+private:
+	FNullApplication();
+
+	/** Handles "Cursor" exec commands" */
+	bool HandleCursorCommand(const TCHAR* Cmd, FOutputDevice& Ar);
+
+	/** Handles "Window" exec commands" */
+	bool HandleWindowCommand(const TCHAR* Cmd, FOutputDevice& Ar);
+
+	/** Handles parsing the work area resolution from the command line */
+	bool ParseResolution(const TCHAR* InResolution, uint32& OutX, uint32& OutY);
+
+public:
+	virtual IInputInterface* GetInputInterface() override
+	{
+		return this;
+	}
+
+	// IInputInterface overrides
+	virtual void SetForceFeedbackChannelValue(int32 ControllerId, FForceFeedbackChannelType ChannelType, float Value) override;
+	virtual void SetForceFeedbackChannelValues(int32 ControllerId, const FForceFeedbackValues& Values) override;
+	virtual void SetHapticFeedbackValues(int32 ControllerId, int32 Hand, const FHapticFeedbackValues& Values) override;
+	virtual void SetLightColor(int32 ControllerId, FColor Color) override {}
+	virtual void ResetLightColor(int32 ControllerId) override {}
+
+private:
+	TArray<TSharedRef<FNullWindow>> Windows;
+
+	/** List of input devices implemented in external modules. */
+	TArray<TSharedPtr<class IInputDevice>> ExternalInputDevices;
+	bool bHasLoadedInputPlugins;
+
+	/** Using high precision mouse input */
+	bool bUsingHighPrecisionMouseInput;
+
+	/** Window that we think has been activated last. */
+	TSharedPtr<FNullWindow> CurrentlyActiveWindow;
+
+	/** Window that we think has been previously active. */
+	TSharedPtr<FNullWindow> PreviousActiveWindow;
+
+	/** The virtual work area*/
+	FPlatformRect WorkArea;
+};
+
+extern FNullApplication* NullApplication;
diff --git a/Engine/Source/Runtime/ApplicationCore/Public/Null/NullCursor.h b/Engine/Source/Runtime/ApplicationCore/Public/Null/NullCursor.h
new file mode 100644
--- /dev/null
+++ b/Engine/Source/Runtime/ApplicationCore/Public/Null/NullCursor.h
@@ -0,0 +1,60 @@
+// Copyright Epic Games, Inc. All Rights Reserved.
+
+#pragma once
+
+#include "CoreTypes.h"
+#include "Math/Vector2D.h"
+#include "GenericPlatform/ICursor.h"
+
+/**
+ * An implementation of an ICursor specifically for use when rendering off screen. 
+ * This cursor has no platform backing so instead keeps track of its position and other properties itself.
+ */
+class FNullCursor : public ICursor
+{
+public:
+	FNullCursor();
+
+	virtual ~FNullCursor();
+
+	/** The position of the cursor */
+	virtual FVector2D GetPosition() const override;
+
+	/** Sets the position of the cursor */
+	virtual void SetPosition(const int32 X, const int32 Y) override;
+
+	/** Sets the cursor */
+	virtual void SetType(const EMouseCursor::Type InNewCursor) override;
+
+	/** Gets the current type of the cursor */
+	virtual EMouseCursor::Type GetType() const override;
+
+	/** Gets the size of the cursor */
+	virtual void GetSize(int32& Width, int32& Height) const override;
+
+	/**
+	 * Shows or hides the cursor
+	 *
+	 * @param bShow	true to show the mouse cursor, false to hide it
+	 */
+	virtual void Show(bool bShow) override;
+
+	/**
+	 * Locks the cursor to the passed in bounds
+	 *
+	 * @param Bounds	The bounds to lock the cursor to.  Pass NULL to unlock.
+	 */
+	virtual void Lock(const RECT* const Bounds) override;
+
+	/**
+	 * Allows overriding the shape of a particular cursor.
+	 */
+	virtual void SetTypeShape(EMouseCursor::Type InCursorType, void* CursorHandle) override;
+
+private:
+	/** Cached global X position */
+	int32 CachedGlobalXPosition;
+
+	/** Cached global Y position */
+	int32 CachedGlobalYPosition;
+};
diff --git a/Engine/Source/Runtime/ApplicationCore/Public/Null/NullPlatformApplicationMisc.h b/Engine/Source/Runtime/ApplicationCore/Public/Null/NullPlatformApplicationMisc.h
new file mode 100644
--- /dev/null
+++ b/Engine/Source/Runtime/ApplicationCore/Public/Null/NullPlatformApplicationMisc.h
@@ -0,0 +1,12 @@
+// Copyright Epic Games, Inc. All Rights Reserved.
+
+#pragma once
+
+#include "GenericPlatform/GenericPlatformApplicationMisc.h"
+
+struct APPLICATIONCORE_API FNullPlatformApplicationMisc : public FGenericPlatformApplicationMisc
+{
+	static class GenericApplication* CreateApplication();
+
+	static bool IsUsingNullApplication();
+};
diff --git a/Engine/Source/Runtime/ApplicationCore/Public/Null/NullWindow.h b/Engine/Source/Runtime/ApplicationCore/Public/Null/NullWindow.h
new file mode 100644
--- /dev/null
+++ b/Engine/Source/Runtime/ApplicationCore/Public/Null/NullWindow.h
@@ -0,0 +1,186 @@
+// Copyright Epic Games, Inc. All Rights Reserved.
+
+#pragma once
+
+#include "CoreTypes.h"
+#include "Logging/LogMacros.h"
+#include "Templates/SharedPointer.h"
+#include "GenericPlatform/GenericWindow.h"
+#include "Math/IntPoint.h"
+
+class FNullApplication;
+struct FGenericWindowDefinition;
+
+DECLARE_LOG_CATEGORY_EXTERN(LogNullWindow, Log, All);
+DECLARE_LOG_CATEGORY_EXTERN(LogNullWindowType, Log, All);
+DECLARE_LOG_CATEGORY_EXTERN(LogNullWindowEvent, Log, All);
+
+/**
+ * An implementation of FGenericWindow specifically for use when rendering off screen.
+ * This window has no platform backing so instead keeps track of its position and other properties itself.
+ */
+class APPLICATIONCORE_API FNullWindow : public FGenericWindow
+{
+public:
+	static TSharedRef<FNullWindow> Make();
+
+	virtual ~FNullWindow();
+
+	void Initialize(class FNullApplication* const Application, const TSharedRef<FGenericWindowDefinition>& InDefinition, const TSharedPtr<FNullWindow>& InParent, const bool bShowImmediately);
+
+	/** Native windows should implement ReshapeWindow by changing the platform-specific window to be located at (X,Y) and be the dimensions Width x Height. */
+	virtual void ReshapeWindow(int32 X, int32 Y, int32 Width, int32 Height) override;
+
+	/** Returns the rectangle of the screen the window is associated with */
+	virtual bool GetFullScreenInfo(int32& X, int32& Y, int32& Width, int32& Height) const override;
+
+	/** Native windows should implement MoveWindowTo by relocating the platform-specific window to (X,Y). */
+	virtual void MoveWindowTo(int32 X, int32 Y) override;
+
+	/** Native windows should implement BringToFront by making this window the top-most window (i.e. focused). */
+	virtual void BringToFront(bool bForce = false) override;
+
+	/** @hack Force a window to front even if a different application is in front. */
+	virtual void HACK_ForceToFront() override;
+
+	/** Native windows should implement this function by asking the OS to destroy OS-specific resource associated with the window (e.g. Win32 window handle) */
+	virtual void Destroy() override;
+
+	/** Native window should implement this function by performing the equivalent of the Win32 minimize-to-taskbar operation */
+	virtual void Minimize() override;
+
+	/** Native window should implement this function by performing the equivalent of the Win32 maximize operation */
+	virtual void Maximize() override;
+
+	/** Native window should implement this function by performing the equivalent of the Win32 restore operation */
+	virtual void Restore() override;
+
+	/** Native window should make itself visible */
+	virtual void Show() override;
+
+	/** Native window should hide itself */
+	virtual void Hide() override;
+
+	/** Toggle native window between fullscreen and normal mode */
+	virtual void SetWindowMode(EWindowMode::Type InNewWindowMode) override;
+
+	/** @return true if the native window is currently in fullscreen mode, false otherwise */
+	virtual EWindowMode::Type GetWindowMode() const override;
+
+	/** @return true if the native window is maximized, false otherwise */
+	virtual bool IsMaximized() const override;
+
+	/** @return true if the native window is minimized, false otherwise */
+	virtual bool IsMinimized() const override;
+
+	/** @return true if the native window is visible, false otherwise */
+	virtual bool IsVisible() const override;
+
+	/**
+	 * Populates the size and location of the window when it is restored.
+	 * If the function fails, false is returned and X,Y,Width,Height will be undefined.
+	 *
+	 * @return true when the size and location and successfully retrieved; false otherwise.
+	 */
+	virtual bool GetRestoredDimensions(int32& X, int32& Y, int32& Width, int32& Height) override;
+
+	/**
+	 * Native windows should implement SetWindowFocus to let the OS know that a window has taken focus.
+	 * Slate handles focus on a per widget basis internally but the OS still needs to know what window has focus for proper message routing
+	 */
+	virtual void SetWindowFocus() override;
+
+	/**
+	 * Sets the opacity of this window
+	 *
+	 * @param	InOpacity	The new window opacity represented as a floating point scalar
+	 */
+	virtual void SetOpacity(const float InOpacity) override;
+
+	/**
+	 * Enables or disables the window.  If disabled the window receives no input
+	 *
+	 * @param bEnable	true to enable the window, false to disable it.
+	 */
+	virtual void Enable(bool bEnable) override;
+
+	/** @return true if native window exists underneath the coordinates */
+	virtual bool IsPointInWindow(int32 X, int32 Y) const override;
+
+	/** Gets OS specific window border size. This is necessary because Win32 does not give control over this size. */
+	virtual int32 GetWindowBorderSize() const override;
+
+	/** Gets OS specific window title bar size */
+	virtual int32 GetWindowTitleBarSize() const override;
+
+	/** Gets the OS Window handle in the form of a void pointer for other API's */
+	virtual void* GetOSWindowHandle() const override;
+
+	/** @return true if the window is in the foreground */
+	virtual bool IsForegroundWindow() const override;
+
+	/** @return true if the window is in the foreground */
+	virtual bool IsFullscreenSupported() const override;
+
+	/**
+	 * Sets the window text - usually the title but can also be text content for things like controls
+	 *
+	 * @param Text	The window's title or content text
+	 */
+	virtual void SetText(const TCHAR* const Text) override;
+
+	/** @return	The definition describing properties of the window */
+	virtual const FGenericWindowDefinition& GetDefinition() const override;
+
+	/** @return	Returns true if the window definition is valid */
+	virtual bool IsDefinitionValid() const override;
+
+	/** @return	Gives the native window a chance to adjust our stored window size before we cache it off */
+	virtual void AdjustCachedSize(FVector2D& Size) const override;
+
+	virtual float GetDPIScaleFactor() const override
+	{
+		return DPIScaleFactor;
+	}
+
+	virtual void SetDPIScaleFactor(float Value) override
+	{
+		DPIScaleFactor = Value;
+	}
+
+	/** determines whether or not this window does its own DPI management */
+	virtual bool IsManualManageDPIChanges() const override;
+
+	/** call with a true argument if this window need to do its custom size management in response to DPI variations */
+	virtual void SetManualManageDPIChanges(const bool bAutoHandle) override;
+
+	/**
+	 * Attempts to draw the user's attention to this window in whatever way is
+	 * appropriate for the platform if this window is not the current active
+	 * window.
+	 *
+	 * @param Parameters The parameters for drawing attention. Depending on
+	 *        the platform, not all parameters may be supported.
+	 */
+	virtual void DrawAttention(const FWindowDrawAttentionParameters& Parameters) override;
+
+	/** Shows or hides native window buttons on platforms that use them */
+	virtual void SetNativeWindowButtonsVisibility(bool bVisible);
+
+private:
+	/**
+	 * Protect the constructor; only TSharedRefs of this class can be made.
+	 */
+	FNullWindow();
+
+	FNullApplication* OwningApplication;
+
+	float DPIScaleFactor;
+
+	FIntPoint ScreenPosition;
+
+	FIntPoint SizeInScreen;
+
+	/** Manually store window visibility as OS doesn't do that for us in nullplatform */
+	bool bIsVisible;
+};
diff --git a/Engine/Source/Runtime/Launch/Private/LaunchEngineLoop.cpp b/Engine/Source/Runtime/Launch/Private/LaunchEngineLoop.cpp
--- a/Engine/Source/Runtime/Launch/Private/LaunchEngineLoop.cpp
+++ b/Engine/Source/Runtime/Launch/Private/LaunchEngineLoop.cpp
@@ -3153,7 +3153,7 @@ PRAGMA_ENABLE_DEPRECATION_WARNINGS
 	bool bSlateApplication = !IsRunningDedicatedServer() && (bIsRegularClient || bHasEditorToken);
 	if (bSlateApplication)
 	{
-		if (FPlatformProcess::SupportsMultithreading())
+		if (FPlatformProcess::SupportsMultithreading() && !FParse::Param(FCommandLine::Get(), TEXT("RenderOffScreen")))
 		{
 			SCOPED_BOOT_TIMING("FPlatformSplash::Show()");
 			FPlatformSplash::Show();
diff --git a/Engine/Source/Runtime/Slate/Private/Framework/Application/SlateApplication.cpp b/Engine/Source/Runtime/Slate/Private/Framework/Application/SlateApplication.cpp
--- a/Engine/Source/Runtime/Slate/Private/Framework/Application/SlateApplication.cpp
+++ b/Engine/Source/Runtime/Slate/Private/Framework/Application/SlateApplication.cpp
@@ -26,6 +26,7 @@
 #include "Input/Events.h"
 #include "Input/HittestGrid.h"
 #include "HAL/PlatformApplicationMisc.h"
+#include "Null/NullPlatformApplicationMisc.h"
 #include "GenericPlatform/GenericPlatformInputDeviceMapper.h"
 
 #if WITH_ACCESSIBILITY
@@ -1816,15 +1817,17 @@ TSharedRef<SWindow> FSlateApplication::AddWindow( TSharedRef<SWindow> InSlateWin
 
 TSharedRef< FGenericWindow > FSlateApplication::MakeWindow( TSharedRef<SWindow> InSlateWindow, const bool bShowImmediately )
 {
+	// Windows and Linux support a NullPlatform which is handled in the Windows/LinuxApplication. Other platforms do not.
+#if !(PLATFORM_WINDOWS || PLATFORM_LINUX)
 	// When rendering off-screen don't render to screen, create a dummy generic window
 	if (bRenderOffScreen)
 	{
 		TSharedRef< FGenericWindow > NewWindow = MakeShareable(new FGenericWindow());
 		InSlateWindow->SetNativeWindow(NewWindow);
-
 		FSlateApplicationBase::Get().GetRenderer()->CreateViewport(InSlateWindow);
 		return NewWindow;
 	}
+#endif
 
 	TSharedPtr<FGenericWindow> NativeParent = nullptr;
 	TSharedPtr<SWindow> ParentWindow = InSlateWindow->GetParentWindow();
