[Engine]
* WindowsConsoleOutputDevice2 - Fixed clang compile error by renaming parameter (shadowing variable)

diff --git a/Engine/Source/Runtime/ApplicationCore/Private/Windows/WindowsConsoleOutputDevice2.cpp b/Engine/Source/Runtime/ApplicationCore/Private/Windows/WindowsConsoleOutputDevice2.cpp
--- a/Engine/Source/Runtime/ApplicationCore/Private/Windows/WindowsConsoleOutputDevice2.cpp
+++ b/Engine/Source/Runtime/ApplicationCore/Private/Windows/WindowsConsoleOutputDevice2.cpp
@@ -111,10 +111,10 @@ public:
 			DeleteObject(StatusLightBrush[I]);
 	}
 
-	void AddLogEntry(const FStringView& Text, ELogVerbosity::Type Verbosity, const class FName& Category, double Time, uint16 TextAttribute)
+	void AddLogEntry(const FStringView& Text, ELogVerbosity::Type Verbosity, const class FName& Category, double Time, uint16 InTextAttribute)
 	{
 		FScopeLock _(&NewLogEntriesCs);
-		NewLogEntries.Add({FString(Text), Verbosity, Category, Time, TextAttribute});
+		NewLogEntries.Add({FString(Text), Verbosity, Category, Time, InTextAttribute});
 		SetEvent(DirtyEvent);
 	}
 
