Added a console variable to control if we should use an invisible cursor instead of setting the HCursor to NULL when the cursor type is set to None.

diff --git a/Engine/Source/Runtime/ApplicationCore/Private/Windows/WindowsCursor.cpp b/Engine/Source/Runtime/ApplicationCore/Private/Windows/WindowsCursor.cpp
--- a/Engine/Source/Runtime/ApplicationCore/Private/Windows/WindowsCursor.cpp
+++ b/Engine/Source/Runtime/ApplicationCore/Private/Windows/WindowsCursor.cpp
@@ -9,6 +9,7 @@
 #include "Math/Color.h"
 #include "Windows/WindowsWindow.h"
 #include "Windows/WindowsHWrapper.h"
+#include "HAL/IConsoleManager.h"
 #include "HAL/PlatformProcess.h"
 #include "Misc/FileHelper.h"
 
@@ -17,6 +18,13 @@
 	#include <oleidl.h>
 #include "Windows/HideWindowsPlatformTypes.h"
 
+static bool bUseInvisibleCursorForNoneCursorType = false;
+FAutoConsoleVariableRef CVarUseInvisibleCursorForNoneCursorType(
+	TEXT("WindowsCursor.UseInvisibleCursorForNoneCursorType"),
+	bUseInvisibleCursorForNoneCursorType,
+	TEXT("If enabled, sets the platform HCursor to a transparent cursor instead of null when the mouse cursor type to None."),
+	ECVF_Default);
+
 FWindowsCursor::FWindowsCursor()
 {
 	// Load up cursors that we'll be using
@@ -29,11 +37,11 @@ FWindowsCursor::FWindowsCursor()
 		switch( CurCursorIndex )
 		{
 		case EMouseCursor::None:
-			if (FPlatformMisc::IsRemoteSession())
+			if (FPlatformMisc::IsRemoteSession() || bUseInvisibleCursorForNoneCursorType)
 			{
 				// during remote sessions we rely on constantly resetting the mouse cursor position for infinite mouse deltas while dragging (WM_INPUT doesnt work over rdp)
 				// this hack here is to avoid having a hidden cursor which prevents SetCursorPos from working.  This is just a completely transparent cursor yet "visible" cursor to work around that
-				CursorHandle = LoadCursorFromFile((LPCTSTR) * (FString(FPlatformProcess::BaseDir()) / FString::Printf(TEXT("%sEditor/Slate/Cursor/invisible.cur"), *FPaths::EngineContentDir())));
+				CursorHandle = LoadCursorFromFile((LPCTSTR) * (FString(FPlatformProcess::BaseDir()) / FString::Printf(TEXT("%s/Slate/Cursor/invisible.cur"), *FPaths::EngineContentDir())));
 			}
 			break;
 		case EMouseCursor::Custom:
diff --git a/Engine/Source/Runtime/Slate/Slate.Build.cs b/Engine/Source/Runtime/Slate/Slate.Build.cs
--- a/Engine/Source/Runtime/Slate/Slate.Build.cs
+++ b/Engine/Source/Runtime/Slate/Slate.Build.cs
@@ -32,9 +32,17 @@ public class Slate : ModuleRules
 		// Add slate runtime dependencies
 		if (Target.bUsesSlate)
 		{
-			RuntimeDependencies.Add("$(EngineDir)/Content/Slate/...", StagedFileType.UFS);
+			RuntimeDependencies.Add("$(EngineDir)/Content/Slate/...*.ttf", StagedFileType.UFS);
+			RuntimeDependencies.Add("$(EngineDir)/Content/Slate/...*.png", StagedFileType.UFS);
+			RuntimeDependencies.Add("$(EngineDir)/Content/Slate/...*.svg", StagedFileType.UFS);
+			RuntimeDependencies.Add("$(EngineDir)/Content/Slate/...*.tps", StagedFileType.UFS);
 			RuntimeDependencies.Add("$(EngineDir)/Content/SlateDebug/...", StagedFileType.DebugNonUFS);
 
+			if (Target.Platform == UnrealTargetPlatform.Win64)
+			{ 
+				RuntimeDependencies.Add("$(EngineDir)/Content/Slate/...*.cur", StagedFileType.NonUFS);
+			}
+
 			if (Target.ProjectFile != null)
 			{
 				RuntimeDependencies.Add("$(ProjectDir)/Content/Slate/...", StagedFileType.UFS);
