Fix up DPI scale in Linux being incorrect on 5120x1440 displays

diff --git a/Engine/Source/Runtime/ApplicationCore/Private/Linux/LinuxPlatformApplicationMisc.cpp b/Engine/Source/Runtime/ApplicationCore/Private/Linux/LinuxPlatformApplicationMisc.cpp
--- a/Engine/Source/Runtime/ApplicationCore/Private/Linux/LinuxPlatformApplicationMisc.cpp
+++ b/Engine/Source/Runtime/ApplicationCore/Private/Linux/LinuxPlatformApplicationMisc.cpp
@@ -514,6 +514,12 @@ float FLinuxPlatformApplicationMisc::GetDPIScaleFactorAtPoint(float X, float Y)
 				{
 					float Scale = LinuxPlatformApplicationMisc::QuantizeScale((HorzDPI + VertDPI) / 192.0f);	// average between two scales (divided by 96.0f)
 					UE_LOG(LogLinux, Verbose, TEXT("Scale at X=%f, Y=%f: %f (monitor=#%d, HDPI=%f (horz scale: %f), VDPI=%f (vert scale: %f))"), X, Y, Scale, Idx, HorzDPI, HorzDPI / 96.0f, VertDPI, VertDPI / 96.0f);
+
+					// There is a bug in the X11 RandR library in some versions of Linux that can return a bonkers scale.  Even on an 8K display,
+					// it's unlikely that we'd need more than a 4x scale, so capping there for now
+					if(Scale > 4.0f)
+						Scale = 1.0f;
+
 					return Scale;
 				}
 				else
