Fix: NullPlatformApplicationMisc platform check is flipped

diff --git a/Engine/Source/Runtime/ApplicationCore/Private/Null/NullPlatformApplicationMisc.cpp b/Engine/Source/Runtime/ApplicationCore/Private/Null/NullPlatformApplicationMisc.cpp
--- a/Engine/Source/Runtime/ApplicationCore/Private/Null/NullPlatformApplicationMisc.cpp
+++ b/Engine/Source/Runtime/ApplicationCore/Private/Null/NullPlatformApplicationMisc.cpp
@@ -15,7 +15,7 @@ bool FNullPlatformApplicationMisc::IsUsingNullApplication()
 	// NullPlatform Application only supported on Linux and Windows
 	// Theoretically it's supported on all platforms but we need to use the Pixel Streaming plugin to view
 	// the output, so we don't want to enable the null application on platforms that Pixel Streaming doesn't support
-#if !(PLATFORM_WINDOWS || PLATFORM_LINUX)
+#if PLATFORM_WINDOWS || PLATFORM_LINUX
 	static bool bIsRenderingOffScreen = FParse::Param(FCommandLine::Get(), TEXT("RenderOffScreen"));
 	return bIsRenderingOffScreen;
 #endif
