Now using alternate display dimensions which take into account the Dock size when computing work area.

diff --git a/Engine/Source/Runtime/ApplicationCore/Private/Linux/LinuxApplication.cpp b/Engine/Source/Runtime/ApplicationCore/Private/Linux/LinuxApplication.cpp
--- a/Engine/Source/Runtime/ApplicationCore/Private/Linux/LinuxApplication.cpp
+++ b/Engine/Source/Runtime/ApplicationCore/Private/Linux/LinuxApplication.cpp
@@ -1501,7 +1501,7 @@ void FLinuxApplication::SetHighPrecisionMouseMode( const bool Enable, const TSha
 void FLinuxApplication::RefreshDisplayCache()
 {
 	const double kCacheLifetime = 5.0;	// ask once in 5 seconds
-	
+
 	double CurrentTime = FPlatformTime::Seconds();
 	if (CurrentTime - LastTimeCachedDisplays > kCacheLifetime)
 	{
@@ -1511,10 +1511,10 @@ void FLinuxApplication::RefreshDisplayCache()
 
 		for (int DisplayIdx = 0; DisplayIdx < NumDisplays; ++DisplayIdx)
 		{
-			SDL_Rect DisplayBounds;
-			SDL_GetDisplayBounds(DisplayIdx, &DisplayBounds);
-			
-			CachedDisplays.Add(DisplayBounds);
+			SDL_Rect UsableBounds;
+			SDL_GetDisplayUsableBounds(DisplayIdx, &UsableBounds);
+
+			CachedDisplays.Add(UsableBounds);
 		}
 
 		LastTimeCachedDisplays = CurrentTime;
