Refactor a global variable and function into IOSAppDelegate as a recent commit appears to have broken it (plus, globals are (generally) bad)

#rnx

diff --git a/Engine/Source/Runtime/ApplicationCore/Private/IOS/IOSAppDelegate.cpp b/Engine/Source/Runtime/ApplicationCore/Private/IOS/IOSAppDelegate.cpp
--- a/Engine/Source/Runtime/ApplicationCore/Private/IOS/IOSAppDelegate.cpp
+++ b/Engine/Source/Runtime/ApplicationCore/Private/IOS/IOSAppDelegate.cpp
@@ -984,7 +984,8 @@ static FAutoConsoleVariableRef CVarGEnableThermalsReport(
 	self.bDeviceInPortraitMode = false;
 #else
 	// use the status bar orientation to properly determine landscape vs portrait
-	self.bDeviceInPortraitMode = UIInterfaceOrientationIsPortrait([self.Window.windowScene interfaceOrientation]);
+	self.InterfaceOrientation = [self.Window.windowScene interfaceOrientation];
+	self.bDeviceInPortraitMode = UIInterfaceOrientationIsPortrait(self.InterfaceOrientation);
 	printf("========= This app is in %s mode\n", self.bDeviceInPortraitMode ? "PORTRAIT" : "LANDSCAPE");
 #endif
 
@@ -1241,21 +1242,15 @@ static FAutoConsoleVariableRef CVarGEnableThermalsReport(
 	return false;
 }
 
-#if !PLATFORM_TVOS
-extern EDeviceScreenOrientation ConvertFromUIInterfaceOrientation(UIInterfaceOrientation Orientation);
-#endif
-
 - (void) didRotate:(NSNotification *)notification
 {   
 #if !PLATFORM_TVOS
-	// get the interfaec orientation
+	// get the interface orientation
 	
 	NSLog(@"didRotate orientation = %d", (int)[self.Window.windowScene interfaceOrientation]);
 	
     UIInterfaceOrientation Orientation = [self.Window.windowScene interfaceOrientation];
-	
-	extern UIInterfaceOrientation GInterfaceOrientation;
-	GInterfaceOrientation = Orientation;
+	self.InterfaceOrientation = Orientation;
 	
     if (bEngineInit)
     {
@@ -1263,7 +1258,7 @@ extern EDeviceScreenOrientation ConvertFromUIInterfaceOrientation(UIInterfaceOri
 		{
 			FIOSApplication* Application = [IOSAppDelegate GetDelegate].IOSApplication;
 			Application->OrientationChanged(Orientation);
-			FCoreDelegates::ApplicationReceivedScreenOrientationChangedNotificationDelegate.Broadcast((int32)ConvertFromUIInterfaceOrientation(Orientation));
+			FCoreDelegates::ApplicationReceivedScreenOrientationChangedNotificationDelegate.Broadcast((int32)[IOSAppDelegate ConvertFromUIInterfaceOrientation:Orientation]);
 
 			//we also want to fire off the safe frame event
 			FCoreDelegates::OnSafeFrameChangedEvent.Broadcast();
@@ -1588,6 +1583,19 @@ extern double GCStartTime;
 
 #if !PLATFORM_TVOS
 
++(EDeviceScreenOrientation) ConvertFromUIInterfaceOrientation:(UIInterfaceOrientation)Orientation
+{
+	switch(Orientation)
+	{
+		default:
+		case UIInterfaceOrientationUnknown : return EDeviceScreenOrientation::Unknown; break;
+		case UIInterfaceOrientationPortrait : return EDeviceScreenOrientation::Portrait; break;
+		case UIInterfaceOrientationPortraitUpsideDown : return EDeviceScreenOrientation::PortraitUpsideDown; break;
+		case UIInterfaceOrientationLandscapeLeft : return EDeviceScreenOrientation::LandscapeLeft; break;
+		case UIInterfaceOrientationLandscapeRight : return EDeviceScreenOrientation::LandscapeRight; break;
+	}
+}
+
 void HandleReceivedNotification(UNNotification* notification)
 {
 	bool IsLocal = false;
diff --git a/Engine/Source/Runtime/ApplicationCore/Public/IOS/IOSAppDelegate.h b/Engine/Source/Runtime/ApplicationCore/Public/IOS/IOSAppDelegate.h
--- a/Engine/Source/Runtime/ApplicationCore/Public/IOS/IOSAppDelegate.h
+++ b/Engine/Source/Runtime/ApplicationCore/Public/IOS/IOSAppDelegate.h
@@ -173,8 +173,12 @@ APPLICATIONCORE_API
 @property (retain) NSTimer* CommandLineParseTimer;
 @property (atomic) bool bCommandLineReady;
 
+#if !PLATFORM_TVOS
+@property (assign) UIInterfaceOrientation InterfaceOrientation;
+#endif
+
 /** initial launch options */
-@property(retain) NSDictionary* launchOptions;
+@property (retain) NSDictionary* launchOptions;
 
 @property (assign) NSProcessInfoThermalState ThermalState;
 @property (assign) bool bBatteryState;
@@ -210,6 +214,10 @@ APPLICATIONCORE_API
 -(float)GetBackgroundingMainThreadBlockTime;
 -(void)OverrideBackgroundingMainThreadBlockTime:(float)BlockTime;
 
+#if !PLATFORM_TVOS
+  +(EDeviceScreenOrientation) ConvertFromUIInterfaceOrientation:(UIInterfaceOrientation)Orientation;
+#endif
+
 -(bool)IsUpdateAvailable;
 
 @property (assign) bool bAudioSessionInitialized;
diff --git a/Engine/Source/Runtime/Core/Private/IOS/IOSPlatformMisc.cpp b/Engine/Source/Runtime/Core/Private/IOS/IOSPlatformMisc.cpp
--- a/Engine/Source/Runtime/Core/Private/IOS/IOSPlatformMisc.cpp
+++ b/Engine/Source/Runtime/Core/Private/IOS/IOSPlatformMisc.cpp
@@ -285,21 +285,6 @@ bool FIOSPlatformMisc::IsInLowPowerMode()
 }
 
 
-#if !PLATFORM_TVOS
-EDeviceScreenOrientation ConvertFromUIInterfaceOrientation(UIInterfaceOrientation Orientation)
-{
-	switch(Orientation)
-	{
-		default:
-		case UIInterfaceOrientationUnknown : return EDeviceScreenOrientation::Unknown; break;
-		case UIInterfaceOrientationPortrait : return EDeviceScreenOrientation::Portrait; break;
-		case UIInterfaceOrientationPortraitUpsideDown : return EDeviceScreenOrientation::PortraitUpsideDown; break;
-		case UIInterfaceOrientationLandscapeLeft : return EDeviceScreenOrientation::LandscapeLeft; break;
-		case UIInterfaceOrientationLandscapeRight : return EDeviceScreenOrientation::LandscapeRight; break;
-	}
-}
-#endif
-
 #if !PLATFORM_TVOS
 UIInterfaceOrientationMask GetUIInterfaceOrientationMask(EDeviceScreenOrientation ScreenOrientation)
 {
@@ -341,19 +326,16 @@ UIInterfaceOrientationMask GetUIInterfaceOrientationMask(EDeviceScreenOrientatio
 }
 #endif
 
-#if !PLATFORM_TVOS
-UIInterfaceOrientation GInterfaceOrientation = UIInterfaceOrientationUnknown;
-#endif
-
 EDeviceScreenOrientation FIOSPlatformMisc::GetDeviceOrientation()
 {
 #if !PLATFORM_TVOS
-	if (GInterfaceOrientation == UIInterfaceOrientationUnknown)
+	IOSAppDelegate* AppDelegate = [IOSAppDelegate GetDelegate];
+	if (AppDelegate.InterfaceOrientation == UIInterfaceOrientationUnknown)
 	{
-		GInterfaceOrientation = [[[[[UIApplication sharedApplication] delegate] window] windowScene] interfaceOrientation];
+		AppDelegate.InterfaceOrientation = [[[[[UIApplication sharedApplication] delegate] window] windowScene] interfaceOrientation];
 	}
 
-	return ConvertFromUIInterfaceOrientation(GInterfaceOrientation);
+	return [AppDelegate ConvertFromUIInterfaceOrientation:AppDelegate.InterfaceOrientation];
 #else
 	return EDeviceScreenOrientation::Unknown;
 #endif
