Stop wiping negative bits when masking away higher bits of window position data for use when moving windows

diff --git a/Engine/Source/Runtime/ApplicationCore/Private/Linux/LinuxApplication.cpp b/Engine/Source/Runtime/ApplicationCore/Private/Linux/LinuxApplication.cpp
--- a/Engine/Source/Runtime/ApplicationCore/Private/Linux/LinuxApplication.cpp
+++ b/Engine/Source/Runtime/ApplicationCore/Private/Linux/LinuxApplication.cpp
@@ -794,10 +794,27 @@ void FLinuxApplication::ProcessDeferredMessage( SDL_Event Event )
 
 				case SDL_WINDOWEVENT_MOVED:
 					{
-						// Mask away the higher bits, as SDL uses those as flags
+						// Mask away the higher bits (but preserve negative flags), as SDL uses those as flags
 						// See: SDL_WINDOWPOS_UNDEFINED_MASK & SDL_WINDOWPOS_CENTERED_MASK for context
-						int32 ClientScreenX = windowEvent.data1 & 0xFFFF;
-						int32 ClientScreenY = windowEvent.data2 & 0xFFFF;
+						int32 ClientScreenX;
+						int32 ClientScreenY;
+
+						if (windowEvent.data1 < 0)
+						{
+							ClientScreenX = -(-windowEvent.data1 & 0xFFFF);
+						}
+						else
+						{
+							ClientScreenX = windowEvent.data1 & 0xFFFF;
+						}
+						if (windowEvent.data2 < 0)
+						{
+							ClientScreenY = -(-windowEvent.data2 & 0xFFFF);
+						}
+						else
+						{
+							ClientScreenY = windowEvent.data2 & 0xFFFF;
+						}
 
 						int32 BorderSizeX, BorderSizeY;
 						CurrentEventWindow->GetNativeBordersSize(BorderSizeX, BorderSizeY);
