Only update stored console window width/size in SIZE_RESTORED and SIZE_MAXIMIZED

* Fixes a bug where invalid window dimensions would be stored in the .ini file if application is closed while the console is minimized

diff --git a/Engine/Source/Runtime/ApplicationCore/Private/Windows/WindowsConsoleOutputDevice2.cpp b/Engine/Source/Runtime/ApplicationCore/Private/Windows/WindowsConsoleOutputDevice2.cpp
--- a/Engine/Source/Runtime/ApplicationCore/Private/Windows/WindowsConsoleOutputDevice2.cpp
+++ b/Engine/Source/Runtime/ApplicationCore/Private/Windows/WindowsConsoleOutputDevice2.cpp
@@ -1793,11 +1793,19 @@ public:
 
 		case WM_SIZE:
 		{
-			UpdateSize(LOWORD(lParam), HIWORD(lParam), true);
-			RECT WindowRect;
-			GetWindowRect(MainHwnd, &WindowRect);
-			ConsoleWidth = WindowRect.right - WindowRect.left;
-			ConsoleHeight = WindowRect.bottom - WindowRect.top;
+			switch (wParam)
+			{
+			case SIZE_MAXIMIZED:
+			case SIZE_RESTORED:
+				UpdateSize(LOWORD(lParam), HIWORD(lParam), true);
+				RECT WindowRect;
+				GetWindowRect(MainHwnd, &WindowRect);
+				ConsoleWidth = WindowRect.right - WindowRect.left;
+				ConsoleHeight = WindowRect.bottom - WindowRect.top;
+			default:
+				// ignore all other messages, such as minimize
+				break;
+			}
 			return 0;
 		}
 		case WM_MOVE:
