Mac: Don't rebuild cached menu on every key press.

diff --git a/Engine/Source/Runtime/ApplicationCore/Private/Mac/MacApplication.cpp b/Engine/Source/Runtime/ApplicationCore/Private/Mac/MacApplication.cpp
--- a/Engine/Source/Runtime/ApplicationCore/Private/Mac/MacApplication.cpp
+++ b/Engine/Source/Runtime/ApplicationCore/Private/Mac/MacApplication.cpp
@@ -1182,7 +1182,6 @@ void FMacApplication::ProcessKeyUpEvent(const FDeferredMacEvent& Event)
 	{
 		ResendEvent(Event.Event);
 	}
-	FPlatformApplicationMisc::bChachedMacMenuStateNeedsUpdate = true;
 }
 
 void FMacApplication::OnWindowDidMove(TSharedRef<FMacWindow> Window)
