Fix shutdown crash when using AddressSanitizer. The crash happens in AtExit deep in external DLLs, so instead disable TextInput altogether.

#rnx

diff --git a/Engine/Source/Runtime/ApplicationCore/Private/Windows/WindowsApplication.cpp b/Engine/Source/Runtime/ApplicationCore/Private/Windows/WindowsApplication.cpp
--- a/Engine/Source/Runtime/ApplicationCore/Private/Windows/WindowsApplication.cpp
+++ b/Engine/Source/Runtime/ApplicationCore/Private/Windows/WindowsApplication.cpp
@@ -151,11 +151,15 @@ FWindowsApplication::FWindowsApplication( const HINSTANCE HInstance, const HICON
 	CA_SUPPRESS(6031);
 	OleInitialize( NULL );
 
+#if !USING_ADDRESS_SANITISER
+	// Disabled under AddressSanitizer as otherwise AddresSanitizer catches an issue in the cleanup code
+	// of one of the input method Microsoft DLLs during shutdown
 	TextInputMethodSystem = MakeShareable( new FWindowsTextInputMethodSystem );
 	if(!TextInputMethodSystem->Initialize())
 	{
 		TextInputMethodSystem.Reset();
 	}
+#endif
 
 	TaskbarList = FTaskbarList::Create();
 
