Prevent activating context when HWND is destroyed.
bug: we call activate context immediately after the window is destroyed, and the window pointer is stilll not null. We have also destroyed HWND at this point through the microsoft function ::DestroyWindow, so we check ::IsWindowEnabled as a way to confirm the validity of HWND.

diff --git a/Engine/Source/Runtime/ApplicationCore/Private/Windows/WindowsTextInputMethodSystem.cpp b/Engine/Source/Runtime/ApplicationCore/Private/Windows/WindowsTextInputMethodSystem.cpp
--- a/Engine/Source/Runtime/ApplicationCore/Private/Windows/WindowsTextInputMethodSystem.cpp
+++ b/Engine/Source/Runtime/ApplicationCore/Private/Windows/WindowsTextInputMethodSystem.cpp
@@ -851,7 +851,7 @@ void FWindowsTextInputMethodSystem::ActivateContext(const TSharedRef<ITextInputM
 	const TSharedPtr<FGenericWindow> GenericWindow = Context->GetWindow();
 	InternalContext.WindowHandle = GenericWindow.IsValid() ? reinterpret_cast<HWND>(GenericWindow->GetOSWindowHandle()) : nullptr;
 
-	if (InternalContext.WindowHandle)
+	if (InternalContext.WindowHandle && ::IsWindowEnabled(InternalContext.WindowHandle))
 	{
 		// IMM Implementation
 		InternalContext.IMMContext.IsComposing = false;
