Changed how configrules safezones values are evaluated and applied.

#tests Tested on the following affected devices Pixel 3 XL, Nokia 8, OnePlus 7, OnePlus 8T, Samsung S8+, Samsung S10

#rnx

diff --git a/Engine/Source/Runtime/ApplicationCore/Private/Android/AndroidApplication.cpp b/Engine/Source/Runtime/ApplicationCore/Private/Android/AndroidApplication.cpp
--- a/Engine/Source/Runtime/ApplicationCore/Private/Android/AndroidApplication.cpp
+++ b/Engine/Source/Runtime/ApplicationCore/Private/Android/AndroidApplication.cpp
@@ -157,39 +157,28 @@ void FDisplayMetrics::RebuildDisplayMetrics( FDisplayMetrics& OutDisplayMetrics
 	// Apply the debug safe zones
 	OutDisplayMetrics.ApplyDefaultSafeZones();
 
-	double Inset_Left = -1.0f;
-	double Inset_Top = -1.0f;
-	double Inset_Right = -1.0f;
-	double Inset_Bottom = -1.0f;
-
 	bool bIsPortrait = FAndroidWindow::IsPortraitOrientation();
+	FVector4 SafeZoneRect = FAndroidWindow::GetSafezone(bIsPortrait);
 
 	// ConfigRules values override values from device
 	if (FString* SafeZoneVar = FAndroidMisc::GetConfigRulesVariable(bIsPortrait ? TEXT("SafeZone_Portrait") : TEXT("SafeZone_Landscape")))
 	{
 		TArray<FString> ZoneVector;
+		ZoneVector.Reserve(4);
+
 		int ZoneParseCount = SafeZoneVar->ParseIntoArray(ZoneVector, TEXT(","), true);
 		ensureMsgf(ZoneParseCount == 4, TEXT("SafeZone variable not properly formatted."));
 
 		if (ZoneParseCount == 4)
 		{
-			// these are already in pixels
-			Inset_Left = FCString::Atof(*ZoneVector[0]);
-			Inset_Top = FCString::Atof(*ZoneVector[1]);
-			Inset_Right = FCString::Atof(*ZoneVector[2]);
-			Inset_Bottom = FCString::Atof(*ZoneVector[3]);
+			SafeZoneRect = FVector4(FCString::Atof(*ZoneVector[0]), FCString::Atof(*ZoneVector[1]), FCString::Atof(*ZoneVector[2]), FCString::Atof(*ZoneVector[3]));
 		}
 	}
-	else
-	{
-		FVector4 SafeZoneRect = FAndroidWindow::GetSafezone(bIsPortrait);
 
-		// These values will be negative if there is not a safe zone set by GameActivity and need scaling
-		Inset_Left = SafeZoneRect.X * OutDisplayMetrics.PrimaryDisplayWidth;
-		Inset_Top = SafeZoneRect.Y * OutDisplayMetrics.PrimaryDisplayHeight;
-		Inset_Right = SafeZoneRect.Z * OutDisplayMetrics.PrimaryDisplayWidth;
-		Inset_Bottom = SafeZoneRect.W * OutDisplayMetrics.PrimaryDisplayHeight;
-	}
+	const double Inset_Left = SafeZoneRect.X > 0 ? SafeZoneRect.X * OutDisplayMetrics.PrimaryDisplayWidth : -1.0f;
+	const double Inset_Top = SafeZoneRect.Y > 0 ? SafeZoneRect.Y * OutDisplayMetrics.PrimaryDisplayHeight : -1.0f;
+	const double Inset_Right =  SafeZoneRect.Z > 0 ? SafeZoneRect.Z * OutDisplayMetrics.PrimaryDisplayWidth : -1.0f;
+	const double Inset_Bottom = SafeZoneRect.W > 0 ? SafeZoneRect.W * OutDisplayMetrics.PrimaryDisplayHeight : -1.0f;
 
 	OutDisplayMetrics.TitleSafePaddingSize.X = (Inset_Left >= 0.0f) ? Inset_Left : OutDisplayMetrics.TitleSafePaddingSize.X;
 	OutDisplayMetrics.TitleSafePaddingSize.Y = (Inset_Top >= 0.0f) ? Inset_Top : OutDisplayMetrics.TitleSafePaddingSize.Y;
