[Engine]

diff --git a/Engine/Source/Runtime/ApplicationCore/Public/GenericPlatform/Accessibility/GenericAccessibleInterfaces.h b/Engine/Source/Runtime/ApplicationCore/Public/GenericPlatform/Accessibility/GenericAccessibleInterfaces.h
--- a/Engine/Source/Runtime/ApplicationCore/Public/GenericPlatform/Accessibility/GenericAccessibleInterfaces.h
+++ b/Engine/Source/Runtime/ApplicationCore/Public/GenericPlatform/Accessibility/GenericAccessibleInterfaces.h
@@ -107,7 +107,7 @@ typedef int32 FAccessibleUserIndex;
  * All accessible users must be registered with an FGenericAccessibleUserRegistry for the users to be retrieved and interacted with.
  * @see FGenericAccessibleUserRegistry
  */
-class APPLICATIONCORE_API FGenericAccessibleUser
+class FGenericAccessibleUser
 {
 	friend class FGenericAccessibleUserRegistry;
 public:
@@ -177,7 +177,7 @@ private:
  * Users can subclass this class to provide additional functionality for their custom use cases.
  * @see FGenericAccessibleUser, FGenericAccessibleMessageHandler
  */
-class APPLICATIONCORE_API FGenericAccessibleUserRegistry
+class FGenericAccessibleUserRegistry
 {
 public:
 	virtual ~FGenericAccessibleUserRegistry() = default;
@@ -190,7 +190,7 @@ public:
 	 * @param User An accessible user you want to register with the accessible user registry.
 	 * @return True if the passed in user was successfully registered. Else, returns false.
 	 */
-	bool RegisterUser(const TSharedRef<FGenericAccessibleUser>& User);
+	APPLICATIONCORE_API bool RegisterUser(const TSharedRef<FGenericAccessibleUser>& User);
 	/**
 	 * Unregisters an accessible user from the accessible registry.
 	 * If the passed in user index is not associated with a registered accessible user, nothing will happen.
@@ -199,20 +199,20 @@ public:
 	 * @param UserIndex The index of the accessible user you want to unregister.
 	 * @return True if the accessible user with the associated user index is successfully unregistered. Else, returns false.
 	 */
-	bool UnregisterUser(const FAccessibleUserIndex UserIndex);
+	APPLICATIONCORE_API bool UnregisterUser(const FAccessibleUserIndex UserIndex);
 	/**
 	 * Unregisters all accessible users from the accessible registry.
 	 * This function calls FGenericAccessibleUser::OnUnregistered when each registered user is successfully unregistered.
 	 */
-	void UnregisterAllUsers();
+	APPLICATIONCORE_API void UnregisterAllUsers();
 	/** Returns true if the passed in user index is associated with a registered accessible user. Else, returns false. */
-	bool IsUserRegistered(const FAccessibleUserIndex UserIndex) const;
+	APPLICATIONCORE_API bool IsUserRegistered(const FAccessibleUserIndex UserIndex) const;
 	/** Returns the accessible user associated with the passed in user index. If the passed in user index is not associated with any accessible user, nullptr is returned. */
-	TSharedPtr<FGenericAccessibleUser> GetUser(const FAccessibleUserIndex UserIndex) const;
+	APPLICATIONCORE_API TSharedPtr<FGenericAccessibleUser> GetUser(const FAccessibleUserIndex UserIndex) const;
 	/** Returns the number of accessible users registered with the accessible user registry. */
-	int32 GetNumberofUsers() const;
+	APPLICATIONCORE_API int32 GetNumberofUsers() const;
 	/** Returns an array of all accessible users that are currently registered with the accessible registry. */
-	TArray<TSharedRef<FGenericAccessibleUser>> GetAllUsers() const;
+	APPLICATIONCORE_API TArray<TSharedRef<FGenericAccessibleUser>> GetAllUsers() const;
 	/**
 	 * Returns the index of the primary accessible user.
 	 * The primary accessible user should correspond to the default input source that every application should have. This user should also be the primary cursor user.
@@ -830,7 +830,7 @@ public:
  * It is up to the client to create an instance of this struct and fill in the appropriate data members.
  * @see FGenericAccessibleMessageHandler::RaiseEvent
  */
-struct APPLICATIONCORE_API FAccessibleEventArgs
+struct FAccessibleEventArgs
 {
 	FAccessibleEventArgs(TSharedRef<IAccessibleWidget> InWidget, EAccessibleEvent InEvent, FVariant InOldValue = FVariant(), FVariant InNewValue = FVariant(), FAccessibleUserIndex InUserIndex = 0)
 		: Widget(InWidget)
@@ -867,7 +867,7 @@ struct APPLICATIONCORE_API FAccessibleEventArgs
  * Callers should also use GetAccessibleUserRegistry() to register and interact with accessible users. Accessible users must be registered
  * with the set accessible user registry for the rest of the application to be able to retrieve and interact with the accessible users.
  */
-class APPLICATIONCORE_API FGenericAccessibleMessageHandler
+class FGenericAccessibleMessageHandler
 {
 public:
 	/**
@@ -876,7 +876,7 @@ public:
 	 */
 	DECLARE_DELEGATE_OneParam(FAccessibleEvent, const FAccessibleEventArgs&);
 
-	FGenericAccessibleMessageHandler();
+	APPLICATIONCORE_API FGenericAccessibleMessageHandler();
 	virtual ~FGenericAccessibleMessageHandler()
 	{
 		UnbindAccessibleEventDelegate();
@@ -887,7 +887,7 @@ public:
 	 *
 	 * @return true if the application intends to return valid accessible widgets when queried.
 	 */
-	bool ApplicationIsAccessible() const;
+	APPLICATIONCORE_API bool ApplicationIsAccessible() const;
 
 	/**
 	 * Checks if accessibility is enabled in the application. Usually this happens when screen-reading software is turned on.
@@ -902,7 +902,7 @@ public:
 	 *
 	 * @param bActive Whether to enable to disable the message handler.
 	 */
-	void SetActive(bool bActive);
+	APPLICATIONCORE_API void SetActive(bool bActive);
 
 	/**
 	 * Creates or retrieves an accessible object for a native OS window.
diff --git a/Engine/Source/Runtime/ApplicationCore/Public/GenericPlatform/GenericApplication.h b/Engine/Source/Runtime/ApplicationCore/Public/GenericPlatform/GenericApplication.h
--- a/Engine/Source/Runtime/ApplicationCore/Public/GenericPlatform/GenericApplication.h
+++ b/Engine/Source/Runtime/ApplicationCore/Public/GenericPlatform/GenericApplication.h
@@ -395,7 +395,7 @@ struct FDisplayMetrics
 	FVector4 ActionSafePaddingSize;
 
 	UE_DEPRECATED(4.21, "Please use RebuildDisplayMetrics - it is functionally the same but is clearer about the function cost")
-	APPLICATIONCORE_API static void GetDisplayMetrics(struct FDisplayMetrics& OutDisplayMetrics) { RebuildDisplayMetrics(OutDisplayMetrics); };
+	static void GetDisplayMetrics(struct FDisplayMetrics& OutDisplayMetrics) { RebuildDisplayMetrics(OutDisplayMetrics); };
 
 	APPLICATIONCORE_API static void RebuildDisplayMetrics(struct FDisplayMetrics& OutDisplayMetrics);
 
diff --git a/Engine/Source/Runtime/ApplicationCore/Public/GenericPlatform/GenericApplicationMessageHandler.h b/Engine/Source/Runtime/ApplicationCore/Public/GenericPlatform/GenericApplicationMessageHandler.h
--- a/Engine/Source/Runtime/ApplicationCore/Public/GenericPlatform/GenericApplicationMessageHandler.h
+++ b/Engine/Source/Runtime/ApplicationCore/Public/GenericPlatform/GenericApplicationMessageHandler.h
@@ -26,47 +26,47 @@ namespace EMouseButtons
 	};
 }
 
-struct APPLICATIONCORE_API FGamepadKeyNames
+struct FGamepadKeyNames
 {
 	typedef FName Type;
 
-	static const FName Invalid;
-
-	static const FName LeftAnalogX;
-	static const FName LeftAnalogY;
-	static const FName RightAnalogX;
-	static const FName RightAnalogY;
-	static const FName LeftTriggerAnalog;
-	static const FName RightTriggerAnalog;
-
-	static const FName LeftThumb;
-	static const FName RightThumb;
-	static const FName SpecialLeft;
-	static const FName SpecialLeft_X;
-	static const FName SpecialLeft_Y;
-	static const FName SpecialRight;
-	static const FName FaceButtonBottom;
-	static const FName FaceButtonRight;
-	static const FName FaceButtonLeft;
-	static const FName FaceButtonTop;
-	static const FName LeftShoulder;
-	static const FName RightShoulder;
-	static const FName LeftTriggerThreshold;
-	static const FName RightTriggerThreshold;
-	static const FName DPadUp;
-	static const FName DPadDown;
-	static const FName DPadRight;
-	static const FName DPadLeft;
-
-	static const FName LeftStickUp;
-	static const FName LeftStickDown;
-	static const FName LeftStickRight;
-	static const FName LeftStickLeft;
-
-	static const FName RightStickUp;
-	static const FName RightStickDown;
-	static const FName RightStickRight;
-	static const FName RightStickLeft;
+	static APPLICATIONCORE_API const FName Invalid;
+
+	static APPLICATIONCORE_API const FName LeftAnalogX;
+	static APPLICATIONCORE_API const FName LeftAnalogY;
+	static APPLICATIONCORE_API const FName RightAnalogX;
+	static APPLICATIONCORE_API const FName RightAnalogY;
+	static APPLICATIONCORE_API const FName LeftTriggerAnalog;
+	static APPLICATIONCORE_API const FName RightTriggerAnalog;
+
+	static APPLICATIONCORE_API const FName LeftThumb;
+	static APPLICATIONCORE_API const FName RightThumb;
+	static APPLICATIONCORE_API const FName SpecialLeft;
+	static APPLICATIONCORE_API const FName SpecialLeft_X;
+	static APPLICATIONCORE_API const FName SpecialLeft_Y;
+	static APPLICATIONCORE_API const FName SpecialRight;
+	static APPLICATIONCORE_API const FName FaceButtonBottom;
+	static APPLICATIONCORE_API const FName FaceButtonRight;
+	static APPLICATIONCORE_API const FName FaceButtonLeft;
+	static APPLICATIONCORE_API const FName FaceButtonTop;
+	static APPLICATIONCORE_API const FName LeftShoulder;
+	static APPLICATIONCORE_API const FName RightShoulder;
+	static APPLICATIONCORE_API const FName LeftTriggerThreshold;
+	static APPLICATIONCORE_API const FName RightTriggerThreshold;
+	static APPLICATIONCORE_API const FName DPadUp;
+	static APPLICATIONCORE_API const FName DPadDown;
+	static APPLICATIONCORE_API const FName DPadRight;
+	static APPLICATIONCORE_API const FName DPadLeft;
+
+	static APPLICATIONCORE_API const FName LeftStickUp;
+	static APPLICATIONCORE_API const FName LeftStickDown;
+	static APPLICATIONCORE_API const FName LeftStickRight;
+	static APPLICATIONCORE_API const FName LeftStickLeft;
+
+	static APPLICATIONCORE_API const FName RightStickUp;
+	static APPLICATIONCORE_API const FName RightStickDown;
+	static APPLICATIONCORE_API const FName RightStickRight;
+	static APPLICATIONCORE_API const FName RightStickLeft;
 };
 
 enum class EWindowActivation : uint8
@@ -173,7 +173,7 @@ private:
  * This can be used to determine hardware-specific information when handling input from FGenericApplicationMessageHandler subclasses.
  * This is generally set during SendControllerEvents or Tick and is only valid on the game thread.
  */
-class APPLICATIONCORE_API FInputDeviceScope
+class FInputDeviceScope
 {
 public:
 	/** The specific InputDevice that is currently being polled. This is only valid within the current function scope and may be null */
@@ -189,8 +189,8 @@ public:
 	FString HardwareDeviceIdentifier;
 
 	/** Constructor, this should only be allocated directly on the stack */
-	FInputDeviceScope(IInputDevice* InInputDevice, FName InInputDeviceName, int32 InHardwareDeviceHandle = -1, FString InHardwareDeviceIdentifier = FString());
-	~FInputDeviceScope();
+	APPLICATIONCORE_API FInputDeviceScope(IInputDevice* InInputDevice, FName InInputDeviceName, int32 InHardwareDeviceHandle = -1, FString InHardwareDeviceIdentifier = FString());
+	APPLICATIONCORE_API ~FInputDeviceScope();
 
 	/** Cannot be copied/moved */
 	FInputDeviceScope() = delete;
@@ -200,10 +200,10 @@ public:
 	FInputDeviceScope& operator=(FInputDeviceScope&&) = delete;
 
 	/** Returns the currently active InputDeviceScope. This is only valid to call on the game thread and may return null */
-	static const FInputDeviceScope* GetCurrent();
+	static APPLICATIONCORE_API const FInputDeviceScope* GetCurrent();
 
 private:
-	static TArray<FInputDeviceScope*> ScopeStack;
+	static APPLICATIONCORE_API TArray<FInputDeviceScope*> ScopeStack;
 };
 
 PRAGMA_DISABLE_DEPRECATION_WARNINGS
@@ -627,4 +627,4 @@ public:
 	}
 };
 
-PRAGMA_ENABLE_DEPRECATION_WARNINGS
\ No newline at end of file
+PRAGMA_ENABLE_DEPRECATION_WARNINGS
diff --git a/Engine/Source/Runtime/ApplicationCore/Public/GenericPlatform/GenericPlatformApplicationMisc.h b/Engine/Source/Runtime/ApplicationCore/Public/GenericPlatform/GenericPlatformApplicationMisc.h
--- a/Engine/Source/Runtime/ApplicationCore/Public/GenericPlatform/GenericPlatformApplicationMisc.h
+++ b/Engine/Source/Runtime/ApplicationCore/Public/GenericPlatform/GenericPlatformApplicationMisc.h
@@ -66,15 +66,15 @@ enum class EPlatformUserSelectorFlags : uint8
 ENUM_CLASS_FLAGS(EPlatformUserSelectorFlags);
 
 
-struct APPLICATIONCORE_API FGenericPlatformApplicationMisc
+struct FGenericPlatformApplicationMisc
 {
-	static void PreInit();
+	static APPLICATIONCORE_API void PreInit();
 
-	static void Init();
+	static APPLICATIONCORE_API void Init();
 
-	static void PostInit();
+	static APPLICATIONCORE_API void PostInit();
 
-	static void TearDown();
+	static APPLICATIONCORE_API void TearDown();
 
 	/**
 	 * Load the preinit modules required by this platform, typically they are the renderer modules
@@ -93,38 +93,38 @@ struct APPLICATIONCORE_API FGenericPlatformApplicationMisc
 	/**
 	 * Creates a console output device for this platform. Should only be called once.
 	 */
-	static FOutputDeviceConsole* CreateConsoleOutputDevice();
+	static APPLICATIONCORE_API FOutputDeviceConsole* CreateConsoleOutputDevice();
 
 	/**
 	 * Gets a pointer to the platform error output device singleton.
 	 */
-	static FOutputDeviceError* GetErrorOutputDevice();
+	static APPLICATIONCORE_API FOutputDeviceError* GetErrorOutputDevice();
 
 	/**
 	 * Gets a pointer to the default platform feedback context implementation.
 	 */
-	static FFeedbackContext* GetFeedbackContext();
+	static APPLICATIONCORE_API FFeedbackContext* GetFeedbackContext();
 
 	/**
 	 * Gets a pointer to the default platform input device manager.
 	 */
-	static IPlatformInputDeviceMapper* CreatePlatformInputDeviceManager();
+	static APPLICATIONCORE_API IPlatformInputDeviceMapper* CreatePlatformInputDeviceManager();
 
 	/**
 	 * Creates an application instance.
 	 */
-	static class GenericApplication* CreateApplication();
+	static APPLICATIONCORE_API class GenericApplication* CreateApplication();
 
 	/** Request application to minimize (goto background). **/
-	static void RequestMinimize();
+	static APPLICATIONCORE_API void RequestMinimize();
 
 	/** Returns true if the specified application has a visible window, and that window is active/has focus/is selected */
-	static bool IsThisApplicationForeground();	
+	static APPLICATIONCORE_API bool IsThisApplicationForeground();	
 
 	/**
 	* Returns whether the platform wants to use a touch screen for a virtual keyboard.
 	*/
-	static bool RequiresVirtualKeyboard();
+	static APPLICATIONCORE_API bool RequiresVirtualKeyboard();
 
 	/**
 	 *	Pumps Windows messages.
@@ -177,7 +177,7 @@ struct APPLICATIONCORE_API FGenericPlatformApplicationMisc
 	 * @param	InGamma			Optional gamma correction to apply to the screen color
 	 * @return					The color of the pixel displayed at the chosen location
 	 */
-	static struct FLinearColor GetScreenPixelColor(const FVector2D& InScreenPos, float InGamma = 1.0f);
+	static APPLICATIONCORE_API struct FLinearColor GetScreenPixelColor(const FVector2D& InScreenPos, float InGamma = 1.0f);
 
 	/**
 	 * Searches for a window that matches the window name or the title starts with a particular text. When
@@ -213,7 +213,7 @@ struct APPLICATIONCORE_API FGenericPlatformApplicationMisc
 
 	/** @return true if the application is high dpi aware */
 
-	static bool IsHighDPIAwarenessEnabled();
+	static APPLICATIONCORE_API bool IsHighDPIAwarenessEnabled();
 
 	/*
 	 * UE expects mouse coordinates in screen space. Some platforms provides in client space. 
@@ -292,20 +292,20 @@ struct APPLICATIONCORE_API FGenericPlatformApplicationMisc
 	}
 				
 	/** Copies text to the operating system clipboard. */
-	static void ClipboardCopy(const TCHAR* Str);
+	static APPLICATIONCORE_API void ClipboardCopy(const TCHAR* Str);
 
 	/** Pastes in text from the operating system clipboard. */
-	static void ClipboardPaste(class FString& Dest);
+	static APPLICATIONCORE_API void ClipboardPaste(class FString& Dest);
 
 	/**
 	 * Gets the physical size of the screen if possible.  Some platforms lie, some platforms don't know.
 	 */
-	static EScreenPhysicalAccuracy GetPhysicalScreenDensity(int32& OutScreenDensity);
+	static APPLICATIONCORE_API EScreenPhysicalAccuracy GetPhysicalScreenDensity(int32& OutScreenDensity);
 
 	/**
 	 * Gets the physical size of the screen if possible.  Some platforms lie, some platforms don't know.
 	 */
-	static EScreenPhysicalAccuracy ComputePhysicalScreenDensity(int32& OutScreenDensity);
+	static APPLICATIONCORE_API EScreenPhysicalAccuracy ComputePhysicalScreenDensity(int32& OutScreenDensity);
 
 	/**
 	 * If we know or can approximate the pixel density of the screen we will convert the incoming inches
@@ -382,10 +382,10 @@ struct APPLICATIONCORE_API FGenericPlatformApplicationMisc
 
 
 protected:
-	static bool CachedPhysicalScreenData;
-	static EScreenPhysicalAccuracy CachedPhysicalScreenAccuracy;
-	static int32 CachedPhysicalScreenDensity;
-	static FAutoConsoleVariableRef CVarEnableHighDPIAwareness;
-	static FAutoConsoleVariableRef CVarAllowVirtualKeyboard;
+	static APPLICATIONCORE_API bool CachedPhysicalScreenData;
+	static APPLICATIONCORE_API EScreenPhysicalAccuracy CachedPhysicalScreenAccuracy;
+	static APPLICATIONCORE_API int32 CachedPhysicalScreenDensity;
+	static APPLICATIONCORE_API FAutoConsoleVariableRef CVarEnableHighDPIAwareness;
+	static APPLICATIONCORE_API FAutoConsoleVariableRef CVarAllowVirtualKeyboard;
 
 };
diff --git a/Engine/Source/Runtime/ApplicationCore/Public/GenericPlatform/GenericPlatformFramePacer.h b/Engine/Source/Runtime/ApplicationCore/Public/GenericPlatform/GenericPlatformFramePacer.h
--- a/Engine/Source/Runtime/ApplicationCore/Public/GenericPlatform/GenericPlatformFramePacer.h
+++ b/Engine/Source/Runtime/ApplicationCore/Public/GenericPlatform/GenericPlatformFramePacer.h
@@ -14,7 +14,7 @@
 /**
  * Generic implementation for most platforms
  **/
-struct APPLICATIONCORE_API FGenericPlatformRHIFramePacer
+struct FGenericPlatformRHIFramePacer
 {
 	/**
 	 * The pace we are running at (30 = 30fps, 0 = unpaced)
@@ -39,18 +39,18 @@ struct APPLICATIONCORE_API FGenericPlatformRHIFramePacer
 	/**
 	 * Returns whether the hardware is able to frame pace at the specified frame rate
 	 */
-	static bool SupportsFramePace(int32 QueryFramePace);
+	static APPLICATIONCORE_API bool SupportsFramePace(int32 QueryFramePace);
 
 protected:
 	/**
 	 * The generic implementation returns a result based on rhi.SyncInterval and FPlatformMisc::GetMaxRefreshRate().
 	 */
-	static int32 GetFramePaceFromSyncInterval();
+	static APPLICATIONCORE_API int32 GetFramePaceFromSyncInterval();
 
 	/**
 	 * The generic implementation sets rhi.SyncInterval based on FPlatformMisc::GetMaxRefreshRate().
 	 *
 	 * @return the pace we will run at.
 	 */
-	static int32 SetFramePaceToSyncInterval(int32 FramePace);
+	static APPLICATIONCORE_API int32 SetFramePaceToSyncInterval(int32 FramePace);
 };
diff --git a/Engine/Source/Runtime/ApplicationCore/Public/GenericPlatform/GenericPlatformInputDeviceMapper.h b/Engine/Source/Runtime/ApplicationCore/Public/GenericPlatform/GenericPlatformInputDeviceMapper.h
--- a/Engine/Source/Runtime/ApplicationCore/Public/GenericPlatform/GenericPlatformInputDeviceMapper.h
+++ b/Engine/Source/Runtime/ApplicationCore/Public/GenericPlatform/GenericPlatformInputDeviceMapper.h
@@ -22,17 +22,17 @@ DECLARE_LOG_CATEGORY_EXTERN(LogInputDeviceMapper, Log, All);
  * 
  * @note Keep any function comments up to date with those in PlatformInputDeviceMapperLibrary.h!
  */
-class APPLICATIONCORE_API IPlatformInputDeviceMapper
+class IPlatformInputDeviceMapper
 {
 public:
 	
 	/** Get the platform input device mapper */
-	static IPlatformInputDeviceMapper& Get();
+	static APPLICATIONCORE_API IPlatformInputDeviceMapper& Get();
 	
-	IPlatformInputDeviceMapper();
+	APPLICATIONCORE_API IPlatformInputDeviceMapper();
 	
 	/** Virtual destructor */
-	virtual ~IPlatformInputDeviceMapper();
+	APPLICATIONCORE_API virtual ~IPlatformInputDeviceMapper();
 	
 	/**
 	 * Populates the OutInputDevices array with any InputDeviceID's that are mapped to the given platform user
@@ -41,7 +41,7 @@ public:
 	 * @param OutInputDevices		Array of input device ID's that will be populated with the mapped devices.
 	 * @return						The number of mapped devices, INDEX_NONE if the user was not found.
 	 */
-	virtual int32 GetAllInputDevicesForUser(const FPlatformUserId UserId, TArray<FInputDeviceId>& OutInputDevices) const;
+	APPLICATIONCORE_API virtual int32 GetAllInputDevicesForUser(const FPlatformUserId UserId, TArray<FInputDeviceId>& OutInputDevices) const;
 
 	/**
 	 * Get all mapped input devices on this platform regardless of their connection state.
@@ -49,7 +49,7 @@ public:
 	 * @param OutInputDevices	Array of input devices to populate
 	 * @return					The number of connected input devices
 	 */
-	virtual int32 GetAllInputDevices(TArray<FInputDeviceId>& OutInputDevices) const;
+	APPLICATIONCORE_API virtual int32 GetAllInputDevices(TArray<FInputDeviceId>& OutInputDevices) const;
 	
 	/**
 	 * Gather all currently connected input devices
@@ -57,7 +57,7 @@ public:
 	 * @param OutInputDevices	Array of input devices to populate
 	 * @return					The number of connected input devices
 	 */
-	virtual int32 GetAllConnectedInputDevices(TArray<FInputDeviceId>& OutInputDevices) const;
+	APPLICATIONCORE_API virtual int32 GetAllConnectedInputDevices(TArray<FInputDeviceId>& OutInputDevices) const;
 
 	/**
 	 * Get all currently active platform ids, anyone who has a mapped input device
@@ -65,7 +65,7 @@ public:
 	 * @param OutUsers		Array that will be populated with the platform users.
 	 * @return				The number of active platform users
 	 */
-	virtual int32 GetAllActiveUsers(TArray<FPlatformUserId>& OutUsers) const;
+	APPLICATIONCORE_API virtual int32 GetAllActiveUsers(TArray<FPlatformUserId>& OutUsers) const;
 
 	/**
 	 * Returns the platform user id that is being used for unmapped input devices.
@@ -74,10 +74,10 @@ public:
 	virtual FPlatformUserId GetUserForUnpairedInputDevices() const = 0;
 
 	/** Returns true if the given Platform User Id is the user for unpaired input devices on this platform. */
-	virtual bool IsUnpairedUserId(const FPlatformUserId PlatformId) const;
+	APPLICATIONCORE_API virtual bool IsUnpairedUserId(const FPlatformUserId PlatformId) const;
 
 	/** Returns true if the given input device is mapped to the unpaired platform user id. */
-	virtual bool IsInputDeviceMappedToUnpairedUser(const FInputDeviceId InputDevice) const;
+	APPLICATIONCORE_API virtual bool IsInputDeviceMappedToUnpairedUser(const FInputDeviceId InputDevice) const;
 
 	/**
 	 * Returns the 'Primary' Platform user for this platform.
@@ -91,10 +91,10 @@ public:
 	virtual FInputDeviceId GetDefaultInputDevice() const = 0;
 
 	/** Returns the platform user attached to this input device, or PLATFORMUSERID_NONE if invalid */
-	virtual FPlatformUserId GetUserForInputDevice(FInputDeviceId DeviceId) const;
+	APPLICATIONCORE_API virtual FPlatformUserId GetUserForInputDevice(FInputDeviceId DeviceId) const;
 
 	/** Returns the primary input device used by a specific player, or INPUTDEVICEID_NONE if invalid */
-	virtual FInputDeviceId GetPrimaryInputDeviceForUser(FPlatformUserId UserId) const;
+	APPLICATIONCORE_API virtual FInputDeviceId GetPrimaryInputDeviceForUser(FPlatformUserId UserId) const;
 
 	/**
 	 * Set the connection state of a given Input Device to something new. This will
@@ -105,7 +105,7 @@ public:
 	 * @param NewState		The new connection state of the given device
 	 * @return				True if the connection state was set successfully
 	 */
-	virtual bool Internal_SetInputDeviceConnectionState(FInputDeviceId DeviceId, EInputDeviceConnectionState NewState);
+	APPLICATIONCORE_API virtual bool Internal_SetInputDeviceConnectionState(FInputDeviceId DeviceId, EInputDeviceConnectionState NewState);
 	
 	/**
 	 * Gets the connection state of the given input device.
@@ -113,7 +113,7 @@ public:
 	 * @param DeviceId		The device to get the connection state of
 	 * @return				The connection state of the given device. EInputDeviceConnectionState::Unknown if the device is not mapped
 	 */
-	virtual EInputDeviceConnectionState GetInputDeviceConnectionState(const FInputDeviceId DeviceId) const;
+	APPLICATIONCORE_API virtual EInputDeviceConnectionState GetInputDeviceConnectionState(const FInputDeviceId DeviceId) const;
 
 	/**
 	 * Maps the given Input Device to the given userID. This will broadcast the OnInputDeviceConnectionChange delegate.
@@ -124,7 +124,7 @@ public:
 	 * @param ConnectionState	The connection state of the device
 	 * @return					True if the input device was successfully mapped
 	 */
-	virtual bool Internal_MapInputDeviceToUser(FInputDeviceId DeviceId, FPlatformUserId UserId, EInputDeviceConnectionState ConnectionState);
+	APPLICATIONCORE_API virtual bool Internal_MapInputDeviceToUser(FInputDeviceId DeviceId, FPlatformUserId UserId, EInputDeviceConnectionState ConnectionState);
 
 	/**
 	 * Change the user mapping of the given input device from an old user to a new one.
@@ -137,7 +137,7 @@ public:
 	 * @param OldUserId		The old platform user that this input device is currently mapped to
 	 * @return				True if the device was successfully remapped
 	 */
-	virtual bool Internal_ChangeInputDeviceUserMapping(FInputDeviceId DeviceId, FPlatformUserId NewUserId, FPlatformUserId OldUserId);
+	APPLICATIONCORE_API virtual bool Internal_ChangeInputDeviceUserMapping(FInputDeviceId DeviceId, FPlatformUserId NewUserId, FPlatformUserId OldUserId);
 
 	//////////////////////////////////////////////////////////////////////////////
 	// Delegates for listening to input device changes
@@ -223,13 +223,13 @@ protected:
 	* Binds to any core delegates that the platform may broadcast. By default it is only the OnUserLoginChangedEvent delegate. 
 	* Called during the constructor of this device mapper.
 	*/
-	virtual void BindCoreDelegates();
+	APPLICATIONCORE_API virtual void BindCoreDelegates();
 	
 	/** 
 	* Unbind from any delegates that have been hooked into.
 	* Called during the destructor of this device mapper.
 	*/
-	virtual void UnbindCoreDelegates();
+	APPLICATIONCORE_API virtual void UnbindCoreDelegates();
 
 	/** Callback for when FCoreDelegates::OnUserLoginChangedEvent is broadcasted */
 	virtual void OnUserLoginChangedEvent(bool bLoggedIn, int32 UserId, int32 UserIndex) = 0;
@@ -245,7 +245,7 @@ protected:
 	 * and owning FPlatformUserId we should create a new platform user id if this controller id is
 	 * not already mapped. 
 	 */
-	virtual bool ShouldCreateUniqueUserForEachDevice() const;
+	APPLICATIONCORE_API virtual bool ShouldCreateUniqueUserForEachDevice() const;
 
 	/**
 	 * If true, than this device mapper will broadcast the older 
@@ -256,10 +256,10 @@ protected:
 	virtual bool ShouldBroadcastLegacyDelegates() const = 0;
 
 	/** Callback when input devices are disconnected/reconnected */
-	static FOnUserInputDeviceConnectionChange OnInputDeviceConnectionChange;
+	static APPLICATIONCORE_API FOnUserInputDeviceConnectionChange OnInputDeviceConnectionChange;
 
 	/** Callback when an input device's owning platform user pairing changes */
-	static FOnUserInputDevicePairingChange OnInputDevicePairingChange;
+	static APPLICATIONCORE_API FOnUserInputDevicePairingChange OnInputDevicePairingChange;
 	
 	/** A map of all input devices to their current state */
 	TMap<FInputDeviceId, FPlatformInputDeviceState> MappedInputDevices;
@@ -275,29 +275,29 @@ protected:
  * Generic implementation of the IPlatformInputDeviceMapper.
  * This provides the base functionality that can be used on most platforms.
  */
-class APPLICATIONCORE_API FGenericPlatformInputDeviceMapper : public IPlatformInputDeviceMapper
+class FGenericPlatformInputDeviceMapper : public IPlatformInputDeviceMapper
 {
 public:
 
-	FGenericPlatformInputDeviceMapper(const bool InbUsingControllerIdAsUserId, const bool InbShouldBroadcastLegacyDelegates);
+	APPLICATIONCORE_API FGenericPlatformInputDeviceMapper(const bool InbUsingControllerIdAsUserId, const bool InbShouldBroadcastLegacyDelegates);
 	
 	/** This is unsupported by default and will return PLATFORMUSERID_NONE on the generic platform */  
-	virtual FPlatformUserId GetUserForUnpairedInputDevices() const override;
-	virtual FPlatformUserId GetPrimaryPlatformUser() const override;
-	virtual FInputDeviceId GetDefaultInputDevice() const override;
+	APPLICATIONCORE_API virtual FPlatformUserId GetUserForUnpairedInputDevices() const override;
+	APPLICATIONCORE_API virtual FPlatformUserId GetPrimaryPlatformUser() const override;
+	APPLICATIONCORE_API virtual FInputDeviceId GetDefaultInputDevice() const override;
 
-	virtual bool RemapControllerIdToPlatformUserAndDevice(int32 ControllerId, FPlatformUserId& InOutUserId, FInputDeviceId& OutInputDeviceId) override;
-	virtual FPlatformUserId GetPlatformUserForUserIndex(int32 LocalUserIndex) override;
-	virtual bool RemapUserAndDeviceToControllerId(FPlatformUserId UserId, int32& OutControllerId, FInputDeviceId OptionalDeviceId = INPUTDEVICEID_NONE) override;
-	virtual int32 GetUserIndexForPlatformUser(FPlatformUserId UserId) override;
-	virtual bool IsUsingControllerIdAsUserId() const override;
-	virtual bool ShouldBroadcastLegacyDelegates() const override;
+	APPLICATIONCORE_API virtual bool RemapControllerIdToPlatformUserAndDevice(int32 ControllerId, FPlatformUserId& InOutUserId, FInputDeviceId& OutInputDeviceId) override;
+	APPLICATIONCORE_API virtual FPlatformUserId GetPlatformUserForUserIndex(int32 LocalUserIndex) override;
+	APPLICATIONCORE_API virtual bool RemapUserAndDeviceToControllerId(FPlatformUserId UserId, int32& OutControllerId, FInputDeviceId OptionalDeviceId = INPUTDEVICEID_NONE) override;
+	APPLICATIONCORE_API virtual int32 GetUserIndexForPlatformUser(FPlatformUserId UserId) override;
+	APPLICATIONCORE_API virtual bool IsUsingControllerIdAsUserId() const override;
+	APPLICATIONCORE_API virtual bool ShouldBroadcastLegacyDelegates() const override;
 
 	/** Allocates a new user id when a user becomes active, will return none if no more can be created */
-	virtual FPlatformUserId AllocateNewUserId() override;
+	APPLICATIONCORE_API virtual FPlatformUserId AllocateNewUserId() override;
 
 	/** Returns the next available input device id */
-	virtual FInputDeviceId AllocateNewInputDeviceId() override;
+	APPLICATIONCORE_API virtual FInputDeviceId AllocateNewInputDeviceId() override;
 	
 protected:
 
@@ -305,7 +305,7 @@ protected:
 	* Callback for when FCoreDelegates::OnUserLoginChangedEvent is broadcasted.
 	* If the user has logged out, then remap any input devices that the user had to the "Unpaired" user on this platform (if bUnpairInputDevicesWhenLoggingOut is true)
 	*/
-	virtual void OnUserLoginChangedEvent(bool bLoggedIn, int32 RawPlatformUserId, int32 UserIndex) override;
+	APPLICATIONCORE_API virtual void OnUserLoginChangedEvent(bool bLoggedIn, int32 RawPlatformUserId, int32 UserIndex) override;
 
 	/** Flags for backwards compatibility with the older "int32 ControllerId" implementation */
 	const bool bUsingControllerIdAsUserId = true;
@@ -313,4 +313,4 @@ protected:
 
 	/** Additional functionality that specializations can customize */
 	bool bUnpairInputDevicesWhenLoggingOut = true;
-};
\ No newline at end of file
+};
diff --git a/Engine/Source/Runtime/ApplicationCore/Public/GenericPlatform/GenericPlatformSplash.h b/Engine/Source/Runtime/ApplicationCore/Public/GenericPlatform/GenericPlatformSplash.h
--- a/Engine/Source/Runtime/ApplicationCore/Public/GenericPlatform/GenericPlatformSplash.h
+++ b/Engine/Source/Runtime/ApplicationCore/Public/GenericPlatform/GenericPlatformSplash.h
@@ -36,7 +36,7 @@ class FString;
 /**
  * Generic implementation for most platforms
  */
-struct APPLICATIONCORE_API FGenericPlatformSplash
+struct FGenericPlatformSplash
 {
 	/** Show the splash screen. */
 	FORCEINLINE static void Show() { }
@@ -49,7 +49,7 @@ struct APPLICATIONCORE_API FGenericPlatformSplash
 	 * 
 	 * @param SplashFilename Full path to the splash image to display
 	 */
-	static void SetCustomSplashImage(const TCHAR* SplashFilename);
+	static APPLICATIONCORE_API void SetCustomSplashImage(const TCHAR* SplashFilename);
 
 
 	/**
@@ -92,6 +92,6 @@ protected:
 	*
 	* @return true if a splash screen was found
 	*/
-	static bool GetSplashPath(const TCHAR* SplashFilename, FString& OutPath, bool& OutIsCustom);
-	static bool GetSplashPath(const TCHAR* SplashFilename, const TCHAR* IconFilename, FString& OutPath, FString& OutIconPath, bool& OutIsCustom);
+	static APPLICATIONCORE_API bool GetSplashPath(const TCHAR* SplashFilename, FString& OutPath, bool& OutIsCustom);
+	static APPLICATIONCORE_API bool GetSplashPath(const TCHAR* SplashFilename, const TCHAR* IconFilename, FString& OutPath, FString& OutIconPath, bool& OutIsCustom);
 };
diff --git a/Engine/Source/Runtime/ApplicationCore/Public/GenericPlatform/GenericPlatformSurvey.h b/Engine/Source/Runtime/ApplicationCore/Public/GenericPlatform/GenericPlatformSurvey.h
--- a/Engine/Source/Runtime/ApplicationCore/Public/GenericPlatform/GenericPlatformSurvey.h
+++ b/Engine/Source/Runtime/ApplicationCore/Public/GenericPlatform/GenericPlatformSurvey.h
@@ -224,7 +224,7 @@ struct FHardwareSurveyResults
 /**
 * Generic implementation for most platforms, these tend to be unused and unimplemented
 **/
-struct APPLICATIONCORE_API FGenericPlatformSurvey
+struct FGenericPlatformSurvey
 {
 	/**
 	 * Attempt to get hardware survey results now.
diff --git a/Engine/Source/Runtime/ApplicationCore/Public/GenericPlatform/GenericWindow.h b/Engine/Source/Runtime/ApplicationCore/Public/GenericPlatform/GenericWindow.h
--- a/Engine/Source/Runtime/ApplicationCore/Public/GenericPlatform/GenericWindow.h
+++ b/Engine/Source/Runtime/ApplicationCore/Public/GenericPlatform/GenericWindow.h
@@ -89,61 +89,61 @@ struct FWindowDrawAttentionParameters
 	EWindowDrawAttentionRequestType RequestType = EWindowDrawAttentionRequestType::UntilActivated;
 };
 
-class APPLICATIONCORE_API FGenericWindow
+class FGenericWindow
 {
 public:
 
-	FGenericWindow();
+	APPLICATIONCORE_API FGenericWindow();
 
-	virtual ~FGenericWindow();
+	APPLICATIONCORE_API virtual ~FGenericWindow();
 
 	/** Native windows should implement ReshapeWindow by changing the platform-specific window to be located at (X,Y) and be the dimensions Width x Height. */
-	virtual void ReshapeWindow( int32 X, int32 Y, int32 Width, int32 Height );
+	APPLICATIONCORE_API virtual void ReshapeWindow( int32 X, int32 Y, int32 Width, int32 Height );
 
 	/** Returns the rectangle of the screen the window is associated with */
-	virtual bool GetFullScreenInfo( int32& X, int32& Y, int32& Width, int32& Height ) const;
+	APPLICATIONCORE_API virtual bool GetFullScreenInfo( int32& X, int32& Y, int32& Width, int32& Height ) const;
 
 	/** Native windows should implement MoveWindowTo by relocating the platform-specific window to (X,Y). */
-	virtual void MoveWindowTo ( int32 X, int32 Y );
+	APPLICATIONCORE_API virtual void MoveWindowTo ( int32 X, int32 Y );
 
 	/** Native windows should implement BringToFront by making this window the top-most window (i.e. focused). */
-	virtual void BringToFront( bool bForce = false );
+	APPLICATIONCORE_API virtual void BringToFront( bool bForce = false );
 
 	/** @hack Force a window to front even if a different application is in front. */
-	virtual void HACK_ForceToFront();
+	APPLICATIONCORE_API virtual void HACK_ForceToFront();
 
 	/** Native windows should implement this function by asking the OS to destroy OS-specific resource associated with the window (e.g. Win32 window handle) */
-	virtual void Destroy();
+	APPLICATIONCORE_API virtual void Destroy();
 
 	/** Native window should implement this function by performing the equivalent of the Win32 minimize-to-taskbar operation */
-	virtual void Minimize();
+	APPLICATIONCORE_API virtual void Minimize();
 
 	/** Native window should implement this function by performing the equivalent of the Win32 maximize operation */
-	virtual void Maximize();
+	APPLICATIONCORE_API virtual void Maximize();
 
 	/** Native window should implement this function by performing the equivalent of the Win32 restore operation */
-	virtual void Restore();
+	APPLICATIONCORE_API virtual void Restore();
 
 	/** Native window should make itself visible */
-	virtual void Show();
+	APPLICATIONCORE_API virtual void Show();
 
 	/** Native window should hide itself */
-	virtual void Hide();
+	APPLICATIONCORE_API virtual void Hide();
 
 	/** Toggle native window between fullscreen and normal mode */
-	virtual void SetWindowMode( EWindowMode::Type InNewWindowMode );
+	APPLICATIONCORE_API virtual void SetWindowMode( EWindowMode::Type InNewWindowMode );
 
 	/** @return true if the native window is currently in fullscreen mode, false otherwise */
-	virtual EWindowMode::Type GetWindowMode() const;
+	APPLICATIONCORE_API virtual EWindowMode::Type GetWindowMode() const;
 
 	/** @return true if the native window is maximized, false otherwise */
-	virtual bool IsMaximized() const;
+	APPLICATIONCORE_API virtual bool IsMaximized() const;
 
 	/** @return true if the native window is minimized, false otherwise */
-	virtual bool IsMinimized() const;
+	APPLICATIONCORE_API virtual bool IsMinimized() const;
 
 	/** @return true if the native window is visible, false otherwise */
-	virtual bool IsVisible() const;
+	APPLICATIONCORE_API virtual bool IsVisible() const;
 
 	/**
 	 * Populates the size and location of the window when it is restored.
@@ -151,76 +151,76 @@ public:
 	 *
 	 * @return true when the size and location and successfully retrieved; false otherwise.
 	 */
-	virtual bool GetRestoredDimensions(int32& X, int32& Y, int32& Width, int32& Height);
+	APPLICATIONCORE_API virtual bool GetRestoredDimensions(int32& X, int32& Y, int32& Width, int32& Height);
 
 	/** 
 	 * Native windows should implement SetWindowFocus to let the OS know that a window has taken focus.  
 	 * Slate handles focus on a per widget basis internally but the OS still needs to know what window has focus for proper message routing
 	 */
-	virtual void SetWindowFocus();
+	APPLICATIONCORE_API virtual void SetWindowFocus();
 
 	/**
 	 * Sets the opacity of this window
 	 *
 	 * @param	InOpacity	The new window opacity represented as a floating point scalar
 	 */
-	virtual void SetOpacity( const float InOpacity );
+	APPLICATIONCORE_API virtual void SetOpacity( const float InOpacity );
 
 	/** 
 	 * Enables or disables the window.  If disabled the window receives no input
 	 * 
 	 * @param bEnable	true to enable the window, false to disable it.
 	 */
-	virtual void Enable( bool bEnable );
+	APPLICATIONCORE_API virtual void Enable( bool bEnable );
 
 	/** @return true if native window exists underneath the coordinates */
-	virtual bool IsPointInWindow( int32 X, int32 Y ) const;
+	APPLICATIONCORE_API virtual bool IsPointInWindow( int32 X, int32 Y ) const;
 	
 	/** Gets OS specific window border size. This is necessary because Win32 does not give control over this size. */
-	virtual int32 GetWindowBorderSize() const;
+	APPLICATIONCORE_API virtual int32 GetWindowBorderSize() const;
 
 	/** Gets OS specific window title bar size */
-	virtual int32 GetWindowTitleBarSize() const;
+	APPLICATIONCORE_API virtual int32 GetWindowTitleBarSize() const;
 
 	/** Gets the OS Window handle in the form of a void pointer for other API's */
-	virtual void* GetOSWindowHandle() const;
+	APPLICATIONCORE_API virtual void* GetOSWindowHandle() const;
 
 	/** @return true if the window is in the foreground */
-	virtual bool IsForegroundWindow() const;
+	APPLICATIONCORE_API virtual bool IsForegroundWindow() const;
 
 	/** @return true if the window is in the foreground */
-	virtual bool IsFullscreenSupported() const;
+	APPLICATIONCORE_API virtual bool IsFullscreenSupported() const;
 
 	/**
 	 * Sets the window text - usually the title but can also be text content for things like controls
 	 *
 	 * @param Text	The window's title or content text
 	 */
-	virtual void SetText(const TCHAR* const Text);
+	APPLICATIONCORE_API virtual void SetText(const TCHAR* const Text);
 
 	/** @return	The definition describing properties of the window */
-	virtual const FGenericWindowDefinition& GetDefinition() const;
+	APPLICATIONCORE_API virtual const FGenericWindowDefinition& GetDefinition() const;
 
 	/** @return	Returns true if the window definition is valid */
-	virtual bool IsDefinitionValid() const;
+	APPLICATIONCORE_API virtual bool IsDefinitionValid() const;
 
 	/** @return	Gives the native window a chance to adjust our stored window size before we cache it off */
-	virtual void AdjustCachedSize( FVector2D& Size ) const;
+	APPLICATIONCORE_API virtual void AdjustCachedSize( FVector2D& Size ) const;
 
 	/**
 	 * @return ratio of pixels to SlateUnits in this window.
 	 * E.g. DPIScale of 2.0 means there is a 2x2 pixel square for every 1x1 SlateUnit.
 	 */
-	virtual float GetDPIScaleFactor() const;
+	APPLICATIONCORE_API virtual float GetDPIScaleFactor() const;
 
 	/** sets a new DPI scale factor */
-	virtual void SetDPIScaleFactor(const float Factor);
+	APPLICATIONCORE_API virtual void SetDPIScaleFactor(const float Factor);
 
 	/** determines whether or not this window does its own DPI management */
-	virtual bool IsManualManageDPIChanges() const;
+	APPLICATIONCORE_API virtual bool IsManualManageDPIChanges() const;
 
 	/** call with a true argument if this window need to do its custom size management in response to DPI variations */
-	virtual void SetManualManageDPIChanges(const bool bAutoHandle);
+	APPLICATIONCORE_API virtual void SetManualManageDPIChanges(const bool bAutoHandle);
 
 	/**
 	 * Attempts to draw the user's attention to this window in whatever way is
@@ -230,10 +230,10 @@ public:
 	 * @param Parameters The parameters for drawing attention. Depending on
 	 *        the platform, not all parameters may be supported.
 	 */
-	virtual void DrawAttention(const FWindowDrawAttentionParameters& Parameters);
+	APPLICATIONCORE_API virtual void DrawAttention(const FWindowDrawAttentionParameters& Parameters);
 
 	/** Shows or hides native window buttons on platforms that use them */
-	virtual void SetNativeWindowButtonsVisibility(bool bVisible);
+	APPLICATIONCORE_API virtual void SetNativeWindowButtonsVisibility(bool bVisible);
 
 protected:
 
diff --git a/Engine/Source/Runtime/ApplicationCore/Public/GenericPlatform/GenericWindowDefinition.h b/Engine/Source/Runtime/ApplicationCore/Public/GenericPlatform/GenericWindowDefinition.h
--- a/Engine/Source/Runtime/ApplicationCore/Public/GenericPlatform/GenericWindowDefinition.h
+++ b/Engine/Source/Runtime/ApplicationCore/Public/GenericPlatform/GenericWindowDefinition.h
@@ -60,7 +60,7 @@ enum class EWindowActivationPolicy
 };
 
 
-struct APPLICATIONCORE_API FGenericWindowDefinition
+struct FGenericWindowDefinition
 {
 	/** Window type */
 	EWindowType Type;
diff --git a/Engine/Source/Runtime/ApplicationCore/Public/Linux/LinuxConsoleOutputDevice.h b/Engine/Source/Runtime/ApplicationCore/Public/Linux/LinuxConsoleOutputDevice.h
--- a/Engine/Source/Runtime/ApplicationCore/Public/Linux/LinuxConsoleOutputDevice.h
+++ b/Engine/Source/Runtime/ApplicationCore/Public/Linux/LinuxConsoleOutputDevice.h
@@ -8,7 +8,7 @@
 /**
 * Linux implementation of console log window, just a printf to the terminal for now
 */
-class APPLICATIONCORE_API FLinuxConsoleOutputDevice : public FOutputDeviceConsole
+class FLinuxConsoleOutputDevice : public FOutputDeviceConsole
 {
 
 private:
@@ -30,27 +30,27 @@ public:
 	/**
 	 * Constructor, setting console control handler.
 	 */
-	FLinuxConsoleOutputDevice();
-	~FLinuxConsoleOutputDevice();
+	APPLICATIONCORE_API FLinuxConsoleOutputDevice();
+	APPLICATIONCORE_API ~FLinuxConsoleOutputDevice();
 
 	/**
 	 * Shows or hides the console window.
 	 *
 	 * @param ShowWindow	Whether to show (true) or hide (false) the console window.
 	 */
-	virtual void Show(bool bShowWindow);
+	APPLICATIONCORE_API virtual void Show(bool bShowWindow);
 
 	/**
 	 * Returns whether console is currently shown or not
 	 *
 	 * @return true if console is shown, false otherwise
 	 */
-	virtual bool IsShown();
+	APPLICATIONCORE_API virtual bool IsShown();
 
 	virtual bool IsAttached() {return false;}
 
-	virtual bool CanBeUsedOnAnyThread() const override;
-	virtual bool CanBeUsedOnPanicThread() const override;
+	APPLICATIONCORE_API virtual bool CanBeUsedOnAnyThread() const override;
+	APPLICATIONCORE_API virtual bool CanBeUsedOnPanicThread() const override;
 
 	/**
 	 * Displays text on the console and scrolls if necessary.
@@ -58,6 +58,6 @@ public:
 	 * @param Data	Text to display
 	 * @param Event	Event type, used for filtering/ suppression
 	 */
-	void Serialize(const TCHAR* Data, ELogVerbosity::Type Verbosity, const class FName& Category);
+	APPLICATIONCORE_API void Serialize(const TCHAR* Data, ELogVerbosity::Type Verbosity, const class FName& Category);
 
 };
diff --git a/Engine/Source/Runtime/ApplicationCore/Public/Linux/LinuxPlatformApplicationMisc.h b/Engine/Source/Runtime/ApplicationCore/Public/Linux/LinuxPlatformApplicationMisc.h
--- a/Engine/Source/Runtime/ApplicationCore/Public/Linux/LinuxPlatformApplicationMisc.h
+++ b/Engine/Source/Runtime/ApplicationCore/Public/Linux/LinuxPlatformApplicationMisc.h
@@ -4,37 +4,37 @@
 
 #include "GenericPlatform/GenericPlatformApplicationMisc.h"
 
-struct APPLICATIONCORE_API FLinuxPlatformApplicationMisc : public FGenericPlatformApplicationMisc
+struct FLinuxPlatformApplicationMisc : public FGenericPlatformApplicationMisc
 {
-	static void PreInit();
-	static void Init();
-	static bool InitSDL();
-	static void TearDown();
-	static void LoadPreInitModules();
-	static void LoadStartupModules();
-	static uint32 WindowStyle();
-	static class FOutputDeviceConsole* CreateConsoleOutputDevice();
-	static class FOutputDeviceError* GetErrorOutputDevice();
-	static class FFeedbackContext* GetFeedbackContext();
-	static class GenericApplication* CreateApplication();
-	static bool IsThisApplicationForeground();
-	static void PumpMessages(bool bFromMainLoop);
-	static bool IsScreensaverEnabled();
-	static bool ControlScreensaver(EScreenSaverAction Action);
-	static float GetDPIScaleFactorAtPoint(float X, float Y);
-	static void ClipboardCopy(const TCHAR* Str);
-	static void ClipboardPaste(class FString& Dest);
+	static APPLICATIONCORE_API void PreInit();
+	static APPLICATIONCORE_API void Init();
+	static APPLICATIONCORE_API bool InitSDL();
+	static APPLICATIONCORE_API void TearDown();
+	static APPLICATIONCORE_API void LoadPreInitModules();
+	static APPLICATIONCORE_API void LoadStartupModules();
+	static APPLICATIONCORE_API uint32 WindowStyle();
+	static APPLICATIONCORE_API class FOutputDeviceConsole* CreateConsoleOutputDevice();
+	static APPLICATIONCORE_API class FOutputDeviceError* GetErrorOutputDevice();
+	static APPLICATIONCORE_API class FFeedbackContext* GetFeedbackContext();
+	static APPLICATIONCORE_API class GenericApplication* CreateApplication();
+	static APPLICATIONCORE_API bool IsThisApplicationForeground();
+	static APPLICATIONCORE_API void PumpMessages(bool bFromMainLoop);
+	static APPLICATIONCORE_API bool IsScreensaverEnabled();
+	static APPLICATIONCORE_API bool ControlScreensaver(EScreenSaverAction Action);
+	static APPLICATIONCORE_API float GetDPIScaleFactorAtPoint(float X, float Y);
+	static APPLICATIONCORE_API void ClipboardCopy(const TCHAR* Str);
+	static APPLICATIONCORE_API void ClipboardPaste(class FString& Dest);
 	static bool FullscreenSameAsWindowedFullscreen() { return true; }
 
 	// Unix specific
-	static void EarlyUnixInitialization(class FString& OutCommandLine);
+	static APPLICATIONCORE_API void EarlyUnixInitialization(class FString& OutCommandLine);
 	static bool ShouldIncreaseProcessLimits() { return true; }
 
 	// Linux specific
 	/** Informs ApplicationCore that it needs to create Vulkan-compatible windows (mutually exclusive with OpenGL) */
-	static void UsingVulkan();
+	static APPLICATIONCORE_API void UsingVulkan();
 	/** Informs ApplicationCore that it needs to create OpenGL-compatible windows (mutually exclusive with Vulkan) */
-	static void UsingOpenGL();
+	static APPLICATIONCORE_API void UsingOpenGL();
 };
 
 typedef FLinuxPlatformApplicationMisc FPlatformApplicationMisc;
diff --git a/Engine/Source/Runtime/ApplicationCore/Public/Linux/LinuxPlatformSplash.h b/Engine/Source/Runtime/ApplicationCore/Public/Linux/LinuxPlatformSplash.h
--- a/Engine/Source/Runtime/ApplicationCore/Public/Linux/LinuxPlatformSplash.h
+++ b/Engine/Source/Runtime/ApplicationCore/Public/Linux/LinuxPlatformSplash.h
@@ -12,16 +12,16 @@
 /**
  * Linux splash implementation
  */
-struct APPLICATIONCORE_API FLinuxPlatformSplash : public FGenericPlatformSplash
+struct FLinuxPlatformSplash : public FGenericPlatformSplash
 {
 	/**
 	 * Show the splash screen
 	 */
-	static void Show();
+	static APPLICATIONCORE_API void Show();
 	/**
 	 * Hide the splash screen
 	 */
-	static void Hide();
+	static APPLICATIONCORE_API void Hide();
 
 	/**
 	 * Sets the text displayed on the splash screen (for startup/loading progress)
@@ -29,7 +29,7 @@ struct APPLICATIONCORE_API FLinuxPlatformSplash : public FGenericPlatformSplash
 	 * @param	InType		Type of text to change
 	 * @param	InText		Text to display
 	 */
-	static void SetSplashText(const SplashTextType::Type InType, const TCHAR* InText);
+	static APPLICATIONCORE_API void SetSplashText(const SplashTextType::Type InType, const TCHAR* InText);
 };
 
 typedef FLinuxPlatformSplash FPlatformSplash;
diff --git a/Engine/Source/Runtime/ApplicationCore/Public/Linux/LinuxWindow.h b/Engine/Source/Runtime/ApplicationCore/Public/Linux/LinuxWindow.h
--- a/Engine/Source/Runtime/ApplicationCore/Public/Linux/LinuxWindow.h
+++ b/Engine/Source/Runtime/ApplicationCore/Public/Linux/LinuxWindow.h
@@ -24,10 +24,10 @@ enum class EWindowActivationPolicy;
  * A platform specific implementation of FNativeWindow.
  * Native Windows provide platform-specific backing for and are always owned by an SWindow.
  */
-class APPLICATIONCORE_API FLinuxWindow : public FGenericWindow
+class FLinuxWindow : public FGenericWindow
 {
 public:
-	~FLinuxWindow();
+	APPLICATIONCORE_API ~FLinuxWindow();
 
 	/** Create a new SDLWindow.
 	 *
@@ -37,58 +37,58 @@ public:
 	 * @param InParent			Parent Win32 window; usually NULL.
 	 * @param bShowImmediately	True to show this window as soon as its initialized
 	 */
-	static TSharedRef< FLinuxWindow > Make();
+	static APPLICATIONCORE_API TSharedRef< FLinuxWindow > Make();
 
-	SDL_HWindow GetHWnd() const;
+	APPLICATIONCORE_API SDL_HWindow GetHWnd() const;
 
-	void Initialize( class FLinuxApplication* const Application, const TSharedRef< FGenericWindowDefinition >& InDefinition, const TSharedPtr< FLinuxWindow >& InParent, const bool bShowImmediately );
+	APPLICATIONCORE_API void Initialize( class FLinuxApplication* const Application, const TSharedRef< FGenericWindowDefinition >& InDefinition, const TSharedPtr< FLinuxWindow >& InParent, const bool bShowImmediately );
 
 	/** TODO: describe */
-	bool IsRegularWindow() const;
+	APPLICATIONCORE_API bool IsRegularWindow() const;
 
 	/** TODO: describe */
-	bool IsPopupMenuWindow() const;
+	APPLICATIONCORE_API bool IsPopupMenuWindow() const;
 
 	/** TODO: describe */
-	bool IsTooltipWindow() const;
+	APPLICATIONCORE_API bool IsTooltipWindow() const;
 
 	/** TODO: describe */
-	bool IsNotificationWindow() const;
+	APPLICATIONCORE_API bool IsNotificationWindow() const;
 
 	/** TODO: describe */
-	bool IsTopLevelWindow() const;
+	APPLICATIONCORE_API bool IsTopLevelWindow() const;
 
 	/** @return true if this is a modal window */
-	bool IsModalWindow() const;
+	APPLICATIONCORE_API bool IsModalWindow() const;
 
 	/** TODO: describe */
-	bool IsDialogWindow() const;
+	APPLICATIONCORE_API bool IsDialogWindow() const;
 
 	/** TODO: describe */
-	bool IsDragAndDropWindow() const;
+	APPLICATIONCORE_API bool IsDragAndDropWindow() const;
 
 	/** TODO: describe */
-	bool IsUtilityWindow() const;
+	APPLICATIONCORE_API bool IsUtilityWindow() const;
 
 	/** @return the window activation policy used when showing the window */
-	EWindowActivationPolicy GetActivationPolicy() const;
+	APPLICATIONCORE_API EWindowActivationPolicy GetActivationPolicy() const;
 
 	/** TODO: describe */
-	bool IsFocusWhenFirstShown() const;
+	APPLICATIONCORE_API bool IsFocusWhenFirstShown() const;
 
 	/** TODO: describe */
-	const TSharedPtr< FLinuxWindow >& GetParent() const;
+	APPLICATIONCORE_API const TSharedPtr< FLinuxWindow >& GetParent() const;
 
 	/** Internal ID using for debugging */
-	uint32 GetID() const;
+	APPLICATIONCORE_API uint32 GetID() const;
 
 	/** Debugging function - dumps window info to log */
-	void LogInfo();
+	APPLICATIONCORE_API void LogInfo();
 
-	bool IsNativeMoving() const;
-	bool BeginNativeMove();
-	void AfterNativeMove();
-	void EndNativeMove();
+	APPLICATIONCORE_API bool IsNativeMoving() const;
+	APPLICATIONCORE_API bool BeginNativeMove();
+	APPLICATIONCORE_API void AfterNativeMove();
+	APPLICATIONCORE_API void EndNativeMove();
 
 	enum NativeResizeDirection 
 	{
@@ -104,67 +104,67 @@ public:
 		InvalidDirection
 	};
 
-	bool IsNativeResizing() const;
-	bool BeginNativeResize( NativeResizeDirection Direction );
-	void AfterNativeResize();
-	void EndNativeResize();
+	APPLICATIONCORE_API bool IsNativeResizing() const;
+	APPLICATIONCORE_API bool BeginNativeResize( NativeResizeDirection Direction );
+	APPLICATIONCORE_API void AfterNativeResize();
+	APPLICATIONCORE_API void EndNativeResize();
 
 public:
-	virtual void ReshapeWindow( int32 X, int32 Y, int32 Width, int32 Height ) override;
+	APPLICATIONCORE_API virtual void ReshapeWindow( int32 X, int32 Y, int32 Width, int32 Height ) override;
 
 	virtual void* GetOSWindowHandle() const  override { return HWnd; }
 
 	//	not finished
-	virtual bool GetFullScreenInfo( int32& X, int32& Y, int32& Width, int32& Height ) const override;
+	APPLICATIONCORE_API virtual bool GetFullScreenInfo( int32& X, int32& Y, int32& Width, int32& Height ) const override;
 
-	virtual void MoveWindowTo ( int32 X, int32 Y ) override;
+	APPLICATIONCORE_API virtual void MoveWindowTo ( int32 X, int32 Y ) override;
 
-	virtual void BringToFront( bool bForce = false ) override;
+	APPLICATIONCORE_API virtual void BringToFront( bool bForce = false ) override;
 
-	virtual void Destroy() override;
+	APPLICATIONCORE_API virtual void Destroy() override;
 
-	virtual void Minimize() override;
+	APPLICATIONCORE_API virtual void Minimize() override;
 
-	virtual void Maximize() override;
+	APPLICATIONCORE_API virtual void Maximize() override;
 
-	virtual void Restore() override;
+	APPLICATIONCORE_API virtual void Restore() override;
 
-	virtual void Show() override;
+	APPLICATIONCORE_API virtual void Show() override;
 
-	virtual void Hide() override;
+	APPLICATIONCORE_API virtual void Hide() override;
 
-	virtual void SetWindowMode( EWindowMode::Type NewWindowMode ) override;
+	APPLICATIONCORE_API virtual void SetWindowMode( EWindowMode::Type NewWindowMode ) override;
 
 	virtual EWindowMode::Type GetWindowMode() const override { return WindowMode; } 
 
-	virtual bool IsMaximized() const override;
+	APPLICATIONCORE_API virtual bool IsMaximized() const override;
 
-	virtual bool IsVisible() const override;
+	APPLICATIONCORE_API virtual bool IsVisible() const override;
 
-	virtual bool IsMinimized() const override;
+	APPLICATIONCORE_API virtual bool IsMinimized() const override;
 
-	virtual bool GetRestoredDimensions(int32& X, int32& Y, int32& Width, int32& Height) override;
+	APPLICATIONCORE_API virtual bool GetRestoredDimensions(int32& X, int32& Y, int32& Width, int32& Height) override;
 
-	virtual void SetWindowFocus() override;
+	APPLICATIONCORE_API virtual void SetWindowFocus() override;
 
-	virtual void SetOpacity( const float InOpacity ) override;
+	APPLICATIONCORE_API virtual void SetOpacity( const float InOpacity ) override;
 
-	virtual void Enable( bool bEnable ) override;
+	APPLICATIONCORE_API virtual void Enable( bool bEnable ) override;
 
-	virtual bool IsPointInWindow( int32 X, int32 Y ) const override;
+	APPLICATIONCORE_API virtual bool IsPointInWindow( int32 X, int32 Y ) const override;
 
-	virtual int32 GetWindowBorderSize() const override;
+	APPLICATIONCORE_API virtual int32 GetWindowBorderSize() const override;
 
-	virtual bool IsForegroundWindow() const override;
+	APPLICATIONCORE_API virtual bool IsForegroundWindow() const override;
 	/**
 	 * Sets the window text - usually the title but can also be text content for things like controls
 	 *
 	 * @param Text	The window's title or content text
 	 */
-	virtual void SetText(const TCHAR* const Text) override;
+	APPLICATIONCORE_API virtual void SetText(const TCHAR* const Text) override;
 
 	/** @return	Gives the native window a chance to adjust our stored window size before we cache it off */
-	virtual void AdjustCachedSize( FVector2D& Size ) const override;
+	APPLICATIONCORE_API virtual void AdjustCachedSize( FVector2D& Size ) const override;
 
 	virtual float GetDPIScaleFactor() const override
 	{
@@ -182,20 +182,20 @@ public:
 	 * @param BorderWidth - width of a vertical border (left side)
 	 * @param BorderHeight - height of a horizontal border (from the top)
 	 */
-	void GetNativeBordersSize(int32& OutLeftBorderWidth, int32& OutTopBorderHeight) const;
+	APPLICATIONCORE_API void GetNativeBordersSize(int32& OutLeftBorderWidth, int32& OutTopBorderHeight) const;
 
 	/**
 	 *  Caches window properties that are too expensive to query all the time.
 	 *  (e.g. border sizes)
 	 */
-	void CacheNativeProperties();
+	APPLICATIONCORE_API void CacheNativeProperties();
 
 private:
 
 	/**
 	 * Protect the constructor; only TSharedRefs of this class can be made.
 	 */
-	FLinuxWindow();
+	APPLICATIONCORE_API FLinuxWindow();
 
 private:
 
@@ -258,7 +258,7 @@ private:
 
 	/** Parent window as given by Slate */
 	TSharedPtr< FLinuxWindow > ParentWindow;
-	static SDL_HitTestResult HitTest( SDL_Window *SDLwin, const SDL_Point *point, void *data );
+	static APPLICATIONCORE_API SDL_HitTestResult HitTest( SDL_Window *SDLwin, const SDL_Point *point, void *data );
 
 	/** Cached width of the left border */
 	int32 LeftBorderWidth;
diff --git a/Engine/Source/Runtime/ApplicationCore/Public/Null/NullApplication.h b/Engine/Source/Runtime/ApplicationCore/Public/Null/NullApplication.h
--- a/Engine/Source/Runtime/ApplicationCore/Public/Null/NullApplication.h
+++ b/Engine/Source/Runtime/ApplicationCore/Public/Null/NullApplication.h
@@ -29,88 +29,88 @@ struct FNullPlatformDisplayMetrics : public FDisplayMetrics
  * An implementation of GenericApplication specifically for use when rendering off screen.
  * This application has no platform backing so instead keeps track of its associated NullWindows itself.
  */
-class APPLICATIONCORE_API FNullApplication : public GenericApplication, public FSelfRegisteringExec, public IInputInterface
+class FNullApplication : public GenericApplication, public FSelfRegisteringExec, public IInputInterface
 {
 public:
-	static FNullApplication* CreateNullApplication();
+	static APPLICATIONCORE_API FNullApplication* CreateNullApplication();
 
-	static void MoveWindowTo(FGenericWindow* Window, const int32 X, const int32 Y);
+	static APPLICATIONCORE_API void MoveWindowTo(FGenericWindow* Window, const int32 X, const int32 Y);
 
-	static void OnSizeChanged(FGenericWindow* Window, const int32 Width, const int32 Height);
+	static APPLICATIONCORE_API void OnSizeChanged(FGenericWindow* Window, const int32 Width, const int32 Height);
 
-	static void GetFullscreenInfo(int32& X, int32& Y, int32& Width, int32& Height);
+	static APPLICATIONCORE_API void GetFullscreenInfo(int32& X, int32& Y, int32& Width, int32& Height);
 
-	static void ShowWindow(FGenericWindow* Window);
+	static APPLICATIONCORE_API void ShowWindow(FGenericWindow* Window);
 
-	static void HideWindow(FGenericWindow* Window);
+	static APPLICATIONCORE_API void HideWindow(FGenericWindow* Window);
 
-	static void DestroyWindow(FGenericWindow* Window);
+	static APPLICATIONCORE_API void DestroyWindow(FGenericWindow* Window);
 
 public:
-	virtual ~FNullApplication();
+	APPLICATIONCORE_API virtual ~FNullApplication();
 
-	virtual void DestroyApplication() override;
+	APPLICATIONCORE_API virtual void DestroyApplication() override;
 
 	// FSelfRegisteringExec
-	virtual bool Exec(UWorld* InWorld, const TCHAR* Cmd, FOutputDevice& Ar) override;
+	APPLICATIONCORE_API virtual bool Exec(UWorld* InWorld, const TCHAR* Cmd, FOutputDevice& Ar) override;
 
-	TSharedPtr<FNullWindow> FindWindowByPtr(FGenericWindow* WindowToFind);
+	APPLICATIONCORE_API TSharedPtr<FNullWindow> FindWindowByPtr(FGenericWindow* WindowToFind);
 
-	void ActivateWindow(const TSharedPtr<FNullWindow>& Window);
+	APPLICATIONCORE_API void ActivateWindow(const TSharedPtr<FNullWindow>& Window);
 
 public:
-	virtual void SetMessageHandler(const TSharedRef<class FGenericApplicationMessageHandler>& InMessageHandler) override;
+	APPLICATIONCORE_API virtual void SetMessageHandler(const TSharedRef<class FGenericApplicationMessageHandler>& InMessageHandler) override;
 
-	virtual void PollGameDeviceState(const float TimeDelta) override;
+	APPLICATIONCORE_API virtual void PollGameDeviceState(const float TimeDelta) override;
 
-	virtual void PumpMessages(const float TimeDelta) override;
+	APPLICATIONCORE_API virtual void PumpMessages(const float TimeDelta) override;
 
-	virtual void ProcessDeferredEvents(const float TimeDelta) override;
+	APPLICATIONCORE_API virtual void ProcessDeferredEvents(const float TimeDelta) override;
 
-	virtual TSharedRef<FGenericWindow> MakeWindow() override;
+	APPLICATIONCORE_API virtual TSharedRef<FGenericWindow> MakeWindow() override;
 
-	virtual void InitializeWindow(const TSharedRef<FGenericWindow>& Window, const TSharedRef<FGenericWindowDefinition>& InDefinition, const TSharedPtr<FGenericWindow>& InParent, const bool bShowImmediately) override;
+	APPLICATIONCORE_API virtual void InitializeWindow(const TSharedRef<FGenericWindow>& Window, const TSharedRef<FGenericWindowDefinition>& InDefinition, const TSharedPtr<FGenericWindow>& InParent, const bool bShowImmediately) override;
 
-	void DestroyWindow(TSharedRef<FNullWindow> WindowToRemove);
+	APPLICATIONCORE_API void DestroyWindow(TSharedRef<FNullWindow> WindowToRemove);
 
-	virtual void SetCapture(const TSharedPtr<FGenericWindow>& InWindow) override;
+	APPLICATIONCORE_API virtual void SetCapture(const TSharedPtr<FGenericWindow>& InWindow) override;
 
-	virtual void* GetCapture(void) const override;
+	APPLICATIONCORE_API virtual void* GetCapture(void) const override;
 
-	virtual void SetHighPrecisionMouseMode(const bool Enable, const TSharedPtr<FGenericWindow>& InWindow) override;
+	APPLICATIONCORE_API virtual void SetHighPrecisionMouseMode(const bool Enable, const TSharedPtr<FGenericWindow>& InWindow) override;
 
 	virtual bool IsUsingHighPrecisionMouseMode() const override { return bUsingHighPrecisionMouseInput; }
 
-	virtual bool IsGamepadAttached() const override;
+	APPLICATIONCORE_API virtual bool IsGamepadAttached() const override;
 
-	virtual FModifierKeysState GetModifierKeys() const override;
+	APPLICATIONCORE_API virtual FModifierKeysState GetModifierKeys() const override;
 
-	virtual FPlatformRect GetWorkArea(const FPlatformRect& CurrentWindow) const override;
+	APPLICATIONCORE_API virtual FPlatformRect GetWorkArea(const FPlatformRect& CurrentWindow) const override;
 
-	void SetWorkArea(const FPlatformRect& NewWorkArea);
+	APPLICATIONCORE_API void SetWorkArea(const FPlatformRect& NewWorkArea);
 
 	virtual EWindowTransparency GetWindowTransparencySupport() const override
 	{
 		return EWindowTransparency::PerWindow;
 	}
 
-	virtual bool IsCursorDirectlyOverSlateWindow() const override;
+	APPLICATIONCORE_API virtual bool IsCursorDirectlyOverSlateWindow() const override;
 
-	virtual TSharedPtr<FGenericWindow> GetWindowUnderCursor() override;
+	APPLICATIONCORE_API virtual TSharedPtr<FGenericWindow> GetWindowUnderCursor() override;
 
 	virtual bool IsMouseAttached() const override { return true; }
 
 private:
-	FNullApplication();
+	APPLICATIONCORE_API FNullApplication();
 
 	/** Handles "Cursor" exec commands" */
-	bool HandleCursorCommand(const TCHAR* Cmd, FOutputDevice& Ar);
+	APPLICATIONCORE_API bool HandleCursorCommand(const TCHAR* Cmd, FOutputDevice& Ar);
 
 	/** Handles "Window" exec commands" */
-	bool HandleWindowCommand(const TCHAR* Cmd, FOutputDevice& Ar);
+	APPLICATIONCORE_API bool HandleWindowCommand(const TCHAR* Cmd, FOutputDevice& Ar);
 
 	/** Handles parsing the work area resolution from the command line */
-	bool ParseResolution(const TCHAR* InResolution, uint32& OutX, uint32& OutY);
+	APPLICATIONCORE_API bool ParseResolution(const TCHAR* InResolution, uint32& OutX, uint32& OutY);
 
 public:
 	virtual IInputInterface* GetInputInterface() override
@@ -119,9 +119,9 @@ public:
 	}
 
 	// IInputInterface overrides
-	virtual void SetForceFeedbackChannelValue(int32 ControllerId, FForceFeedbackChannelType ChannelType, float Value) override;
-	virtual void SetForceFeedbackChannelValues(int32 ControllerId, const FForceFeedbackValues& Values) override;
-	virtual void SetHapticFeedbackValues(int32 ControllerId, int32 Hand, const FHapticFeedbackValues& Values) override;
+	APPLICATIONCORE_API virtual void SetForceFeedbackChannelValue(int32 ControllerId, FForceFeedbackChannelType ChannelType, float Value) override;
+	APPLICATIONCORE_API virtual void SetForceFeedbackChannelValues(int32 ControllerId, const FForceFeedbackValues& Values) override;
+	APPLICATIONCORE_API virtual void SetHapticFeedbackValues(int32 ControllerId, int32 Hand, const FHapticFeedbackValues& Values) override;
 	virtual void SetLightColor(int32 ControllerId, FColor Color) override {}
 	virtual void ResetLightColor(int32 ControllerId) override {}
 
diff --git a/Engine/Source/Runtime/ApplicationCore/Public/Null/NullPlatformApplicationMisc.h b/Engine/Source/Runtime/ApplicationCore/Public/Null/NullPlatformApplicationMisc.h
--- a/Engine/Source/Runtime/ApplicationCore/Public/Null/NullPlatformApplicationMisc.h
+++ b/Engine/Source/Runtime/ApplicationCore/Public/Null/NullPlatformApplicationMisc.h
@@ -4,9 +4,9 @@
 
 #include "GenericPlatform/GenericPlatformApplicationMisc.h"
 
-struct APPLICATIONCORE_API FNullPlatformApplicationMisc : public FGenericPlatformApplicationMisc
+struct FNullPlatformApplicationMisc : public FGenericPlatformApplicationMisc
 {
-	static class GenericApplication* CreateApplication();
+	static APPLICATIONCORE_API class GenericApplication* CreateApplication();
 
-	static bool IsUsingNullApplication();
+	static APPLICATIONCORE_API bool IsUsingNullApplication();
 };
diff --git a/Engine/Source/Runtime/ApplicationCore/Public/Null/NullWindow.h b/Engine/Source/Runtime/ApplicationCore/Public/Null/NullWindow.h
--- a/Engine/Source/Runtime/ApplicationCore/Public/Null/NullWindow.h
+++ b/Engine/Source/Runtime/ApplicationCore/Public/Null/NullWindow.h
@@ -19,62 +19,62 @@ DECLARE_LOG_CATEGORY_EXTERN(LogNullWindowEvent, Log, All);
  * An implementation of FGenericWindow specifically for use when rendering off screen.
  * This window has no platform backing so instead keeps track of its position and other properties itself.
  */
-class APPLICATIONCORE_API FNullWindow : public FGenericWindow
+class FNullWindow : public FGenericWindow
 {
 public:
-	static TSharedRef<FNullWindow> Make();
+	static APPLICATIONCORE_API TSharedRef<FNullWindow> Make();
 
-	virtual ~FNullWindow();
+	APPLICATIONCORE_API virtual ~FNullWindow();
 
-	void Initialize(class FNullApplication* const Application, const TSharedRef<FGenericWindowDefinition>& InDefinition, const TSharedPtr<FNullWindow>& InParent, const bool bShowImmediately);
+	APPLICATIONCORE_API void Initialize(class FNullApplication* const Application, const TSharedRef<FGenericWindowDefinition>& InDefinition, const TSharedPtr<FNullWindow>& InParent, const bool bShowImmediately);
 
 	/** Native windows should implement ReshapeWindow by changing the platform-specific window to be located at (X,Y) and be the dimensions Width x Height. */
-	virtual void ReshapeWindow(int32 X, int32 Y, int32 Width, int32 Height) override;
+	APPLICATIONCORE_API virtual void ReshapeWindow(int32 X, int32 Y, int32 Width, int32 Height) override;
 
 	/** Returns the rectangle of the screen the window is associated with */
-	virtual bool GetFullScreenInfo(int32& X, int32& Y, int32& Width, int32& Height) const override;
+	APPLICATIONCORE_API virtual bool GetFullScreenInfo(int32& X, int32& Y, int32& Width, int32& Height) const override;
 
 	/** Native windows should implement MoveWindowTo by relocating the platform-specific window to (X,Y). */
-	virtual void MoveWindowTo(int32 X, int32 Y) override;
+	APPLICATIONCORE_API virtual void MoveWindowTo(int32 X, int32 Y) override;
 
 	/** Native windows should implement BringToFront by making this window the top-most window (i.e. focused). */
-	virtual void BringToFront(bool bForce = false) override;
+	APPLICATIONCORE_API virtual void BringToFront(bool bForce = false) override;
 
 	/** @hack Force a window to front even if a different application is in front. */
-	virtual void HACK_ForceToFront() override;
+	APPLICATIONCORE_API virtual void HACK_ForceToFront() override;
 
 	/** Native windows should implement this function by asking the OS to destroy OS-specific resource associated with the window (e.g. Win32 window handle) */
-	virtual void Destroy() override;
+	APPLICATIONCORE_API virtual void Destroy() override;
 
 	/** Native window should implement this function by performing the equivalent of the Win32 minimize-to-taskbar operation */
-	virtual void Minimize() override;
+	APPLICATIONCORE_API virtual void Minimize() override;
 
 	/** Native window should implement this function by performing the equivalent of the Win32 maximize operation */
-	virtual void Maximize() override;
+	APPLICATIONCORE_API virtual void Maximize() override;
 
 	/** Native window should implement this function by performing the equivalent of the Win32 restore operation */
-	virtual void Restore() override;
+	APPLICATIONCORE_API virtual void Restore() override;
 
 	/** Native window should make itself visible */
-	virtual void Show() override;
+	APPLICATIONCORE_API virtual void Show() override;
 
 	/** Native window should hide itself */
-	virtual void Hide() override;
+	APPLICATIONCORE_API virtual void Hide() override;
 
 	/** Toggle native window between fullscreen and normal mode */
-	virtual void SetWindowMode(EWindowMode::Type InNewWindowMode) override;
+	APPLICATIONCORE_API virtual void SetWindowMode(EWindowMode::Type InNewWindowMode) override;
 
 	/** @return true if the native window is currently in fullscreen mode, false otherwise */
-	virtual EWindowMode::Type GetWindowMode() const override;
+	APPLICATIONCORE_API virtual EWindowMode::Type GetWindowMode() const override;
 
 	/** @return true if the native window is maximized, false otherwise */
-	virtual bool IsMaximized() const override;
+	APPLICATIONCORE_API virtual bool IsMaximized() const override;
 
 	/** @return true if the native window is minimized, false otherwise */
-	virtual bool IsMinimized() const override;
+	APPLICATIONCORE_API virtual bool IsMinimized() const override;
 
 	/** @return true if the native window is visible, false otherwise */
-	virtual bool IsVisible() const override;
+	APPLICATIONCORE_API virtual bool IsVisible() const override;
 
 	/**
 	 * Populates the size and location of the window when it is restored.
@@ -82,61 +82,61 @@ public:
 	 *
 	 * @return true when the size and location and successfully retrieved; false otherwise.
 	 */
-	virtual bool GetRestoredDimensions(int32& X, int32& Y, int32& Width, int32& Height) override;
+	APPLICATIONCORE_API virtual bool GetRestoredDimensions(int32& X, int32& Y, int32& Width, int32& Height) override;
 
 	/**
 	 * Native windows should implement SetWindowFocus to let the OS know that a window has taken focus.
 	 * Slate handles focus on a per widget basis internally but the OS still needs to know what window has focus for proper message routing
 	 */
-	virtual void SetWindowFocus() override;
+	APPLICATIONCORE_API virtual void SetWindowFocus() override;
 
 	/**
 	 * Sets the opacity of this window
 	 *
 	 * @param	InOpacity	The new window opacity represented as a floating point scalar
 	 */
-	virtual void SetOpacity(const float InOpacity) override;
+	APPLICATIONCORE_API virtual void SetOpacity(const float InOpacity) override;
 
 	/**
 	 * Enables or disables the window.  If disabled the window receives no input
 	 *
 	 * @param bEnable	true to enable the window, false to disable it.
 	 */
-	virtual void Enable(bool bEnable) override;
+	APPLICATIONCORE_API virtual void Enable(bool bEnable) override;
 
 	/** @return true if native window exists underneath the coordinates */
-	virtual bool IsPointInWindow(int32 X, int32 Y) const override;
+	APPLICATIONCORE_API virtual bool IsPointInWindow(int32 X, int32 Y) const override;
 
 	/** Gets OS specific window border size. This is necessary because Win32 does not give control over this size. */
-	virtual int32 GetWindowBorderSize() const override;
+	APPLICATIONCORE_API virtual int32 GetWindowBorderSize() const override;
 
 	/** Gets OS specific window title bar size */
-	virtual int32 GetWindowTitleBarSize() const override;
+	APPLICATIONCORE_API virtual int32 GetWindowTitleBarSize() const override;
 
 	/** Gets the OS Window handle in the form of a void pointer for other API's */
-	virtual void* GetOSWindowHandle() const override;
+	APPLICATIONCORE_API virtual void* GetOSWindowHandle() const override;
 
 	/** @return true if the window is in the foreground */
-	virtual bool IsForegroundWindow() const override;
+	APPLICATIONCORE_API virtual bool IsForegroundWindow() const override;
 
 	/** @return true if the window is in the foreground */
-	virtual bool IsFullscreenSupported() const override;
+	APPLICATIONCORE_API virtual bool IsFullscreenSupported() const override;
 
 	/**
 	 * Sets the window text - usually the title but can also be text content for things like controls
 	 *
 	 * @param Text	The window's title or content text
 	 */
-	virtual void SetText(const TCHAR* const Text) override;
+	APPLICATIONCORE_API virtual void SetText(const TCHAR* const Text) override;
 
 	/** @return	The definition describing properties of the window */
-	virtual const FGenericWindowDefinition& GetDefinition() const override;
+	APPLICATIONCORE_API virtual const FGenericWindowDefinition& GetDefinition() const override;
 
 	/** @return	Returns true if the window definition is valid */
-	virtual bool IsDefinitionValid() const override;
+	APPLICATIONCORE_API virtual bool IsDefinitionValid() const override;
 
 	/** @return	Gives the native window a chance to adjust our stored window size before we cache it off */
-	virtual void AdjustCachedSize(FVector2D& Size) const override;
+	APPLICATIONCORE_API virtual void AdjustCachedSize(FVector2D& Size) const override;
 
 	virtual float GetDPIScaleFactor() const override
 	{
@@ -149,10 +149,10 @@ public:
 	}
 
 	/** determines whether or not this window does its own DPI management */
-	virtual bool IsManualManageDPIChanges() const override;
+	APPLICATIONCORE_API virtual bool IsManualManageDPIChanges() const override;
 
 	/** call with a true argument if this window need to do its custom size management in response to DPI variations */
-	virtual void SetManualManageDPIChanges(const bool bAutoHandle) override;
+	APPLICATIONCORE_API virtual void SetManualManageDPIChanges(const bool bAutoHandle) override;
 
 	/**
 	 * Attempts to draw the user's attention to this window in whatever way is
@@ -162,16 +162,16 @@ public:
 	 * @param Parameters The parameters for drawing attention. Depending on
 	 *        the platform, not all parameters may be supported.
 	 */
-	virtual void DrawAttention(const FWindowDrawAttentionParameters& Parameters) override;
+	APPLICATIONCORE_API virtual void DrawAttention(const FWindowDrawAttentionParameters& Parameters) override;
 
 	/** Shows or hides native window buttons on platforms that use them */
-	virtual void SetNativeWindowButtonsVisibility(bool bVisible);
+	APPLICATIONCORE_API virtual void SetNativeWindowButtonsVisibility(bool bVisible);
 
 private:
 	/**
 	 * Protect the constructor; only TSharedRefs of this class can be made.
 	 */
-	FNullWindow();
+	APPLICATIONCORE_API FNullWindow();
 
 	FNullApplication* OwningApplication;
 
diff --git a/Engine/Source/Runtime/ApplicationCore/Public/Windows/WindowsApplication.h b/Engine/Source/Runtime/ApplicationCore/Public/Windows/WindowsApplication.h
--- a/Engine/Source/Runtime/ApplicationCore/Public/Windows/WindowsApplication.h
+++ b/Engine/Source/Runtime/ApplicationCore/Public/Windows/WindowsApplication.h
@@ -70,7 +70,7 @@ namespace ETaskbarProgressState
  * This class can be used to change the appearance of a window's entry in the windows task bar,
  * such as setting an overlay icon or showing a progress indicator.
  */
-class APPLICATIONCORE_API FTaskbarList
+class FTaskbarList
 {
 public:
 
@@ -79,7 +79,7 @@ public:
 	 *
 	 * @return The new task bar list.
 	 */
-	static TSharedRef<FTaskbarList> Create();
+	static APPLICATIONCORE_API TSharedRef<FTaskbarList> Create();
 
 	/**
 	 * Sets the overlay icon of a task bar entry.
@@ -88,7 +88,7 @@ public:
 	 * @param Icon The overlay icon to set.
 	 * @param Description The overlay icon's description text.
 	 */
-	void SetOverlayIcon(HWND WindowHandle, HICON Icon, FText Description);
+	APPLICATIONCORE_API void SetOverlayIcon(HWND WindowHandle, HICON Icon, FText Description);
 
 	/**
 	 * Sets the progress state of a task bar entry.
@@ -96,7 +96,7 @@ public:
 	 * @param WindowHandle The window handle to change the progress state for.
 	 * @param State The new progress state.
 	 */
-	void SetProgressState(HWND WindowHandle, ETaskbarProgressState::Type State);
+	APPLICATIONCORE_API void SetProgressState(HWND WindowHandle, ETaskbarProgressState::Type State);
 
 	/**
 	 * Sets the progress value of a task bar entry.
@@ -105,10 +105,10 @@ public:
 	 * @param Current The current progress value.
 	 * @param Total The total progress value.
 	 */
-	void SetProgressValue(HWND WindowHandle, uint64 Current, uint64 Total);
+	APPLICATIONCORE_API void SetProgressValue(HWND WindowHandle, uint64 Current, uint64 Total);
 
 	/** Destructor. */
-	~FTaskbarList();
+	APPLICATIONCORE_API ~FTaskbarList();
 
 private:
 
@@ -280,7 +280,7 @@ public:
 /**
  * Windows-specific application implementation.
  */
-class APPLICATIONCORE_API FWindowsApplication
+class FWindowsApplication
 	: public GenericApplication
 	, public IInputInterface
 {
@@ -293,29 +293,29 @@ public:
 	 * @param IconHandle Win32 application icon handle.
 	 * @return New application object.
 	 */
-	static FWindowsApplication* CreateWindowsApplication( const HINSTANCE InstanceHandle, const HICON IconHandle );
+	static APPLICATIONCORE_API FWindowsApplication* CreateWindowsApplication( const HINSTANCE InstanceHandle, const HICON IconHandle );
 
 	/** Virtual destructor. */
-	virtual ~FWindowsApplication();
+	APPLICATIONCORE_API virtual ~FWindowsApplication();
 
 public:	
 
 	/** Called by a window when an OLE Drag and Drop operation occurred on a non-game thread */
-	void DeferDragDropOperation( const FDeferredWindowsDragDropOperation& DeferredDragDropOperation );
+	APPLICATIONCORE_API void DeferDragDropOperation( const FDeferredWindowsDragDropOperation& DeferredDragDropOperation );
 
-	TSharedPtr<FTaskbarList> GetTaskbarList();
+	APPLICATIONCORE_API TSharedPtr<FTaskbarList> GetTaskbarList();
 
 	/** Invoked by a window when an OLE Drag and Drop first enters it. */
-	HRESULT OnOLEDragEnter( const HWND HWnd, const FDragDropOLEData& OLEData, ::DWORD KeyState, POINTL CursorPosition, ::DWORD *CursorEffect);
+	APPLICATIONCORE_API HRESULT OnOLEDragEnter( const HWND HWnd, const FDragDropOLEData& OLEData, ::DWORD KeyState, POINTL CursorPosition, ::DWORD *CursorEffect);
 
 	/** Invoked by a window when an OLE Drag and Drop moves over the window. */
-	HRESULT OnOLEDragOver( const HWND HWnd, ::DWORD KeyState, POINTL CursorPosition, ::DWORD *CursorEffect);
+	APPLICATIONCORE_API HRESULT OnOLEDragOver( const HWND HWnd, ::DWORD KeyState, POINTL CursorPosition, ::DWORD *CursorEffect);
 
 	/** Invoked by a window when an OLE Drag and Drop exits the window. */
-	HRESULT OnOLEDragOut( const HWND HWnd );
+	APPLICATIONCORE_API HRESULT OnOLEDragOut( const HWND HWnd );
 
 	/** Invoked by a window when an OLE Drag and Drop is dropped onto the window. */
-	HRESULT OnOLEDrop( const HWND HWnd, const FDragDropOLEData& OLEData, ::DWORD KeyState, POINTL CursorPosition, ::DWORD *CursorEffect);
+	APPLICATIONCORE_API HRESULT OnOLEDrop( const HWND HWnd, const FDragDropOLEData& OLEData, ::DWORD KeyState, POINTL CursorPosition, ::DWORD *CursorEffect);
 
 	/**
 	 * Adds a Windows message handler with the application instance.
@@ -323,7 +323,7 @@ public:
 	 * @param MessageHandler The message handler to register.
 	 * @see RemoveMessageHandler
 	 */
-	virtual void AddMessageHandler(IWindowsMessageHandler& InMessageHandler);
+	APPLICATIONCORE_API virtual void AddMessageHandler(IWindowsMessageHandler& InMessageHandler);
 
 	/**
 	 * Removes a Windows message handler with the application instance.
@@ -331,35 +331,35 @@ public:
 	 * @param MessageHandler The message handler to register.
 	 * @see AddMessageHandler
 	 */
-	virtual void RemoveMessageHandler(IWindowsMessageHandler& InMessageHandler);
+	APPLICATIONCORE_API virtual void RemoveMessageHandler(IWindowsMessageHandler& InMessageHandler);
 
 public:
 
 	// GenericApplication overrides
 
-	virtual void SetMessageHandler( const TSharedRef< class FGenericApplicationMessageHandler >& InMessageHandler ) override;
+	APPLICATIONCORE_API virtual void SetMessageHandler( const TSharedRef< class FGenericApplicationMessageHandler >& InMessageHandler ) override;
 #if WITH_ACCESSIBILITY
-	virtual void SetAccessibleMessageHandler(const TSharedRef<FGenericAccessibleMessageHandler>& InAccessibleMessageHandler) override;
+	APPLICATIONCORE_API virtual void SetAccessibleMessageHandler(const TSharedRef<FGenericAccessibleMessageHandler>& InAccessibleMessageHandler) override;
 #endif
-	virtual void PollGameDeviceState( const float TimeDelta ) override;
-	virtual void PumpMessages( const float TimeDelta ) override;
-	virtual void ProcessDeferredEvents( const float TimeDelta ) override;
-	virtual TSharedRef< FGenericWindow > MakeWindow() override;
-	virtual void InitializeWindow( const TSharedRef< FGenericWindow >& Window, const TSharedRef< FGenericWindowDefinition >& InDefinition, const TSharedPtr< FGenericWindow >& InParent, const bool bShowImmediately ) override;
-	virtual void SetCapture( const TSharedPtr< FGenericWindow >& InWindow ) override;
-	virtual void* GetCapture( void ) const override;
+	APPLICATIONCORE_API virtual void PollGameDeviceState( const float TimeDelta ) override;
+	APPLICATIONCORE_API virtual void PumpMessages( const float TimeDelta ) override;
+	APPLICATIONCORE_API virtual void ProcessDeferredEvents( const float TimeDelta ) override;
+	APPLICATIONCORE_API virtual TSharedRef< FGenericWindow > MakeWindow() override;
+	APPLICATIONCORE_API virtual void InitializeWindow( const TSharedRef< FGenericWindow >& Window, const TSharedRef< FGenericWindowDefinition >& InDefinition, const TSharedPtr< FGenericWindow >& InParent, const bool bShowImmediately ) override;
+	APPLICATIONCORE_API virtual void SetCapture( const TSharedPtr< FGenericWindow >& InWindow ) override;
+	APPLICATIONCORE_API virtual void* GetCapture( void ) const override;
 	virtual bool IsMinimized() const override { return bMinimized; }
-	virtual void SetHighPrecisionMouseMode( const bool Enable, const TSharedPtr< FGenericWindow >& InWindow ) override;
+	APPLICATIONCORE_API virtual void SetHighPrecisionMouseMode( const bool Enable, const TSharedPtr< FGenericWindow >& InWindow ) override;
 	virtual bool IsUsingHighPrecisionMouseMode() const override { return bUsingHighPrecisionMouseInput; }
 	virtual bool IsMouseAttached() const override { return bIsMouseAttached; }
-	virtual bool IsGamepadAttached() const override;
-	virtual FModifierKeysState GetModifierKeys() const override;
-	virtual bool IsCursorDirectlyOverSlateWindow() const override;
-	virtual FPlatformRect GetWorkArea( const FPlatformRect& CurrentWindow ) const override;
-	virtual void GetInitialDisplayMetrics( FDisplayMetrics& OutDisplayMetrics ) const override;
-	virtual EWindowTitleAlignment::Type GetWindowTitleAlignment() const override;
-	virtual EWindowTransparency GetWindowTransparencySupport() const override;
-	virtual void DestroyApplication() override;
+	APPLICATIONCORE_API virtual bool IsGamepadAttached() const override;
+	APPLICATIONCORE_API virtual FModifierKeysState GetModifierKeys() const override;
+	APPLICATIONCORE_API virtual bool IsCursorDirectlyOverSlateWindow() const override;
+	APPLICATIONCORE_API virtual FPlatformRect GetWorkArea( const FPlatformRect& CurrentWindow ) const override;
+	APPLICATIONCORE_API virtual void GetInitialDisplayMetrics( FDisplayMetrics& OutDisplayMetrics ) const override;
+	APPLICATIONCORE_API virtual EWindowTitleAlignment::Type GetWindowTitleAlignment() const override;
+	APPLICATIONCORE_API virtual EWindowTransparency GetWindowTransparencySupport() const override;
+	APPLICATIONCORE_API virtual void DestroyApplication() override;
 
 	virtual IInputInterface* GetInputInterface() override
 	{
@@ -371,42 +371,42 @@ public:
 		return TextInputMethodSystem.Get();
 	}
 
-	virtual void AddExternalInputDevice(TSharedPtr<class IInputDevice> InputDevice);
+	APPLICATIONCORE_API virtual void AddExternalInputDevice(TSharedPtr<class IInputDevice> InputDevice);
 
-	virtual void FinishedInputThisFrame() override;
+	APPLICATIONCORE_API virtual void FinishedInputThisFrame() override;
 public:
 
 	// IInputInterface overrides
 
-	virtual void SetForceFeedbackChannelValue (int32 ControllerId, FForceFeedbackChannelType ChannelType, float Value) override;
-	virtual void SetForceFeedbackChannelValues(int32 ControllerId, const FForceFeedbackValues &Values) override;
-	virtual void SetHapticFeedbackValues(int32 ControllerId, int32 Hand, const FHapticFeedbackValues& Values) override;
+	APPLICATIONCORE_API virtual void SetForceFeedbackChannelValue (int32 ControllerId, FForceFeedbackChannelType ChannelType, float Value) override;
+	APPLICATIONCORE_API virtual void SetForceFeedbackChannelValues(int32 ControllerId, const FForceFeedbackValues &Values) override;
+	APPLICATIONCORE_API virtual void SetHapticFeedbackValues(int32 ControllerId, int32 Hand, const FHapticFeedbackValues& Values) override;
 	virtual void SetLightColor(int32 ControllerId, FColor Color) override { }
 	virtual void ResetLightColor(int32 ControllerId) override { }
-	virtual void SetDeviceProperty(int32 ControllerId, const FInputDeviceProperty* Property) override;
+	APPLICATIONCORE_API virtual void SetDeviceProperty(int32 ControllerId, const FInputDeviceProperty* Property) override;
 
 protected:
 	friend LRESULT WindowsApplication_WndProc(HWND hwnd, uint32 msg, WPARAM wParam, LPARAM lParam);
 
 	/** Windows callback for message processing (forwards messages to the FWindowsApplication instance). */
-	static LRESULT CALLBACK AppWndProc(HWND hwnd, uint32 msg, WPARAM wParam, LPARAM lParam);
+	static APPLICATIONCORE_API LRESULT CALLBACK AppWndProc(HWND hwnd, uint32 msg, WPARAM wParam, LPARAM lParam);
 
 	/** Processes a single Windows message. */
-	int32 ProcessMessage( HWND hwnd, uint32 msg, WPARAM wParam, LPARAM lParam );
+	APPLICATIONCORE_API int32 ProcessMessage( HWND hwnd, uint32 msg, WPARAM wParam, LPARAM lParam );
 
 	/** Processes a deferred Windows message. */
-	int32 ProcessDeferredMessage( const FDeferredWindowsMessage& DeferredMessage );
+	APPLICATIONCORE_API int32 ProcessDeferredMessage( const FDeferredWindowsMessage& DeferredMessage );
 
 	/** Processes deferred drag and drop operations. */
-	void ProcessDeferredDragDropOperation(const FDeferredWindowsDragDropOperation& Op);
+	APPLICATIONCORE_API void ProcessDeferredDragDropOperation(const FDeferredWindowsDragDropOperation& Op);
 
 	/** Hidden constructor. */
-	FWindowsApplication( const HINSTANCE HInstance, const HICON IconHandle );
+	APPLICATIONCORE_API FWindowsApplication( const HINSTANCE HInstance, const HICON IconHandle );
 
-	void ApplyLowLevelMouseFilter();
-	void RemoveLowLevelMouseFilter();
+	APPLICATIONCORE_API void ApplyLowLevelMouseFilter();
+	APPLICATIONCORE_API void RemoveLowLevelMouseFilter();
 
-	static LRESULT CALLBACK HandleLowLevelMouseFilterHook(int nCode, WPARAM wParam, LPARAM lParam);
+	static APPLICATIONCORE_API LRESULT CALLBACK HandleLowLevelMouseFilterHook(int nCode, WPARAM wParam, LPARAM lParam);
 
 	HHOOK LowLevelMouseFilterHook;
 	bool bLowLevelMouseFilterIsApplied = false;
@@ -414,49 +414,49 @@ protected:
 private:
 
 	/** Registers the Windows class for windows and assigns the application instance and icon */
-	static bool RegisterClass( const HINSTANCE HInstance, const HICON HIcon );
+	static APPLICATIONCORE_API bool RegisterClass( const HINSTANCE HInstance, const HICON HIcon );
 
 	/**  @return  True if a windows message is related to user input from the keyboard */
-	static bool IsKeyboardInputMessage( uint32 msg );
+	static APPLICATIONCORE_API bool IsKeyboardInputMessage( uint32 msg );
 
 	/**  @return  True if a windows message is related to user input from the mouse */
-	static bool IsMouseInputMessage( uint32 msg );
+	static APPLICATIONCORE_API bool IsMouseInputMessage( uint32 msg );
 
 	/**  @return  True if a windows message is a fake mouse input message generated after a WM_TOUCH event */
-	static bool IsFakeMouseInputMessage(uint32 msg);
+	static APPLICATIONCORE_API bool IsFakeMouseInputMessage(uint32 msg);
 
 	/**  @return  True if a windows message is related to user input (mouse, keyboard) */
-	static bool IsInputMessage( uint32 msg );
+	static APPLICATIONCORE_API bool IsInputMessage( uint32 msg );
 
 	/** Defers a Windows message for later processing. */
-	void DeferMessage( TSharedPtr<FWindowsWindow>& NativeWindow, HWND InHWnd, uint32 InMessage, WPARAM InWParam, LPARAM InLParam, int32 MouseX = 0, int32 MouseY = 0, uint32 RawInputFlags = 0 );
+	APPLICATIONCORE_API void DeferMessage( TSharedPtr<FWindowsWindow>& NativeWindow, HWND InHWnd, uint32 InMessage, WPARAM InWParam, LPARAM InLParam, int32 MouseX = 0, int32 MouseY = 0, uint32 RawInputFlags = 0 );
 
 	/** Checks a key code for release of the Shift key. */
-	void CheckForShiftUpEvents(const int32 KeyCode);
+	APPLICATIONCORE_API void CheckForShiftUpEvents(const int32 KeyCode);
 
 	/** Shuts down the application (called after an unrecoverable error occurred). */
-	void ShutDownAfterError();
+	APPLICATIONCORE_API void ShutDownAfterError();
 
 	/** Enables or disables Windows accessibility features, such as sticky keys. */
-	void AllowAccessibilityShortcutKeys(const bool bAllowKeys);
+	APPLICATIONCORE_API void AllowAccessibilityShortcutKeys(const bool bAllowKeys);
 
 	/** Queries and caches the number of connected mouse devices. */
-	void QueryConnectedMice();
+	APPLICATIONCORE_API void QueryConnectedMice();
 
 	/** Gets the touch index for a given windows touch ID. */
-	uint32 GetTouchIndexForID(int32 TouchID);
+	APPLICATIONCORE_API uint32 GetTouchIndexForID(int32 TouchID);
 
 	/** Searches for a free touch index. */
-	uint32 GetFirstFreeTouchIndex();
+	APPLICATIONCORE_API uint32 GetFirstFreeTouchIndex();
 
 	/** Helper function to update the cached states of all modifier keys */
-	void UpdateAllModifierKeyStates();
+	APPLICATIONCORE_API void UpdateAllModifierKeyStates();
 
-	FPlatformRect GetWorkAreaFromOS(const FPlatformRect& CurrentWindow) const;
+	APPLICATIONCORE_API FPlatformRect GetWorkAreaFromOS(const FPlatformRect& CurrentWindow) const;
 
 private:
 
-	static const FIntPoint MinimizedWindowPosition;
+	static APPLICATIONCORE_API const FIntPoint MinimizedWindowPosition;
 
 	HINSTANCE InstanceHandle;
 
diff --git a/Engine/Source/Runtime/ApplicationCore/Public/Windows/WindowsConsoleOutputDevice.h b/Engine/Source/Runtime/ApplicationCore/Public/Windows/WindowsConsoleOutputDevice.h
--- a/Engine/Source/Runtime/ApplicationCore/Public/Windows/WindowsConsoleOutputDevice.h
+++ b/Engine/Source/Runtime/ApplicationCore/Public/Windows/WindowsConsoleOutputDevice.h
@@ -9,7 +9,7 @@
 /**
  * Windows implementation of console log window, utilizing the Win32 console API
  */
-class APPLICATIONCORE_API FWindowsConsoleOutputDevice : public FOutputDeviceConsole
+class FWindowsConsoleOutputDevice : public FOutputDeviceConsole
 {
 private:
 	/** Handle to the console log window */
@@ -64,27 +64,27 @@ public:
 	/** 
 	 * Constructor, setting console control handler.
 	 */
-	FWindowsConsoleOutputDevice();
-	~FWindowsConsoleOutputDevice();
+	APPLICATIONCORE_API FWindowsConsoleOutputDevice();
+	APPLICATIONCORE_API ~FWindowsConsoleOutputDevice();
 
 	/**
 	 * Shows or hides the console window. 
 	 *
 	 * @param ShowWindow	Whether to show (true) or hide (false) the console window.
 	 */
-	virtual void Show( bool ShowWindow );
+	APPLICATIONCORE_API virtual void Show( bool ShowWindow );
 
 	/** 
 	 * Returns whether console is currently shown or not
 	 *
 	 * @return true if console is shown, false otherwise
 	 */
-	virtual bool IsShown();
+	APPLICATIONCORE_API virtual bool IsShown();
 
-	virtual bool IsAttached();
+	APPLICATIONCORE_API virtual bool IsAttached();
 
-	virtual bool CanBeUsedOnAnyThread() const override;
-	virtual bool CanBeUsedOnPanicThread() const override;
+	APPLICATIONCORE_API virtual bool CanBeUsedOnAnyThread() const override;
+	APPLICATIONCORE_API virtual bool CanBeUsedOnPanicThread() const override;
 
 	/**
 	 * Displays text on the console and scrolls if necessary.
@@ -92,10 +92,10 @@ public:
 	 * @param Data	Text to display
 	 * @param Event	Event type, used for filtering/ suppression
 	 */
-	virtual void Serialize( const TCHAR* Data, ELogVerbosity::Type Verbosity, const class FName& Category, const double Time ) override;
-	virtual void Serialize( const TCHAR* Data, ELogVerbosity::Type Verbosity, const class FName& Category ) override;
+	APPLICATIONCORE_API virtual void Serialize( const TCHAR* Data, ELogVerbosity::Type Verbosity, const class FName& Category, const double Time ) override;
+	APPLICATIONCORE_API virtual void Serialize( const TCHAR* Data, ELogVerbosity::Type Verbosity, const class FName& Category ) override;
 
-	void SetColor( const TCHAR* Color );
+	APPLICATIONCORE_API void SetColor( const TCHAR* Color );
 };
 
 
diff --git a/Engine/Source/Runtime/ApplicationCore/Public/Windows/WindowsConsoleOutputDevice2.h b/Engine/Source/Runtime/ApplicationCore/Public/Windows/WindowsConsoleOutputDevice2.h
--- a/Engine/Source/Runtime/ApplicationCore/Public/Windows/WindowsConsoleOutputDevice2.h
+++ b/Engine/Source/Runtime/ApplicationCore/Public/Windows/WindowsConsoleOutputDevice2.h
@@ -9,7 +9,7 @@
 /**
  * Windows implementation of console log window, utilizing the Win32 console API
  */
-class APPLICATIONCORE_API FWindowsConsoleOutputDevice2 final : public FOutputDeviceConsole
+class FWindowsConsoleOutputDevice2 final : public FOutputDeviceConsole
 {
 private:
 	uint16 TextAttribute;
@@ -66,26 +66,26 @@ public:
 	/** 
 	 * Constructor, setting console control handler.
 	 */
-	FWindowsConsoleOutputDevice2();
-	~FWindowsConsoleOutputDevice2();
+	APPLICATIONCORE_API FWindowsConsoleOutputDevice2();
+	APPLICATIONCORE_API ~FWindowsConsoleOutputDevice2();
 
 	/**
 	 * Shows or hides the console window. 
 	 *
 	 * @param ShowWindow	Whether to show (true) or hide (false) the console window.
 	 */
-	virtual void Show( bool ShowWindow );
+	APPLICATIONCORE_API virtual void Show( bool ShowWindow );
 
 	/** 
 	 * Returns whether console is currently shown or not
 	 *
 	 * @return true if console is shown, false otherwise
 	 */
-	virtual bool IsShown();
+	APPLICATIONCORE_API virtual bool IsShown();
 
-	virtual bool IsAttached();
+	APPLICATIONCORE_API virtual bool IsAttached();
 
-	virtual bool CanBeUsedOnAnyThread() const override;
+	APPLICATIONCORE_API virtual bool CanBeUsedOnAnyThread() const override;
 
 	/**
 	 * Displays text on the console and scrolls if necessary.
@@ -93,10 +93,10 @@ public:
 	 * @param Data	Text to display
 	 * @param Event	Event type, used for filtering/ suppression
 	 */
-	virtual void Serialize( const TCHAR* Data, ELogVerbosity::Type Verbosity, const class FName& Category, const double Time ) override;
-	virtual void Serialize( const TCHAR* Data, ELogVerbosity::Type Verbosity, const class FName& Category ) override;
+	APPLICATIONCORE_API virtual void Serialize( const TCHAR* Data, ELogVerbosity::Type Verbosity, const class FName& Category, const double Time ) override;
+	APPLICATIONCORE_API virtual void Serialize( const TCHAR* Data, ELogVerbosity::Type Verbosity, const class FName& Category ) override;
 
-	void SetColor( const TCHAR* Color );
+	APPLICATIONCORE_API void SetColor( const TCHAR* Color );
 
 private:
 	FWindowsConsoleOutputDevice2(const FWindowsConsoleOutputDevice2&) = delete;
diff --git a/Engine/Source/Runtime/ApplicationCore/Public/Windows/WindowsFeedbackContext.h b/Engine/Source/Runtime/ApplicationCore/Public/Windows/WindowsFeedbackContext.h
--- a/Engine/Source/Runtime/ApplicationCore/Public/Windows/WindowsFeedbackContext.h
+++ b/Engine/Source/Runtime/ApplicationCore/Public/Windows/WindowsFeedbackContext.h
@@ -12,13 +12,13 @@
 /**
  * Feedback context implementation for windows.
  */
-class APPLICATIONCORE_API FWindowsFeedbackContext : public FFeedbackContext
+class FWindowsFeedbackContext : public FFeedbackContext
 {
 	/** Context information for warning and error messages */
 	FContextSupplier* Context = nullptr;
 
 public:
-	bool YesNof(const FText& Question) override;
+	APPLICATIONCORE_API bool YesNof(const FText& Question) override;
 
 	FContextSupplier* GetContext() const { return Context; }
 	void SetContext(FContextSupplier* InContext) { Context = InContext; }
