Accept pull request fixing mouse input on Wayland applications

#10582

diff --git a/Engine/Source/Runtime/ApplicationCore/Private/Linux/LinuxApplication.cpp b/Engine/Source/Runtime/ApplicationCore/Private/Linux/LinuxApplication.cpp
--- a/Engine/Source/Runtime/ApplicationCore/Private/Linux/LinuxApplication.cpp
+++ b/Engine/Source/Runtime/ApplicationCore/Private/Linux/LinuxApplication.cpp
@@ -24,6 +24,18 @@ namespace
 {
 	// How long we wait from a FocusOut event to deactivate the application (100ms default)
 	double DeactivationThreadshold = 0.1;
+
+	// Mask away the higher bits (but preserve negative flags), as SDL uses those as flags
+	// See: SDL_WINDOWPOS_UNDEFINED_MASK & SDL_WINDOWPOS_CENTERED_MASK for context
+	int MaskAwayHigherBits(int Input)
+	{
+		if (Input < 0)
+		{
+			return -(-Input & 0xFFFF);
+		}
+
+		return Input & 0xFFFF;
+	}
 }
 
 float ShortToNormalFloat(short AxisVal)
@@ -796,25 +808,8 @@ void FLinuxApplication::ProcessDeferredMessage( SDL_Event Event )
 					{
 						// Mask away the higher bits (but preserve negative flags), as SDL uses those as flags
 						// See: SDL_WINDOWPOS_UNDEFINED_MASK & SDL_WINDOWPOS_CENTERED_MASK for context
-						int32 ClientScreenX;
-						int32 ClientScreenY;
-
-						if (windowEvent.data1 < 0)
-						{
-							ClientScreenX = -(-windowEvent.data1 & 0xFFFF);
-						}
-						else
-						{
-							ClientScreenX = windowEvent.data1 & 0xFFFF;
-						}
-						if (windowEvent.data2 < 0)
-						{
-							ClientScreenY = -(-windowEvent.data2 & 0xFFFF);
-						}
-						else
-						{
-							ClientScreenY = windowEvent.data2 & 0xFFFF;
-						}
+						int32 ClientScreenX = MaskAwayHigherBits(windowEvent.data1);
+						int32 ClientScreenY = MaskAwayHigherBits(windowEvent.data2);
 
 						int32 BorderSizeX, BorderSizeY;
 						CurrentEventWindow->GetNativeBordersSize(BorderSizeX, BorderSizeY);
@@ -1897,6 +1892,8 @@ void FLinuxApplication::SaveWindowPropertiesForEventLoop(void)
 		int Height = 0;
 		SDL_HWindow NativeWindow = Window->GetHWnd();
 		SDL_GetWindowPosition(NativeWindow, &X, &Y);
+		X = MaskAwayHigherBits(X);
+		Y = MaskAwayHigherBits(Y);
 		SDL_GetWindowSize(NativeWindow, &Width, &Height);
 
 		FWindowProperties Props;
@@ -1949,6 +1946,8 @@ void FLinuxApplication::GetWindowPropertiesInEventLoop(SDL_HWindow NativeWindow,
 	{
 		int X, Y, Width, Height;
 		SDL_GetWindowPosition(NativeWindow, &X, &Y);
+		X = MaskAwayHigherBits(X);
+		Y = MaskAwayHigherBits(Y);
 		SDL_GetWindowSize(NativeWindow, &Width, &Height);
 		Properties.Location = FVector2D(static_cast<float>(X), static_cast<float>(Y));
 		Properties.Size = FVector2D(static_cast<float>(Width), static_cast<float>(Height));
