Fix: Mouse locations are off when rendering off screen. This is due to the platform work area being much larger than the application size. The NullPlatform work area now listens for resolution updates and updates is size accordingly, effectively meaning the desktop size is always the size of the application.

diff --git a/Engine/Plugins/Media/PixelStreaming/Source/PixelStreamingEditor/Private/PixelStreamingEditorModule.cpp b/Engine/Plugins/Media/PixelStreaming/Source/PixelStreamingEditor/Private/PixelStreamingEditorModule.cpp
--- a/Engine/Plugins/Media/PixelStreaming/Source/PixelStreamingEditor/Private/PixelStreamingEditorModule.cpp
+++ b/Engine/Plugins/Media/PixelStreaming/Source/PixelStreamingEditor/Private/PixelStreamingEditorModule.cpp
@@ -145,6 +145,10 @@ void FPixelStreamingEditorModule::StartStreaming(UE::EditorPixelStreaming::EStre
 				TSharedPtr<SWindow> ParentWindow = IMainFrameModule::Get().GetParentWindow();
 				ParentWindow->Resize(FVector2D(Width, Height));
 				FSlateApplication::Get().OnSizeChanged(ParentWindow->GetNativeWindow().ToSharedRef(), Width, Height);
+                // Triggers the NullApplication to rebuild its DisplayMetrics with the new resolution and inform slate 
+                // about the updated virtual desktop size
+                FSystemResolution::RequestResolutionChange(Width, Height, GSystemResolution.WindowMode);
+	            IConsoleManager::Get().CallAllConsoleVariableSinks();
 			});
 	}
 
@@ -359,8 +363,13 @@ void FPixelStreamingEditorModule::MaybeResizeEditor(TSharedPtr<SWindow> RootWind
 
 	if (bSuccess)
 	{
+        // Update editor window size
 		RootWindow->Resize(FVector2D(ResolutionX, ResolutionY));
 		FSlateApplication::Get().OnSizeChanged(RootWindow->GetNativeWindow().ToSharedRef(), ResolutionX, ResolutionY);
+        // Triggers the NullApplication to rebuild its DisplayMetrics with the new resolution and inform slate 
+        // about the updated virtual desktop size
+        FSystemResolution::RequestResolutionChange(ResolutionX, ResolutionY, GSystemResolution.WindowMode);
+	    IConsoleManager::Get().CallAllConsoleVariableSinks();
 	}
 }
 
diff --git a/Engine/Source/Runtime/ApplicationCore/Private/Null/NullApplication.cpp b/Engine/Source/Runtime/ApplicationCore/Private/Null/NullApplication.cpp
--- a/Engine/Source/Runtime/ApplicationCore/Private/Null/NullApplication.cpp
+++ b/Engine/Source/Runtime/ApplicationCore/Private/Null/NullApplication.cpp
@@ -94,40 +94,32 @@ FNullApplication::FNullApplication()
 	// TODO (william.belcher): There's no way of grabbing a work area from an OS that doesn't exist.
 	// Therefore, we have theoretically uncapped minimum and maximum. Let's just say you've got a
 	// 4096x4096 square as that's the maximum supported resolution for NvEnc
-	WorkArea.Left = 0;
+    WorkArea.Left = 0;
 	WorkArea.Top = 0;
 	WorkArea.Right = 4096;
 	WorkArea.Bottom = 4096;
 
-	uint32 ResolutionX, ResolutionY = 0;
-	FString ResolutionStr;
-	bool bSuccess = FParse::Value(FCommandLine::Get(), TEXT("NullApplicationWorkArea="), ResolutionStr);
-	if (bSuccess)
-	{
-		bSuccess = ParseResolution(*ResolutionStr, ResolutionX, ResolutionY);
-	}
-	else
-	{
-		bool UserSpecifiedWidth = FParse::Value(FCommandLine::Get(), TEXT("NullApplicationWorkAreaX="), ResolutionX);
-		bool UserSpecifiedHeight = FParse::Value(FCommandLine::Get(), TEXT("NullApplicationWorkAreaY="), ResolutionY);
-		bSuccess = UserSpecifiedWidth | UserSpecifiedHeight;
 
-		const float AspectRatio = 16.f / 9.f;
-		if (UserSpecifiedWidth && !UserSpecifiedHeight)
-		{
-			ResolutionY = FMath::FloorToInt((float)ResolutionX / AspectRatio);
-		}
-		else if (UserSpecifiedHeight && !UserSpecifiedWidth)
-		{
-			ResolutionX = FMath::FloorToInt((float)ResolutionY * AspectRatio);
-		}
-	}
+    // Default work area to the config res
+    static const IConsoleVariable* CVarSystemResolution = IConsoleManager::Get().FindConsoleVariable(TEXT("r.SetRes"));
+    FString ResolutionStr = CVarSystemResolution->GetString();
 
-	if (bSuccess)
-	{
-		WorkArea.Right = ResolutionX;
+    uint32 ResolutionX, ResolutionY = 0;
+    if(ParseResolution(*ResolutionStr, ResolutionX, ResolutionY))
+    {
+        WorkArea.Right = ResolutionX;
+        WorkArea.Bottom = ResolutionY;
+    }
+	
+    // Listen for res changes to update our virtual display
+    FCoreDelegates::OnSystemResolutionChanged.AddLambda([this](uint32 ResolutionX, uint32 ResolutionY) {
+        WorkArea.Right = ResolutionX;
 		WorkArea.Bottom = ResolutionY;
-	}
+
+        FDisplayMetrics DisplayMetrics;
+		FNullPlatformDisplayMetrics::RebuildDisplayMetrics(DisplayMetrics);
+        BroadcastDisplayMetricsChanged(DisplayMetrics);
+    });
 }
 
 FNullApplication::~FNullApplication()
diff --git a/Engine/Source/Runtime/ApplicationCore/Public/Null/NullApplication.h b/Engine/Source/Runtime/ApplicationCore/Public/Null/NullApplication.h
--- a/Engine/Source/Runtime/ApplicationCore/Public/Null/NullApplication.h
+++ b/Engine/Source/Runtime/ApplicationCore/Public/Null/NullApplication.h
@@ -10,6 +10,7 @@
 #include "Math/Vector2D.h"
 #include "Templates/SharedPointer.h"
 #include "Misc/CoreMisc.h"
+#include "Misc/CoreDelegates.h"
 #include "GenericPlatform/GenericWindow.h"
 #include "GenericPlatform/GenericApplicationMessageHandler.h"
 #include "GenericPlatform/GenericWindowDefinition.h"
