Add alternative method to determine the native screen DPI when calculating the scale factor required to achieve the requested Android.WindowDPI.

diff --git a/Engine/Source/Runtime/ApplicationCore/Private/Android/AndroidWindow.cpp b/Engine/Source/Runtime/ApplicationCore/Private/Android/AndroidWindow.cpp
--- a/Engine/Source/Runtime/ApplicationCore/Private/Android/AndroidWindow.cpp
+++ b/Engine/Source/Runtime/ApplicationCore/Private/Android/AndroidWindow.cpp
@@ -23,6 +23,15 @@ static FAutoConsoleVariableRef CVarAndroidWindowDPI(
 	TEXT("default: 0"),
 	ECVF_ReadOnly);
 
+int32 GAndroidWindowDPIQueryMethod = 0;
+static FAutoConsoleVariableRef CVarAndroidWindowDPIQueryMethod(
+	TEXT("Android.DPIQueryMethod"),
+	GAndroidWindowDPIQueryMethod,
+	TEXT("The method used to determine the native screen DPI when calculating the scale factor required to achieve the requested Android.WindowDPI.\n")
+	TEXT("0: Use displaymetrics xdpi/ydpi (default)\n")
+	TEXT("1: Use displaymetrics densityDpi"),
+	ECVF_ReadOnly);
+
 int32 GAndroid3DSceneMaxDesiredPixelCount = 0;
 static FAutoConsoleVariableRef CVarAndroid3DSceneMaxDesiredPixelCount(
 	TEXT("Android.3DSceneMaxDesiredPixelCount"),
@@ -600,9 +609,34 @@ static FAndroidDisplayInfo GetAndroidDisplayInfoFromDPITargets(int32 TargetDPI,
 	TMap<FString, FString> metricsParams = StringToMap(JNIMetrics);
 	TMap<FString, FString> displayParams = StringToMap(JNIDisplay);
 	static const FString DensityDpiKey(TEXT("densityDpi"));
+	static const FString xDpiKey(TEXT("xdpi"));
+	static const FString yDpiKey(TEXT("ydpi"));
 	static const FString WidthPixelsKey(TEXT("realWidth"));
 	static const FString HeightPixelsKey(TEXT("realHeight"));
-	int32 NativeScreenDensityDPI = metricsParams.Contains(TEXT("densityDpi")) ? FCString::Atoi(*metricsParams.FindChecked(DensityDpiKey)) : 0;
+
+	int32 DensityDPI = metricsParams.Contains(TEXT("densityDpi")) ? FCString::Atoi(*metricsParams.FindChecked(DensityDpiKey)) : 0;
+	int32 xdpi = metricsParams.Contains(TEXT("xdpi")) ? FCString::Atoi(*metricsParams.FindChecked(xDpiKey)) : 0;
+	int32 ydpi = metricsParams.Contains(TEXT("ydpi")) ? FCString::Atoi(*metricsParams.FindChecked(yDpiKey)) : 0;
+	int32 avgdpi = (xdpi + ydpi) / 2;
+
+	UE_LOG(LogAndroid, Display, TEXT("AndroidDisplayInfoFromDPITargets : DPI info: DensityDPI %d, dpi %d, xdpi %d, ydpi %d"), DensityDPI, avgdpi, xdpi, ydpi);
+
+	int32 NativeScreenDensityDPI;
+	switch (GAndroidWindowDPIQueryMethod)
+	{
+		case 1:
+		{
+			NativeScreenDensityDPI = DensityDPI;
+			break;
+		}
+		case 0:
+		default:
+		{
+			NativeScreenDensityDPI = avgdpi;
+			break;
+		}
+	}
+
 	FIntVector2 NativeScreenPixelDims(
 		displayParams.Contains(WidthPixelsKey) ? FCString::Atoi(*displayParams.FindChecked(WidthPixelsKey)) : 0
 		, displayParams.Contains(HeightPixelsKey) ? FCString::Atoi(*displayParams.FindChecked(HeightPixelsKey)) : 0
