- Fixed IOS build errors

diff --git a/Engine/Source/Runtime/ApplicationCore/Private/IOS/IOSWindow.cpp b/Engine/Source/Runtime/ApplicationCore/Private/IOS/IOSWindow.cpp
--- a/Engine/Source/Runtime/ApplicationCore/Private/IOS/IOSWindow.cpp
+++ b/Engine/Source/Runtime/ApplicationCore/Private/IOS/IOSWindow.cpp
@@ -134,6 +134,7 @@ FPlatformRect FIOSWindow::GetScreenRect()
 		ScreenRect.Bottom = FMath::TruncToInt((Frame.origin.y + Frame.size.height) * Scale);
 		ScreenRect.Right = FMath::TruncToInt((Frame.origin.x + Frame.size.width) * Scale);
 	}
+#if PLATFORM_VISIONOS
 	else if (AppDelegate.SwiftLayer != nullptr)
 	{
 		// using the last viewport since this will work with 1 eye or 2 eye setup, to get the full size of the screen
@@ -143,7 +144,7 @@ FPlatformRect FIOSWindow::GetScreenRect()
 		ScreenRect.Right = (int)(LastViewport.origin.x + LastViewport.size.width);
 		ScreenRect.Bottom = (int)(LastViewport.origin.y + LastViewport.size.height);
 	}
-
+#endif
 	return ScreenRect;
 }
 
