Adding support for windows clipboard contents update events for to be using for engine tooling.

diff --git a/Engine/Source/Runtime/ApplicationCore/Private/Windows/WindowsApplication.cpp b/Engine/Source/Runtime/ApplicationCore/Private/Windows/WindowsApplication.cpp
--- a/Engine/Source/Runtime/ApplicationCore/Private/Windows/WindowsApplication.cpp
+++ b/Engine/Source/Runtime/ApplicationCore/Private/Windows/WindowsApplication.cpp
@@ -1068,6 +1068,12 @@ int32 FWindowsApplication::ProcessMessage( HWND hwnd, uint32 msg, WPARAM wParam,
 			}
 			break;
 
+		case WM_CLIPBOARDUPDATE:
+			{
+				OnClipboardContentChangedEvent.Broadcast();
+			}
+			break;
+		
 		case WM_KEYDOWN:
 		case WM_SYSKEYUP:
 		case WM_KEYUP:
diff --git a/Engine/Source/Runtime/ApplicationCore/Private/Windows/WindowsWindow.cpp b/Engine/Source/Runtime/ApplicationCore/Private/Windows/WindowsWindow.cpp
--- a/Engine/Source/Runtime/ApplicationCore/Private/Windows/WindowsWindow.cpp
+++ b/Engine/Source/Runtime/ApplicationCore/Private/Windows/WindowsWindow.cpp
@@ -315,6 +315,9 @@ void FWindowsWindow::Initialize( FWindowsApplication* const Application, const T
 		// Tell OLE that we are opting into drag and drop.
 		// Only makes sense for regular windows (windows that last a while.)
 		RegisterDragDrop( HWnd, this );
+		
+		// Listing to clipboard change event
+		::AddClipboardFormatListener(HWnd);
 	}
 }
 
diff --git a/Engine/Source/Runtime/ApplicationCore/Public/GenericPlatform/GenericApplication.h b/Engine/Source/Runtime/ApplicationCore/Public/GenericPlatform/GenericApplication.h
--- a/Engine/Source/Runtime/ApplicationCore/Public/GenericPlatform/GenericApplication.h
+++ b/Engine/Source/Runtime/ApplicationCore/Public/GenericPlatform/GenericApplication.h
@@ -522,7 +522,9 @@ public:
 	
 	DECLARE_EVENT(FSlateApplication, FVirtualKeyboardHiddenEvent);
 	FVirtualKeyboardHiddenEvent& OnVirtualKeyboardHidden()  { return VirtualKeyboardHiddenEvent; }
-
+	
+	DECLARE_EVENT(FSlateApplication, FOnClipboardContentChanged);
+	FOnClipboardContentChanged& OnClipboardContentChanged()  { return OnClipboardContentChangedEvent; }
 	
 	/** Gets the horizontal alignment of the window title bar's title text. */
 	virtual EWindowTitleAlignment::Type GetWindowTitleAlignment() const
@@ -579,5 +581,6 @@ protected:
 	/** Delegate for virtual keyboard being shown */
 	FVirtualKeyboardHiddenEvent VirtualKeyboardHiddenEvent;
 	
-
+	/** Delegate for clipboard contents change */
+	FOnClipboardContentChanged OnClipboardContentChangedEvent;
 };
