Fix wrong printf param

#rnx

diff --git a/Engine/Source/Runtime/ApplicationCore/Private/IOS/IOSBackgroundURLSessionHandler.cpp b/Engine/Source/Runtime/ApplicationCore/Private/IOS/IOSBackgroundURLSessionHandler.cpp
--- a/Engine/Source/Runtime/ApplicationCore/Private/IOS/IOSBackgroundURLSessionHandler.cpp
+++ b/Engine/Source/Runtime/ApplicationCore/Private/IOS/IOSBackgroundURLSessionHandler.cpp
@@ -128,7 +128,7 @@ void FBackgroundURLSessionHandler::CreateBackgroundSessionWorkingDirectory()
         {
             NSString* PathAsNSString = DirectoryPath.GetNSString();
             //May not yet have UE_Log
-            ensureAlwaysMsgf(false, TEXT("Warning: Error in CreateBackgroundSessionWorkingDirectory: %s"), [PathAsNSString UTF8String]);
+            ensureAlwaysMsgf(false, TEXT("Warning: Error in CreateBackgroundSessionWorkingDirectory"));
         }
 	}
 }
