Fix multiple controller not supported on Mac

diff --git a/Engine/Source/Runtime/ApplicationCore/Private/Apple/AppleControllerInterface.cpp b/Engine/Source/Runtime/ApplicationCore/Private/Apple/AppleControllerInterface.cpp
--- a/Engine/Source/Runtime/ApplicationCore/Private/Apple/AppleControllerInterface.cpp
+++ b/Engine/Source/Runtime/ApplicationCore/Private/Apple/AppleControllerInterface.cpp
@@ -78,7 +78,6 @@ void FAppleControllerInterface::Tick( float DeltaTime )
 			case EAppleControllerEventType::Connect:
 			{
 				HandleConnection(Event.Controller);
-				SetCurrentController(Event.Controller);
 				break;
 			}
 			case EAppleControllerEventType::Disconnect:
@@ -138,16 +137,28 @@ void FAppleControllerInterface::SetControllerType(uint32 ControllerIndex)
 void FAppleControllerInterface::SetCurrentController(GCController* Controller)
 {
     int32 ControllerIndex = 0;
+    PlayerIndex PreviousIndex = PlayerIndex::PlayerUnset;
 
     for (ControllerIndex = 0; ControllerIndex < UE_ARRAY_COUNT(Controllers); ControllerIndex++)
     {
         if (Controllers[ControllerIndex].Controller == Controller)
         {
+            if (Controllers[ControllerIndex].PlayerIndex == PlayerIndex::PlayerOne)
+            {
+                // Already set as CurrentController
+                return;
+            }
+            PreviousIndex = Controllers[ControllerIndex].PlayerIndex;
             Controllers[ControllerIndex].PlayerIndex = PlayerIndex::PlayerOne;
         }
-        else if (Controllers[ControllerIndex].PlayerIndex == PlayerIndex::PlayerOne)
+    }
+    
+    for (ControllerIndex = 0; ControllerIndex < UE_ARRAY_COUNT(Controllers); ControllerIndex++)
+    {
+        if (Controllers[ControllerIndex].PlayerIndex == PlayerIndex::PlayerOne && Controllers[ControllerIndex].Controller != Controller)
         {
-            Controllers[ControllerIndex].PlayerIndex = PlayerIndex::PlayerUnset;
+            // The old PlayerOne, should swap place with the new PlayerOne
+            Controllers[ControllerIndex].PlayerIndex = PreviousIndex;
         }
     }
 }
@@ -187,6 +198,12 @@ PRAGMA_ENABLE_DEPRECATION_WARNINGS
         
         bFoundSlot = true;
         
+        IPlatformInputDeviceMapper& DeviceMapper = IPlatformInputDeviceMapper::Get();
+        FPlatformUserId UserId = FGenericPlatformMisc::GetPlatformUserForUserIndex(Controllers[ControllerIndex].PlayerIndex);
+        FInputDeviceId DeviceId = INPUTDEVICEID_NONE;
+        DeviceMapper.RemapControllerIdToPlatformUserAndDevice(Controllers[ControllerIndex].PlayerIndex, OUT UserId, OUT DeviceId);
+        DeviceMapper.Internal_MapInputDeviceToUser(DeviceId, UserId, EInputDeviceConnectionState::Connected);
+        
         UE_LOG(LogAppleController, Log, TEXT("New %s controller inserted, assigned to playerIndex %d"),
                Controllers[ControllerIndex].ControllerType == ControllerType::ExtendedGamepad ||
                Controllers[ControllerIndex].ControllerType == ControllerType::XboxGamepad ||
@@ -213,6 +230,12 @@ void FAppleControllerInterface::HandleDisconnect(GCController* Controller)
 			// Player index of unset(-1) would indicate that it has become unset even though it is now trying to disconnect
 			// This can occur on iOS when bGameSupportsMultipleActiveControllers is false
 			UE_LOG(LogAppleController, Log, TEXT("Controller for playerIndex %d, controller Index %d removed"), UserController.PlayerIndex, ControllerIndex);
+            
+            IPlatformInputDeviceMapper& DeviceMapper = IPlatformInputDeviceMapper::Get();
+            FPlatformUserId UserId = FGenericPlatformMisc::GetPlatformUserForUserIndex(UserController.PlayerIndex);
+            FInputDeviceId DeviceId = INPUTDEVICEID_NONE;
+            DeviceMapper.RemapControllerIdToPlatformUserAndDevice(UserController.PlayerIndex, OUT UserId, OUT DeviceId);
+            DeviceMapper.Internal_MapInputDeviceToUser(DeviceId, UserId, EInputDeviceConnectionState::Disconnected);
 			
 			[UserController.Controller release];
 			[UserController.PreviousExtendedGamepad release];
@@ -242,8 +265,7 @@ void FAppleControllerInterface::SendControllerEvents()
 		
 		IPlatformInputDeviceMapper& DeviceMapper = IPlatformInputDeviceMapper::Get();
 		FPlatformUserId UserId = FGenericPlatformMisc::GetPlatformUserForUserIndex(Controller.PlayerIndex);
-		FInputDeviceId DeviceId = INPUTDEVICEID_NONE;
-		DeviceMapper.RemapControllerIdToPlatformUserAndDevice(Controller.PlayerIndex, OUT UserId, OUT DeviceId);
+		FInputDeviceId DeviceId = DeviceMapper.GetPrimaryInputDeviceForUser(UserId);
 		
         GCExtendedGamepad* ExtendedGamepad = [ControllerImpl capture].extendedGamepad;
 		GCMotion* Motion = ControllerImpl.motion;
@@ -356,8 +378,7 @@ void FAppleControllerInterface::HandleInputInternal(const FGamepadKeyNames::Type
     
 	IPlatformInputDeviceMapper& DeviceMapper = IPlatformInputDeviceMapper::Get();
 	FPlatformUserId UserId = FGenericPlatformMisc::GetPlatformUserForUserIndex(Controllers[ControllerIndex].PlayerIndex);
-	FInputDeviceId DeviceId = INPUTDEVICEID_NONE;
-	DeviceMapper.RemapControllerIdToPlatformUserAndDevice(Controllers[ControllerIndex].PlayerIndex, OUT UserId, OUT DeviceId);
+    FInputDeviceId DeviceId = DeviceMapper.GetPrimaryInputDeviceForUser(UserId);
 
     if (bWasPressed != bIsPressed)
     {
@@ -507,8 +528,7 @@ void FAppleControllerInterface::HandleAnalogGamepad(const FGamepadKeyNames::Type
     
     IPlatformInputDeviceMapper& DeviceMapper = IPlatformInputDeviceMapper::Get();
 	FPlatformUserId UserId = FGenericPlatformMisc::GetPlatformUserForUserIndex(Controllers[ControllerIndex].PlayerIndex);
-	FInputDeviceId DeviceId = INPUTDEVICEID_NONE;
-	DeviceMapper.RemapControllerIdToPlatformUserAndDevice(Controllers[ControllerIndex].PlayerIndex, OUT UserId, OUT DeviceId);
+    FInputDeviceId DeviceId = DeviceMapper.GetPrimaryInputDeviceForUser(UserId);
     
     // Send controller events any time we are passed the given input threshold similarly to PC/Console (see: XInputInterface.cpp)
     const float RepeatDeadzone = 0.24f;
diff --git a/Engine/Source/Runtime/ApplicationCore/Private/Mac/HIDInputInterface.cpp b/Engine/Source/Runtime/ApplicationCore/Private/Mac/HIDInputInterface.cpp
--- a/Engine/Source/Runtime/ApplicationCore/Private/Mac/HIDInputInterface.cpp
+++ b/Engine/Source/Runtime/ApplicationCore/Private/Mac/HIDInputInterface.cpp
@@ -464,8 +464,7 @@ void HIDInputInterface::SendControllerEvents()
 
 		IPlatformInputDeviceMapper& DeviceMapper = IPlatformInputDeviceMapper::Get();
 		FPlatformUserId UserId = FGenericPlatformMisc::GetPlatformUserForUserIndex(ControllerState.ControllerId);
-		FInputDeviceId DeviceId = INPUTDEVICEID_NONE;
-		DeviceMapper.RemapControllerIdToPlatformUserAndDevice(ControllerState.ControllerId, OUT UserId, OUT DeviceId);
+        FInputDeviceId DeviceId = DeviceMapper.GetPrimaryInputDeviceForUser(UserId);
 
 		if( ControllerState.Device.DeviceRef )
 		{
