Mac: Fix NSCursor hotSpot logic.

* Need to use the actual size for the user coordinates instead of the image dimensions. Additionally, the hotspot: expects values from 0,0 -> w-1,h1.
* When a blueprint node specifies a value where (hotSpot.x/y * Image.pixelW/H) > Image.size.W/H, where hotSpot.x/y is 0...1.0, then no cursor would be displayed.

diff --git a/Engine/Source/Runtime/ApplicationCore/Private/Mac/MacCursor.cpp b/Engine/Source/Runtime/ApplicationCore/Private/Mac/MacCursor.cpp
--- a/Engine/Source/Runtime/ApplicationCore/Private/Mac/MacCursor.cpp
+++ b/Engine/Source/Runtime/ApplicationCore/Private/Mac/MacCursor.cpp
@@ -247,8 +247,8 @@ void* FMacCursor::CreateCursorFromFile(const FString& InPathToCursorWithoutExten
 
 	NSImageRep* CursorImageRep = [[CursorImage representations] objectAtIndex:0];
 
-	const int32 PixelHotspotX = FMath::RoundToInt(InHotSpot.X * CursorImageRep.pixelsWide);
-	const int32 PixelHotspotY = FMath::RoundToInt(InHotSpot.Y * CursorImageRep.pixelsHigh);
+	const int32 PixelHotspotX = FMath::RoundToInt(InHotSpot.X * (CursorImageRep.size.width-1));
+	const int32 PixelHotspotY = FMath::RoundToInt(InHotSpot.Y * (CursorImageRep.size.height-1));
 
 	NSCursor* CursorHandle = [[NSCursor alloc] initWithImage:CursorImage hotSpot : NSMakePoint(PixelHotspotX, PixelHotspotY)];
 	[CursorImage release];
