Remove unnecessary include

#rnx

- We are only storing a pointer to IInputDevice and never using it, so the include is not needed.

diff --git a/Engine/Source/Runtime/ApplicationCore/Private/GenericPlatform/GenericApplication.cpp b/Engine/Source/Runtime/ApplicationCore/Private/GenericPlatform/GenericApplication.cpp
--- a/Engine/Source/Runtime/ApplicationCore/Private/GenericPlatform/GenericApplication.cpp
+++ b/Engine/Source/Runtime/ApplicationCore/Private/GenericPlatform/GenericApplication.cpp
@@ -5,7 +5,6 @@
 #include "HAL/IConsoleManager.h"
 #include "Misc/CoreDelegates.h"
 #include "Misc/CommandLine.h"
-#include "IInputDevice.h"
 
 const FGamepadKeyNames::Type FGamepadKeyNames::Invalid(NAME_None);
 
