Fix iOS ignoring iOS Editor Settings screen rotation constraints.

diff --git a/Engine/Source/Runtime/ApplicationCore/Private/IOS/IOSAppDelegate.cpp b/Engine/Source/Runtime/ApplicationCore/Private/IOS/IOSAppDelegate.cpp
--- a/Engine/Source/Runtime/ApplicationCore/Private/IOS/IOSAppDelegate.cpp
+++ b/Engine/Source/Runtime/ApplicationCore/Private/IOS/IOSAppDelegate.cpp
@@ -1120,6 +1120,63 @@ static FAutoConsoleVariableRef CVarGEnableThermalsReport(
 	return YES;
 }
 
+#if !PLATFORM_TVOS
+- (UIInterfaceOrientationMask)application:(UIApplication *)application supportedInterfaceOrientationsForWindow:(UIWindow*)window
+{
+	bool bSupportsPortrait;
+	bool bSupportsPortraitUpsideDown;
+	bool bSupportsLandscapeLeft;
+	bool bSupportsLandscapeRight;
+	
+	// This is called during app startup and IOSRuntimeSettings may not have been loaded yet
+	bool hasValue = GConfig->GetBool(TEXT("/Script/IOSRuntimeSettings.IOSRuntimeSettings"), TEXT("bSupportsPortraitOrientation"), bSupportsPortrait, GEngineIni);
+	
+	NSArray<NSString*> *SupportedOrientations = [[NSBundle mainBundle] objectForInfoDictionaryKey:@"UISupportedInterfaceOrientations"];
+	if (!hasValue && SupportedOrientations != NULL)
+	{
+		// Loop through the Info.plist UISupportedInterfaceOrientations array values looking for "Portrait", "Left" and "Right"
+		NSPredicate *predicate = [NSPredicate predicateWithFormat:@"SELF == %@", @"UIInterfaceOrientationPortrait"];
+		bSupportsPortrait = ([SupportedOrientations filteredArrayUsingPredicate:predicate].count > 0);
+		
+		NSPredicate *predicateDown = [NSPredicate predicateWithFormat:@"SELF == %@", @"UIInterfaceOrientationPortraitUpsideDown"];
+		bSupportsPortraitUpsideDown = ([SupportedOrientations filteredArrayUsingPredicate:predicateDown].count > 0);
+		
+		NSPredicate *PredicateLeft = [NSPredicate predicateWithFormat:@"SELF == %@", @"UIInterfaceOrientationLandscapeLeft"];
+		bSupportsLandscapeLeft = ([SupportedOrientations filteredArrayUsingPredicate:PredicateLeft].count > 0);
+		
+		NSPredicate *PredicateRight = [NSPredicate predicateWithFormat:@"SELF == %@", @"UIInterfaceOrientationLandscapeRight"];
+		bSupportsLandscapeRight = ([SupportedOrientations filteredArrayUsingPredicate:PredicateRight].count > 0);
+	}
+	else
+	{
+		GConfig->GetBool(TEXT("/Script/IOSRuntimeSettings.IOSRuntimeSettings"), TEXT("bSupportsUpsideDownOrientation"), bSupportsPortraitUpsideDown, GEngineIni);
+		GConfig->GetBool(TEXT("/Script/IOSRuntimeSettings.IOSRuntimeSettings"), TEXT("bSupportsLandscapeLeftOrientation"), bSupportsLandscapeLeft, GEngineIni);
+		GConfig->GetBool(TEXT("/Script/IOSRuntimeSettings.IOSRuntimeSettings"), TEXT("bSupportsLandscapeRightOrientation"), bSupportsLandscapeRight, GEngineIni);
+	}
+	
+	UIInterfaceOrientationMask Mask = 0;
+	if (bSupportsPortrait)
+	{
+		Mask |= UIInterfaceOrientationMaskPortrait;
+	}
+	if (bSupportsPortraitUpsideDown)
+	{
+		Mask |= UIInterfaceOrientationMaskPortraitUpsideDown;
+	}
+	if (bSupportsLandscapeLeft)
+	{
+		Mask |= UIInterfaceOrientationMaskLandscapeLeft;
+	}
+	if (bSupportsLandscapeRight)
+	{
+		Mask |= UIInterfaceOrientationMaskLandscapeRight;
+	}
+	
+	// If no orientation constraints are set, default to MaskAll
+	return (Mask==0)?UIInterfaceOrientationMaskAll:Mask;
+}
+#endif
+
 #if WITH_ACCESSIBILITY
 -(void)OnVoiceOverStatusChanged
 {
diff --git a/Engine/Source/Runtime/ApplicationCore/Private/IOS/IOSView.cpp b/Engine/Source/Runtime/ApplicationCore/Private/IOS/IOSView.cpp
--- a/Engine/Source/Runtime/ApplicationCore/Private/IOS/IOSView.cpp
+++ b/Engine/Source/Runtime/ApplicationCore/Private/IOS/IOSView.cpp
@@ -1038,21 +1038,28 @@ self.accessibilityElements = @[Window.accessibilityContainer];
 	[super viewDidUnload];
 }
 
+#if !PLATFORM_TVOS
 /**
  * Tell the OS about the default supported orientations
  */
 - (NSUInteger)supportedInterfaceOrientations
 {
-	const FIOSView *View = [[IOSAppDelegate GetDelegate] IOSView];
+	const IOSAppDelegate *AppDelegate = [IOSAppDelegate GetDelegate];
+	const FIOSView *View = [AppDelegate IOSView];
 	if (View != nil)
 	{
-		return View->SupportedInterfaceOrientations;
-	}
-	else
-	{
-		return UIInterfaceOrientationMaskAll;
+		// if a Blueprint has changed the default rotation constraints, honour that change
+		if (View->SupportedInterfaceOrientations != UIInterfaceOrientationMaskAll)
+		{
+			return View->SupportedInterfaceOrientations;
+		}
 	}
+
+	// View either not yet created or Blueprint is not overriding the default, so use what the Window has set
+	UIApplication *app = [UIApplication sharedApplication];
+	return [AppDelegate application:app supportedInterfaceOrientationsForWindow:[AppDelegate window]];
 }
+#endif
 
 /**
  * Tell the OS that our view controller can auto-rotate between supported orientations
