Fix usage of iOS "bSupportsBackgroundAudio" setting which was ignoring how the value was actually set in the ini file

#rnx

diff --git a/Engine/Source/Runtime/Apple/AudioMixerAudioUnit/Private/AudioMixerPlatformAudioUnit.cpp b/Engine/Source/Runtime/Apple/AudioMixerAudioUnit/Private/AudioMixerPlatformAudioUnit.cpp
--- a/Engine/Source/Runtime/Apple/AudioMixerAudioUnit/Private/AudioMixerPlatformAudioUnit.cpp
+++ b/Engine/Source/Runtime/Apple/AudioMixerAudioUnit/Private/AudioMixerPlatformAudioUnit.cpp
@@ -82,7 +82,8 @@ namespace Audio
 			return false;
 		}
 		
-        bSupportsBackgroundAudio = GConfig->GetBool(TEXT("/Script/IOSRuntimeSettings.IOSRuntimeSettings"), TEXT("bSupportsBackgroundAudio"), bSupportsBackgroundAudio, GEngineIni);
+		bSupportsBackgroundAudio = false;
+		GConfig->GetBool(TEXT("/Script/IOSRuntimeSettings.IOSRuntimeSettings"), TEXT("bSupportsBackgroundAudio"), bSupportsBackgroundAudio, GEngineIni);
 
 		OSStatus Status;
 		GraphSampleRate = (double) InternalPlatformSettings.SampleRate;
diff --git a/Engine/Source/Runtime/ApplicationCore/Private/IOS/IOSAppDelegate.cpp b/Engine/Source/Runtime/ApplicationCore/Private/IOS/IOSAppDelegate.cpp
--- a/Engine/Source/Runtime/ApplicationCore/Private/IOS/IOSAppDelegate.cpp
+++ b/Engine/Source/Runtime/ApplicationCore/Private/IOS/IOSAppDelegate.cpp
@@ -612,20 +612,21 @@ static IOSAppDelegate* CachedDelegate = nil;
 	
 	self.bAudioActive = bActive;
 	
-    // get the category and settings to use
-        NSString* Category = AVAudioSessionCategoryAmbient;
-        if([self IsFeatureActive:EAudioFeature::DoNotMixWithOthers])
-        {
-            Category = AVAudioSessionCategorySoloAmbient;
-        }
-    #if !PLATFORM_TVOS
-        bool bSupportsBackgroundAudio = GConfig->GetBool(TEXT("/Script/IOSRuntimeSettings.IOSRuntimeSettings"), TEXT("bSupportsBackgroundAudio"), bSupportsBackgroundAudio, GEngineIni);
-        if (bSupportsBackgroundAudio)
-        {
-            Category = AVAudioSessionCategoryPlayback;
-        }
-    #endif
-    
+	// get the category and settings to use
+	NSString* Category = AVAudioSessionCategoryAmbient;
+	if([self IsFeatureActive:EAudioFeature::DoNotMixWithOthers])
+	{
+		Category = AVAudioSessionCategorySoloAmbient;
+	}
+#if !PLATFORM_TVOS
+	bool bSupportsBackgroundAudio = false;
+	GConfig->GetBool(TEXT("/Script/IOSRuntimeSettings.IOSRuntimeSettings"), TEXT("bSupportsBackgroundAudio"), bSupportsBackgroundAudio, GEngineIni);
+	if (bSupportsBackgroundAudio)
+	{
+		Category = AVAudioSessionCategoryPlayback;
+	}
+#endif
+	
 	NSString* Mode = AVAudioSessionModeDefault;
 	AVAudioSessionCategoryOptions Options = 0;
 	if (self.bAudioActive || [self IsBackgroundAudioPlaying] || [self IsFeatureActive:EAudioFeature::BackgroundAudio])
@@ -1423,10 +1424,11 @@ FCriticalSection RenderSuspend;
 		}
 		UE_LOG(LogTemp, Display, TEXT("Done with entering background tasks time."));
     }
-    bool bSupportsBackgroundAudio = GConfig->GetBool(TEXT("/Script/IOSRuntimeSettings.IOSRuntimeSettings"), TEXT("bSupportsBackgroundAudio"), bSupportsBackgroundAudio, GEngineIni);
     
 // fix for freeze on tvOS, moving to applicationDidEnterBackground. Not making the changes for iOS platforms as the bug does not happen and could bring some side effets.
 #if !PLATFORM_TVOS
+	bool bSupportsBackgroundAudio = false;
+	GConfig->GetBool(TEXT("/Script/IOSRuntimeSettings.IOSRuntimeSettings"), TEXT("bSupportsBackgroundAudio"), bSupportsBackgroundAudio, GEngineIni);
     if (!bSupportsBackgroundAudio)
     {
         [self ToggleSuspend:true];
