Fixed FNullWindow forcing windows to be zero-sized before it was properly initialized

SWindow may call AdjustCachedSize before calling Initialize during its initial creation. This would previously result in the window being forced to zero size and a subsequent call to Initialize would never reshape the Slate window.
This is now fixed by caching whether the window has been initialized or not. This is more in line with the other platform windows that generally do a noop if they have not created a window handle yet.

diff --git a/Engine/Source/Runtime/ApplicationCore/Private/Null/NullWindow.cpp b/Engine/Source/Runtime/ApplicationCore/Private/Null/NullWindow.cpp
--- a/Engine/Source/Runtime/ApplicationCore/Private/Null/NullWindow.cpp
+++ b/Engine/Source/Runtime/ApplicationCore/Private/Null/NullWindow.cpp
@@ -17,6 +17,7 @@ FNullWindow::FNullWindow()
 	, ScreenPosition(FIntPoint::ZeroValue)
 	, SizeInScreen(FIntPoint::ZeroValue)
 	, bIsVisible(false)
+	, bIsInitialized(false)
 {
 }
 
@@ -27,6 +28,13 @@ FNullWindow::~FNullWindow()
 void FNullWindow::Initialize(class FNullApplication* const Application, const TSharedRef<FGenericWindowDefinition>& InDefinition, const TSharedPtr<FNullWindow>& InParent, const bool bShowImmediately)
 {
 	Definition = InDefinition;
+
+	ScreenPosition.X = static_cast<int32>(InDefinition->XDesiredPositionOnScreen);
+	ScreenPosition.Y = static_cast<int32>(InDefinition->YDesiredPositionOnScreen);
+	SizeInScreen.X = static_cast<int32>(InDefinition->WidthDesiredOnScreen);
+	SizeInScreen.Y = static_cast<int32>(InDefinition->HeightDesiredOnScreen);
+
+	bIsInitialized = true;
 }
 
 void FNullWindow::ReshapeWindow(int32 X, int32 Y, int32 Width, int32 Height)
@@ -212,8 +220,11 @@ bool FNullWindow::IsDefinitionValid() const
 
 void FNullWindow::AdjustCachedSize(FVector2D& Size) const
 {
-	Size.X = ScreenPosition.X;
-	Size.Y = ScreenPosition.Y;
+	if (bIsInitialized)
+	{
+		Size.X = SizeInScreen.X;
+		Size.Y = SizeInScreen.Y;
+	}
 }
 
 bool FNullWindow::IsManualManageDPIChanges() const
diff --git a/Engine/Source/Runtime/ApplicationCore/Public/Null/NullWindow.h b/Engine/Source/Runtime/ApplicationCore/Public/Null/NullWindow.h
--- a/Engine/Source/Runtime/ApplicationCore/Public/Null/NullWindow.h
+++ b/Engine/Source/Runtime/ApplicationCore/Public/Null/NullWindow.h
@@ -183,4 +183,6 @@ private:
 
 	/** Manually store window visibility as OS doesn't do that for us in nullplatform */
 	bool bIsVisible;
+
+	bool bIsInitialized;
 };
