Reduce repeating monitor DPI log event spew to only print on the first occurrence

diff --git a/Engine/Source/Runtime/ApplicationCore/Private/Linux/LinuxPlatformApplicationMisc.cpp b/Engine/Source/Runtime/ApplicationCore/Private/Linux/LinuxPlatformApplicationMisc.cpp
--- a/Engine/Source/Runtime/ApplicationCore/Private/Linux/LinuxPlatformApplicationMisc.cpp
+++ b/Engine/Source/Runtime/ApplicationCore/Private/Linux/LinuxPlatformApplicationMisc.cpp
@@ -495,6 +495,8 @@ namespace LinuxPlatformApplicationMisc
 
 float FLinuxPlatformApplicationMisc::GetDPIScaleFactorAtPoint(float X, float Y)
 {
+	float Scale = 1.0f;
+
 	if ((GIsEditor || IS_PROGRAM) && IsHighDPIAwarenessEnabled())
 	{
 		FDisplayMetrics DisplayMetrics;
@@ -502,36 +504,43 @@ float FLinuxPlatformApplicationMisc::GetDPIScaleFactorAtPoint(float X, float Y)
 		// find the monitor
 		int32 XInt = static_cast<int32>(X);
 		int32 YInt = static_cast<int32>(Y);
+		static bool bHaveShownDPIWarnings = false;
+		int NumWarningsShown = 0;
 		for(int Idx = 0, NumMonitors = DisplayMetrics.MonitorInfo.Num(); Idx < NumMonitors; ++Idx)
 		{
 			const FMonitorInfo & MonitorInfo = DisplayMetrics.MonitorInfo[Idx];
-
-			if (MonitorInfo.DisplayRect.Left <= XInt && MonitorInfo.DisplayRect.Right > XInt &&
-				MonitorInfo.DisplayRect.Top <= YInt && MonitorInfo.DisplayRect.Bottom > YInt)
+			float HorzDPI = 1.0f, VertDPI = 1.0f;
+			if (SDL_GetDisplayDPI(Idx, nullptr, &HorzDPI, &VertDPI) == 0)
 			{
-				float HorzDPI = 1.0f, VertDPI = 1.0f;
-				if (SDL_GetDisplayDPI(Idx, nullptr, &HorzDPI, &VertDPI) == 0)
+				if (MonitorInfo.DisplayRect.Left <= XInt && MonitorInfo.DisplayRect.Right > XInt &&
+					MonitorInfo.DisplayRect.Top <= YInt && MonitorInfo.DisplayRect.Bottom > YInt)
 				{
-					float Scale = LinuxPlatformApplicationMisc::QuantizeScale((HorzDPI + VertDPI) / 192.0f);	// average between two scales (divided by 96.0f)
+					Scale = LinuxPlatformApplicationMisc::QuantizeScale((HorzDPI + VertDPI) / 192.0f);	// average between two scales (divided by 96.0f)
 					UE_LOG(LogLinux, Verbose, TEXT("Scale at X=%f, Y=%f: %f (monitor=#%d, HDPI=%f (horz scale: %f), VDPI=%f (vert scale: %f))"), X, Y, Scale, Idx, HorzDPI, HorzDPI / 96.0f, VertDPI, VertDPI / 96.0f);
 
 					// There is a bug in the X11 RandR library in some versions of Linux that can return a bonkers scale.  Even on an 8K display,
 					// it's unlikely that we'd need more than a 4x scale, so capping there for now
 					if(Scale > 4.0f)
 						Scale = 1.0f;
-
-					return Scale;
 				}
-				else
+			}
+			else
+			{
+				if (!bHaveShownDPIWarnings)
 				{
-					// this can also happen for headless, so don't use Warning here
-					UE_LOG(LogLinux, Log, TEXT("Could not get DPI information for monitor #%d, assuming 1.0f"), Idx);
-					break;	// should fall-through to 1.0f
+					UE_LOG(LogLinux, Log, TEXT("Could not get DPI information for monitor #%d, assuming it to have 1.0f scale"), Idx);
+					++NumWarningsShown;
 				}
 			}
 		}
+
+		if(NumWarningsShown > 0)
+		{
+			bHaveShownDPIWarnings = true;
+		}
 	}
-	return 1.0f;
+
+	return Scale;
 }
 
 void FLinuxPlatformApplicationMisc::ClipboardCopy(const TCHAR* Str)
