Fix LiveLinkHub linux build

diff --git a/Engine/Source/Runtime/ApplicationCore/ApplicationCore.Build.cs b/Engine/Source/Runtime/ApplicationCore/ApplicationCore.Build.cs
--- a/Engine/Source/Runtime/ApplicationCore/ApplicationCore.Build.cs
+++ b/Engine/Source/Runtime/ApplicationCore/ApplicationCore.Build.cs
@@ -64,7 +64,7 @@ public class ApplicationCore : ModuleRules
 			);
 
 			// We need FreeType2 and GL for the Splash, but only in the Editor
-			if (Target.Type == TargetType.Editor)
+			if (Target.bCompileAgainstEditor)
 			{
 				AddEngineThirdPartyPrivateStaticDependencies(Target, "FreeType2");
 				AddEngineThirdPartyPrivateStaticDependencies(Target, "OpenGL");
