Make SafeZonePaddingTop and SafeZonePaddingRight command line argument values get applied to the correct values of the TitleSafePaddingSize vector.

#rnx

diff --git a/Engine/Source/Runtime/ApplicationCore/Private/GenericPlatform/GenericApplication.cpp b/Engine/Source/Runtime/ApplicationCore/Private/GenericPlatform/GenericApplication.cpp
--- a/Engine/Source/Runtime/ApplicationCore/Private/GenericPlatform/GenericApplication.cpp
+++ b/Engine/Source/Runtime/ApplicationCore/Private/GenericPlatform/GenericApplication.cpp
@@ -144,8 +144,8 @@ void FDisplayMetrics::ApplyDefaultSafeZones()
 	TitleSafePaddingSize = FVector4(0.0f, 0.0f, 0.0f, 0.0f);
 	bool bSetByCommandLine;
 	bSetByCommandLine = FParse::Value(FCommandLine::Get(), TEXT("SafeZonePaddingLeft="),   TitleSafePaddingSize.X);
-	bSetByCommandLine = FParse::Value(FCommandLine::Get(), TEXT("SafeZonePaddingRight="),  TitleSafePaddingSize.Y) || bSetByCommandLine;
-	bSetByCommandLine = FParse::Value(FCommandLine::Get(), TEXT("SafeZonePaddingTop="),    TitleSafePaddingSize.Z) || bSetByCommandLine;
+	bSetByCommandLine = FParse::Value(FCommandLine::Get(), TEXT("SafeZonePaddingTop="),  TitleSafePaddingSize.Y) || bSetByCommandLine;
+	bSetByCommandLine = FParse::Value(FCommandLine::Get(), TEXT("SafeZonePaddingRight="),    TitleSafePaddingSize.Z) || bSetByCommandLine;
 	bSetByCommandLine = FParse::Value(FCommandLine::Get(), TEXT("SafeZonePaddingBottom="), TitleSafePaddingSize.W) || bSetByCommandLine;
 
 	if (!bSetByCommandLine)
