Re-added the text outline to the Windows splash screen text

This makes it easier to read the text on a non-solid background, and had previously been removed in the UE4 -> UE5 styling changes

#rnx

diff --git a/Engine/Source/Runtime/ApplicationCore/Private/Windows/WindowsPlatformSplash.cpp b/Engine/Source/Runtime/ApplicationCore/Private/Windows/WindowsPlatformSplash.cpp
--- a/Engine/Source/Runtime/ApplicationCore/Private/Windows/WindowsPlatformSplash.cpp
+++ b/Engine/Source/Runtime/ApplicationCore/Private/Windows/WindowsPlatformSplash.cpp
@@ -34,6 +34,8 @@ THIRD_PARTY_INCLUDES_END
 
 #pragma comment( lib, "windowscodecs.lib" )
 
+#define UE_WINDOWS_SPLASH_USE_TEXT_OUTLINE (1)
+
 /**
  * Splash screen functions and static globals
  */
@@ -96,9 +98,7 @@ LRESULT CALLBACK SplashScreenWindowProc(HWND hWnd, uint32 message, WPARAM wParam
 							}
 
 							// Alignment
-							{
-								SetTextAlign( hdc, TA_LEFT | TA_TOP | TA_NOUPDATECP );
-							}
+							SetTextAlign( hdc, TA_LEFT | TA_TOP | TA_NOUPDATECP );
 
 							SetBkColor( hdc, 0x00000000 );
 							SetBkMode( hdc, TRANSPARENT );
@@ -106,8 +106,9 @@ LRESULT CALLBACK SplashScreenWindowProc(HWND hWnd, uint32 message, WPARAM wParam
 							RECT ClientRect;
 							GetClientRect( hWnd, &ClientRect );
 
+#if UE_WINDOWS_SPLASH_USE_TEXT_OUTLINE
 							// Draw background text passes
-					/*		const int32 NumBGPasses = 1;
+							const int32 NumBGPasses = 8;
 							for( int32 CurBGPass = 0; CurBGPass < NumBGPasses; ++CurBGPass )
 							{
 								int32 BGXOffset, BGYOffset;
@@ -115,13 +116,13 @@ LRESULT CALLBACK SplashScreenWindowProc(HWND hWnd, uint32 message, WPARAM wParam
 								{
 									default:
 									case 0:	BGXOffset = -1; BGYOffset =  0; break;
-									case 2:	BGXOffset = -1; BGYOffset = -1; break;
-									case 3:	BGXOffset =  0; BGYOffset = -1; break;
-									case 4:	BGXOffset =  1; BGYOffset = -1; break;
-									case 5:	BGXOffset =  1; BGYOffset =  0; break;
-									case 6:	BGXOffset =  1; BGYOffset =  1; break;
-									case 7:	BGXOffset =  0; BGYOffset =  1; break;
-									case 8:	BGXOffset = -1; BGYOffset =  1; break;
+									case 1:	BGXOffset = -1; BGYOffset = -1; break;
+									case 2:	BGXOffset =  0; BGYOffset = -1; break;
+									case 3:	BGXOffset =  1; BGYOffset = -1; break;
+									case 4:	BGXOffset =  1; BGYOffset =  0; break;
+									case 5:	BGXOffset =  1; BGYOffset =  1; break;
+									case 6:	BGXOffset =  0; BGYOffset =  1; break;
+									case 7:	BGXOffset = -1; BGYOffset =  1; break;
 								}
 
 								SetTextColor( hdc, 0x00000000 );
@@ -131,7 +132,8 @@ LRESULT CALLBACK SplashScreenWindowProc(HWND hWnd, uint32 message, WPARAM wParam
 									TextRect.top + BGYOffset,
 									*SplashText.ToString(),
 									SplashText.ToString().Len() );
-							}*/
+							}
+#endif // UE_WINDOWS_SPLASH_USE_TEXT_OUTLINE
 							
 							// Draw foreground text pass
 							if( CurTypeIndex == SplashTextType::StartupProgress )
