Make splash screen movable

diff --git a/Engine/Source/Runtime/ApplicationCore/Private/Windows/WindowsPlatformSplash.cpp b/Engine/Source/Runtime/ApplicationCore/Private/Windows/WindowsPlatformSplash.cpp
--- a/Engine/Source/Runtime/ApplicationCore/Private/Windows/WindowsPlatformSplash.cpp
+++ b/Engine/Source/Runtime/ApplicationCore/Private/Windows/WindowsPlatformSplash.cpp
@@ -36,6 +36,10 @@ THIRD_PARTY_INCLUDES_END
 
 #define UE_WINDOWS_SPLASH_USE_TEXT_OUTLINE (1)
 
+#if !defined(UE_WINDOWS_SPLASH_ENABLE_DRAG)
+#define UE_WINDOWS_SPLASH_ENABLE_DRAG WITH_EDITOR
+#endif
+
 /**
  * Splash screen functions and static globals
  */
@@ -171,6 +175,19 @@ LRESULT CALLBACK SplashScreenWindowProc(HWND hWnd, uint32 message, WPARAM wParam
 			PostQuitMessage(0);
 			break;
 
+#if UE_WINDOWS_SPLASH_ENABLE_DRAG
+		case WM_NCHITTEST:
+			{
+				// Report client area as non-client area to allow dragging the splashscreen around.
+				LRESULT Result = DefWindowProc(hWnd, message, wParam, lParam);
+				if (Result == HTCLIENT)
+				{
+					Result = HTCAPTION;
+				}
+				return Result;
+			}
+#endif
+
 		default:
 			return DefWindowProc(hWnd, message, wParam, lParam);
 	}
