Fixed issue when using Preserve Aspect Ratio on an SWindow where the window would change aspect ratio the first time it was resized if it didn't have an OS border/titlebar

diff --git a/Engine/Source/Runtime/ApplicationCore/Private/Windows/WindowsApplication.cpp b/Engine/Source/Runtime/ApplicationCore/Private/Windows/WindowsApplication.cpp
--- a/Engine/Source/Runtime/ApplicationCore/Private/Windows/WindowsApplication.cpp
+++ b/Engine/Source/Runtime/ApplicationCore/Private/Windows/WindowsApplication.cpp
@@ -1435,15 +1435,19 @@ int32 FWindowsApplication::ProcessMessage( HWND hwnd, uint32 msg, WPARAM wParam,
 					WindowInfo.cbSize = sizeof(WindowInfo);
 					::GetWindowInfo(hwnd, &WindowInfo);
 
-					RECT TestRect;
-					TestRect.left = TestRect.right = TestRect.top = TestRect.bottom = 0;
-					AdjustWindowRectEx(&TestRect, WindowInfo.dwStyle, false, WindowInfo.dwExStyle);
-
 					RECT* Rect = (RECT*)lParam;
-					Rect->left -= TestRect.left;
-					Rect->right -= TestRect.right;
-					Rect->top -= TestRect.top;
-					Rect->bottom -= TestRect.bottom;
+					if (CurrentNativeEventWindow->GetDefinition().HasOSWindowBorder)
+					{
+						RECT TestRect;
+						TestRect.left = TestRect.right = TestRect.top = TestRect.bottom = 0;
+						AdjustWindowRectEx(&TestRect, WindowInfo.dwStyle, false, WindowInfo.dwExStyle);
+
+					
+						Rect->left -= TestRect.left;
+						Rect->right -= TestRect.right;
+						Rect->top -= TestRect.top;
+						Rect->bottom -= TestRect.bottom;
+					}
 
 					const float AspectRatio = CurrentNativeEventWindowPtr->GetAspectRatio();
 					int32 NewWidth = Rect->right - Rect->left;
@@ -1552,7 +1556,10 @@ int32 FWindowsApplication::ProcessMessage( HWND hwnd, uint32 msg, WPARAM wParam,
 						}
 					}
 
-					AdjustWindowRectEx(Rect, WindowInfo.dwStyle, false, WindowInfo.dwExStyle);
+					if (CurrentNativeEventWindow->GetDefinition().HasOSWindowBorder)
+					{
+						AdjustWindowRectEx(Rect, WindowInfo.dwStyle, false, WindowInfo.dwExStyle);
+					}
 
 					return TRUE;
 				}
