CIS fix.

diff --git a/Engine/Source/Runtime/ApplicationCore/Private/IOS/IOSAppDelegate.cpp b/Engine/Source/Runtime/ApplicationCore/Private/IOS/IOSAppDelegate.cpp
--- a/Engine/Source/Runtime/ApplicationCore/Private/IOS/IOSAppDelegate.cpp
+++ b/Engine/Source/Runtime/ApplicationCore/Private/IOS/IOSAppDelegate.cpp
@@ -1357,7 +1357,7 @@ static FAutoConsoleVariableRef CVarGEnableThermalsReport(
 	int32 URLTerminator = CommandLineParameters.Find( TEXT("://"), ESearchCase::CaseSensitive);
 	if ( URLTerminator > -1 )
 	{
-		CommandLineParameters.RightChopInline(URLTerminator + 3, false);
+		CommandLineParameters.RightChopInline(URLTerminator + 3, EAllowShrinking::No);
 	}
 
 	FIOSCommandLineHelper::InitCommandArgs(CommandLineParameters);
