[ApplicationCore]
* Fixed so new console autoscrolling gets automatically enabled when log is cleared

diff --git a/Engine/Source/Runtime/ApplicationCore/Private/Windows/WindowsConsoleOutputDevice2.cpp b/Engine/Source/Runtime/ApplicationCore/Private/Windows/WindowsConsoleOutputDevice2.cpp
--- a/Engine/Source/Runtime/ApplicationCore/Private/Windows/WindowsConsoleOutputDevice2.cpp
+++ b/Engine/Source/Runtime/ApplicationCore/Private/Windows/WindowsConsoleOutputDevice2.cpp
@@ -1366,6 +1366,7 @@ public:
 
 			case 103:
 				ClearLog();
+				bAutoScrollLog = true;
 				break;
 
 			case 104:
@@ -2172,6 +2173,7 @@ public:
 			else if (LOWORD(wParam) == ID_CLEARLOGBUTTON)
 			{
 				ClearLog();
+				bAutoScrollLog = true;
 			}
 			else if (LOWORD(wParam) == ID_ADDCHECKPOINTBUTTON)
 			{
