Add FinputDeviceScope's to the Windows Application. This ensures that we actually have an accurate representation of when you press a Keybaord or Mouse key.

diff --git a/Engine/Source/Runtime/ApplicationCore/Private/Windows/WindowsApplication.cpp b/Engine/Source/Runtime/ApplicationCore/Private/Windows/WindowsApplication.cpp
--- a/Engine/Source/Runtime/ApplicationCore/Private/Windows/WindowsApplication.cpp
+++ b/Engine/Source/Runtime/ApplicationCore/Private/Windows/WindowsApplication.cpp
@@ -96,6 +96,13 @@ const FIntPoint FWindowsApplication::MinimizedWindowPosition(-32000,-32000);
 
 FWindowsApplication* WindowsApplication = nullptr;
 
+namespace UE::WindowsInput
+{
+	static const FName InputClassName = TEXT("WindowsApplication");
+	static const FString KBMInputHardwareName = TEXT("KBM");
+	static const FString TouchInputHardwareName = TEXT("MobileTouch");
+};
+
 static bool ShouldSimulateRawInput()
 {
 	return ForceRawInputSimulation || (EnableRawInputSimulationOverRDP && FPlatformMisc::IsRemoteSession());
@@ -1933,6 +1940,8 @@ void FWindowsApplication::CheckForShiftUpEvents(const int32 KeyCode)
 	if (ModifierKeyState[ModifierKeyIndex] && ((::GetKeyState(KeyCode) & 0x8000) == 0) )
 	{
 		ModifierKeyState[ModifierKeyIndex] = false;
+		
+		FInputDeviceScope InputScope(nullptr, UE::WindowsInput::InputClassName, IPlatformInputDeviceMapper::Get().GetDefaultInputDevice().GetId(), UE::WindowsInput::KBMInputHardwareName);
 		MessageHandler->OnKeyUp( KeyCode, 0, false );
 	}
 }
@@ -1991,7 +2000,7 @@ int32 FWindowsApplication::ProcessDeferredMessage( const FDeferredWindowsMessage
 
 				// LPARAM bit 30 will be ZERO for new presses, or ONE if this is a repeat
 				const bool bIsRepeat = ( lParam & 0x40000000 ) != 0;
-
+				FInputDeviceScope InputScope(nullptr, UE::WindowsInput::InputClassName, IPlatformInputDeviceMapper::Get().GetDefaultInputDevice().GetId(), UE::WindowsInput::KBMInputHardwareName); 
 				MessageHandler->OnKeyChar( Character, bIsRepeat );
 
 				// Note: always return 0 to handle the message.  Win32 beeps if WM_CHAR is not handled...
@@ -2071,7 +2080,7 @@ int32 FWindowsApplication::ProcessDeferredMessage( const FDeferredWindowsMessage
 
 				// Get the character code from the virtual key pressed.  If 0, no translation from virtual key to character exists
 				uint32 CharCode = ::MapVirtualKey( Win32Key, MAPVK_VK_TO_CHAR );
-
+				FInputDeviceScope InputScope(nullptr, UE::WindowsInput::InputClassName, IPlatformInputDeviceMapper::Get().GetDefaultInputDevice().GetId(), UE::WindowsInput::KBMInputHardwareName); 
 				const bool Result = MessageHandler->OnKeyDown( ActualKey, CharCode, bIsRepeat );
 
 				// Always return 0 to handle the message or else windows will beep
@@ -2150,7 +2159,7 @@ int32 FWindowsApplication::ProcessDeferredMessage( const FDeferredWindowsMessage
 
 				// Key up events are never repeats
 				const bool bIsRepeat = false;
-
+				FInputDeviceScope InputScope(nullptr, UE::WindowsInput::InputClassName, IPlatformInputDeviceMapper::Get().GetDefaultInputDevice().GetId(), UE::WindowsInput::KBMInputHardwareName); 
 				const bool Result = MessageHandler->OnKeyUp( ActualKey, CharCode, bIsRepeat );
 
 				// Note that we allow system keys to pass through to DefWndProc here, so that core features
@@ -2238,6 +2247,7 @@ int32 FWindowsApplication::ProcessDeferredMessage( const FDeferredWindowsMessage
 					check(0);
 				}
 
+				FInputDeviceScope InputScope(nullptr, UE::WindowsInput::InputClassName, IPlatformInputDeviceMapper::Get().GetDefaultInputDevice().GetId(), UE::WindowsInput::KBMInputHardwareName); 
 				if (bMouseUp)
 				{
 					return MessageHandler->OnMouseUp( MouseButton, CursorPos ) ? 0 : 1;
@@ -2257,6 +2267,7 @@ int32 FWindowsApplication::ProcessDeferredMessage( const FDeferredWindowsMessage
 		// Mouse Movement
 		case WM_INPUT:
 			{
+				FInputDeviceScope InputScope(nullptr, UE::WindowsInput::InputClassName, IPlatformInputDeviceMapper::Get().GetDefaultInputDevice().GetId(), UE::WindowsInput::KBMInputHardwareName); 
 				if( DeferredMessage.RawInputFlags == MOUSE_MOVE_RELATIVE )
 				{
 					MessageHandler->OnRawMouseMove(DeferredMessage.X, DeferredMessage.Y);
@@ -2278,6 +2289,7 @@ int32 FWindowsApplication::ProcessDeferredMessage( const FDeferredWindowsMessage
 				BOOL Result = false;
 				if (!bUsingHighPrecisionMouseInput)
 				{
+					FInputDeviceScope InputScope(nullptr, UE::WindowsInput::InputClassName, IPlatformInputDeviceMapper::Get().GetDefaultInputDevice().GetId(), UE::WindowsInput::KBMInputHardwareName); 
 					Result = MessageHandler->OnMouseMove();
 				}
 
@@ -2296,6 +2308,7 @@ int32 FWindowsApplication::ProcessDeferredMessage( const FDeferredWindowsMessage
 
 				const FVector2D CursorPos(CursorPoint.x, CursorPoint.y);
 
+				FInputDeviceScope InputScope(nullptr, UE::WindowsInput::InputClassName, IPlatformInputDeviceMapper::Get().GetDefaultInputDevice().GetId(), UE::WindowsInput::KBMInputHardwareName); 
 				const BOOL Result = MessageHandler->OnMouseWheel( static_cast<float>( WheelDelta ) * SpinFactor, CursorPos );
 				return Result ? 0 : 1;
 			}
@@ -2329,6 +2342,8 @@ int32 FWindowsApplication::ProcessDeferredMessage( const FDeferredWindowsMessage
 					TUniquePtr<TOUCHINPUT[]> Inputs = MakeUnique<TOUCHINPUT[]>( InputCount );
 					if ( GetTouchInputInfo( (HTOUCHINPUT)lParam, InputCount, Inputs.Get(), sizeof(TOUCHINPUT) ) )
 					{
+						FInputDeviceScope InputScope(nullptr, UE::WindowsInput::InputClassName, IPlatformInputDeviceMapper::Get().GetDefaultInputDevice().GetId(), UE::WindowsInput::TouchInputHardwareName);
+
 						for ( uint32 i = 0; i < InputCount; i++ )
 						{
 							TOUCHINPUT Input = Inputs[i];
diff --git a/Engine/Source/Runtime/Engine/Private/GameFramework/InputDeviceSubsystem.cpp b/Engine/Source/Runtime/Engine/Private/GameFramework/InputDeviceSubsystem.cpp
--- a/Engine/Source/Runtime/Engine/Private/GameFramework/InputDeviceSubsystem.cpp
+++ b/Engine/Source/Runtime/Engine/Private/GameFramework/InputDeviceSubsystem.cpp
@@ -57,7 +57,12 @@ class FInputDeviceSubsystemProcessor : public IInputProcessor
 			// If there isn't a recent input device scope, then we can check if the key was from a keyboard and mouse
 			else if (!Key.IsGamepadKey())
 			{
-				SubSystem->SetMostRecentlyUsedHardwareDevice(InDeviceId, FHardwareDeviceIdentifier::DefaultKeyboardAndMouse);
+				// We only want to set our "fallback" of DefaultKeyboardAndMouse if there is no known existing data for the device id
+				const FHardwareDeviceIdentifier ExistingDevice = SubSystem->GetInputDeviceHardwareIdentifier(InDeviceId);
+				if (ExistingDevice == FHardwareDeviceIdentifier::Invalid)
+				{
+					SubSystem->SetMostRecentlyUsedHardwareDevice(InDeviceId, FHardwareDeviceIdentifier::DefaultKeyboardAndMouse);	
+				}
 			}
 		}
 	}
@@ -506,12 +511,21 @@ void UInputDeviceSubsystem::SetMostRecentlyUsedHardwareDevice(const FInputDevice
 	const FPlatformUserId OwningUserId = IPlatformInputDeviceMapper::Get().GetUserForInputDevice(InDeviceId);
 
 	// If this hardware is the same as what the platform user already has, then there is no need to fire this event
-	if (const FHardwareDeviceIdentifier* ExistingDevice = LatestUserDeviceIdentifiers.Find(OwningUserId))
+	if (FHardwareDeviceIdentifier* ExistingDevice = LatestUserDeviceIdentifiers.Find(OwningUserId))
 	{
+		// If they are exactly the same, do nothing
 		if (InHardwareId == *ExistingDevice)
 		{
 			return;
 		}
+		// If they are the same device but different input API's, just update the data but don't broadcast the delegate
+		// This can happen if there are multiple input interfaces sending messages for the same device (such as a keyboard)
+		// Which your gameplay code doesn't really care about switching which API the input came from
+		else if (InHardwareId.HardwareDeviceIdentifier == ExistingDevice->HardwareDeviceIdentifier)
+		{
+			*ExistingDevice = InHardwareId;
+			return;
+		}
 	}
 
 	// Keep track of each input device's latest hardware id
