Added verbosity filters to console

diff --git a/Engine/Source/Runtime/ApplicationCore/Private/Windows/WindowsConsoleOutputDevice2.cpp b/Engine/Source/Runtime/ApplicationCore/Private/Windows/WindowsConsoleOutputDevice2.cpp
--- a/Engine/Source/Runtime/ApplicationCore/Private/Windows/WindowsConsoleOutputDevice2.cpp
+++ b/Engine/Source/Runtime/ApplicationCore/Private/Windows/WindowsConsoleOutputDevice2.cpp
@@ -711,6 +711,13 @@ public:
 
 	int32 AddEntryToLogHwnd(LogEntry& E, int32 LogVirtualIndex)
 	{
+		for (const ELogVerbosity::Type Verbosity : IncludeVerbosity)
+			if (E.Verbosity != Verbosity)
+				return -1;
+		for (const ELogVerbosity::Type Verbosity : ExcludeVerbosity)
+			if (E.Verbosity == Verbosity)
+				return -1;
+
 		TStringBuilder<1024> OutString;
 		CreateLogEntryText(OutString, E, true);
 
@@ -1300,10 +1307,14 @@ public:
 			if (LogIndex >= 0)
 			{
 				SelectedCategory = Log[LogIndex].Category.ToString();
+				SelectedVerbosity = Log[LogIndex].Verbosity;
 
 				AddSeparator();
 				AddItem(*TStringBuilder<64>().Append(TEXT("Include category '")).Append(*SelectedCategory).Append(TEXT("' in filter")), 124);
 				AddItem(*TStringBuilder<64>().Append(TEXT("Exclude category '")).Append(*SelectedCategory).Append(TEXT("' in filter")), 125);
+
+				AddItem(*TStringBuilder<64>().Append(TEXT("Include verbosity '")).Append(ToString(SelectedVerbosity)).Append(TEXT("' in filter")), 126);
+				AddItem(*TStringBuilder<64>().Append(TEXT("Exclude verbosity '")).Append(ToString(SelectedVerbosity)).Append(TEXT("' in filter")), 127);
 			}
 
 			int CharIndex = (MousePos.x + (LogFontWidth / 2)) / LogFontWidth;
@@ -1387,6 +1398,8 @@ public:
 			case 121:
 				UpdateIncludeFilter(TEXT(""));
 				UpdateExcludeFilter(TEXT(""));
+				IncludeVerbosity.Reset();
+				ExcludeVerbosity.Reset();
 				RefreshLogHwnd();
 				break;
 
@@ -1415,6 +1428,16 @@ public:
 				RefreshLogHwnd();
 				break;
 			}
+			case 126:
+			case 127:
+			{
+				if (LOWORD(wParam) == 126)
+					IncludeVerbosity.AddUnique(SelectedVerbosity);
+				else
+					ExcludeVerbosity.AddUnique(SelectedVerbosity);
+				RefreshLogHwnd();
+				break;
+			}
 			case 131:
 				ConsoleShowDateTime = 1;
 				RefreshLogHwnd();
@@ -2451,6 +2474,8 @@ public:
 	TArray<NewLogEntry> TempLogEntries;
 	TArray<FString> IncludeFilter;
 	TArray<FString> ExcludeFilter;
+	TArray<ELogVerbosity::Type> IncludeVerbosity;
+	TArray<ELogVerbosity::Type> ExcludeVerbosity;
 	TRingBuffer<LogEntry> Log;
 	FDateTime StartDateTime;
 	double StartTime;
@@ -2503,6 +2528,7 @@ public:
 	int32 RightClickedItem = -1;
 	FString SelectedWord;
 	FString SelectedCategory;
+	ELogVerbosity::Type SelectedVerbosity;
 
 	struct Activity { HWND NameHwnd = 0; HWND StatusHwnd = 0; FString Name; FString Status; int32 Light = 0; int32 SortValue; bool bStatusDirty = false; bool bAlignLeft = false; };
 	struct ActivityModification { FString Name; FString Status; int32 Light = 0; int32 SortValue; bool bAlignLeft = false; bool bRemove = false; };
