Crash fix for null bitmap race in windows draw dispatch

diff --git a/Engine/Source/Runtime/ApplicationCore/Private/Windows/WindowsPlatformSplash.cpp b/Engine/Source/Runtime/ApplicationCore/Private/Windows/WindowsPlatformSplash.cpp
--- a/Engine/Source/Runtime/ApplicationCore/Private/Windows/WindowsPlatformSplash.cpp
+++ b/Engine/Source/Runtime/ApplicationCore/Private/Windows/WindowsPlatformSplash.cpp
@@ -71,6 +71,12 @@ LRESULT CALLBACK SplashScreenWindowProc(HWND hWnd, uint32 message, WPARAM wParam
 	{
 		case WM_PAINT:
 			{
+				// We can continue to receive messages for a period after we've torn down the bitmap resource
+				if (!GSplashScreenBitmap)
+				{
+					return 0;
+				}
+
 				hdc = BeginPaint(hWnd, &ps);
 
 				// Draw splash bitmap
