ios.UnifyMotionSpace added to let iOS motion input match other platforms

Set ios.UnifyMotionSpace to 1 to make iOS motion sensor axis match Unreal's. Can be set to 2 to do the same except swap x and y to make it right-handed, since right-handed is more common outside of Unreal. Set to 0 for old behaviour.

diff --git a/Engine/Source/Runtime/ApplicationCore/Private/IOS/IOSInputInterface.cpp b/Engine/Source/Runtime/ApplicationCore/Private/IOS/IOSInputInterface.cpp
--- a/Engine/Source/Runtime/ApplicationCore/Private/IOS/IOSInputInterface.cpp
+++ b/Engine/Source/Runtime/ApplicationCore/Private/IOS/IOSInputInterface.cpp
@@ -19,6 +19,7 @@ static TAutoConsoleVariable<float> CVarHapticsKickHeavy(TEXT("ios.VibrationHapti
 static TAutoConsoleVariable<float> CVarHapticsKickMedium(TEXT("ios.VibrationHapticsKickMediumValue"), 0.5f, TEXT("Vibation values higher than this will kick a haptics medium Impact"));
 static TAutoConsoleVariable<float> CVarHapticsKickLight(TEXT("ios.VibrationHapticsKickLightValue"), 0.3f, TEXT("Vibation values higher than this will kick a haptics light Impact"));
 static TAutoConsoleVariable<float> CVarHapticsRest(TEXT("ios.VibrationHapticsRestValue"), 0.2f, TEXT("Vibation values lower than this will allow haptics to Kick again when going over ios.VibrationHapticsKickValue"));
+static TAutoConsoleVariable<int32> CVarUnifyMotionSpace(TEXT("ios.UnifyMotionSpace"), 0, TEXT("If set to non-zero, acceleration, gravity, and rotation rate will all be in the same coordinate space. 0 for legacy behaviour. 1 will match Unreal's coordinate space (left-handed, z-up, etc). 2 will be right-handed by swapping x and y. Non-zero also forces rotation rate units to be radians/s and acceleration units to be g."));
 
 uint32 TranslateGCKeyCodeToASCII(GCKeyCode KeyCode)
 {
@@ -492,15 +493,79 @@ void FIOSInputInterface::SendControllerEvents()
 
         GetMovementData(Attitude, RotationRate, Gravity, Acceleration);
 
-        // Fix-up yaw to match directions
-        Attitude.Y = -Attitude.Y;
-        RotationRate.Y = -RotationRate.Y;
-
-        // munge the vectors based on the orientation
-        ModifyVectorByOrientation(Attitude, true);
-        ModifyVectorByOrientation(RotationRate, true);
-        ModifyVectorByOrientation(Gravity, false);
-        ModifyVectorByOrientation(Acceleration, false);
+		const int32 UnifyMotionSpace = CVarUnifyMotionSpace.GetValueOnGameThread();
+		if (UnifyMotionSpace == 0)
+		{
+			// Fix-up yaw to match directions
+			Attitude.Y = -Attitude.Y;
+			RotationRate.Y = -RotationRate.Y;
+
+			// munge the vectors based on the orientation
+			ModifyVectorByOrientation(Attitude, true);
+			ModifyVectorByOrientation(RotationRate, true);
+			ModifyVectorByOrientation(Gravity, false);
+			ModifyVectorByOrientation(Acceleration, false);
+		}
+		else
+		{
+			// Match Unreal coordinate system
+			auto ReorientPortrait = [](FVector InValue)
+				{
+					return FVector(-InValue.Z, InValue.X, InValue.Y);
+				};
+			auto ReorientPortraitUpsideDown = [](FVector InValue)
+				{
+					return FVector(-InValue.Z, -InValue.X, -InValue.Y);
+				};
+			auto ReorientLandscapeRight = [](FVector InValue)
+				{
+					return FVector(-InValue.Z, InValue.Y, -InValue.X);
+				};
+			auto ReorientLandscapeLeft = [](FVector InValue)
+				{
+					return FVector(-InValue.Z, -InValue.Y, InValue.X);
+				};
+
+			switch (FIOSApplication::CachedOrientation)
+			{
+			case UIInterfaceOrientationPortrait:
+				Attitude = ReorientPortrait(Attitude);
+				RotationRate = ReorientPortrait(RotationRate);
+				Gravity = ReorientPortrait(Gravity);
+				Acceleration = ReorientPortrait(Acceleration);
+				break;
+
+			case UIInterfaceOrientationPortraitUpsideDown:
+				Attitude = ReorientPortraitUpsideDown(Attitude);
+				RotationRate = ReorientPortraitUpsideDown(RotationRate);
+				Gravity = ReorientPortraitUpsideDown(Gravity);
+				Acceleration = ReorientPortraitUpsideDown(Acceleration);
+				break;
+
+			case UIInterfaceOrientationLandscapeRight:
+				Attitude = ReorientLandscapeRight(Attitude);
+				RotationRate = ReorientLandscapeRight(RotationRate);
+				Gravity = ReorientLandscapeRight(Gravity);
+				Acceleration = ReorientLandscapeRight(Acceleration);
+				break;
+
+			case UIInterfaceOrientationLandscapeLeft:
+				Attitude = ReorientLandscapeLeft(Attitude);
+				RotationRate = ReorientLandscapeLeft(RotationRate);
+				Gravity = ReorientLandscapeLeft(Gravity);
+				Acceleration = ReorientLandscapeLeft(Acceleration);
+				break;
+			}
+
+			if (UnifyMotionSpace == 2)
+			{
+				// Right-handed variation
+				Attitude = FVector(Attitude.Y, Attitude.X, Attitude.Z);
+				RotationRate = FVector(RotationRate.Y, RotationRate.X, RotationRate.Z);
+				Gravity = FVector(Gravity.Y, Gravity.X, Gravity.Z);
+				Acceleration = FVector(Acceleration.Y, Acceleration.X, Acceleration.Z);
+			}
+		}
 
         MessageHandler->OnMotionDetected(Attitude, RotationRate, Gravity, Acceleration, 0);
     }
