Default platform device input mapper will not automatically unpair input devices when the associated user signs out

 - otherwise there is no way to automatically restore the mapping if the user returns.
 - that responsibility now lies with specializations.

#rnx

diff --git a/Engine/Source/Runtime/ApplicationCore/Public/GenericPlatform/GenericPlatformInputDeviceMapper.h b/Engine/Source/Runtime/ApplicationCore/Public/GenericPlatform/GenericPlatformInputDeviceMapper.h
--- a/Engine/Source/Runtime/ApplicationCore/Public/GenericPlatform/GenericPlatformInputDeviceMapper.h
+++ b/Engine/Source/Runtime/ApplicationCore/Public/GenericPlatform/GenericPlatformInputDeviceMapper.h
@@ -312,5 +312,5 @@ protected:
 	const bool bShouldBroadcastLegacyDelegates = true;
 
 	/** Additional functionality that specializations can customize */
-	bool bUnpairInputDevicesWhenLoggingOut = true;
+	bool bUnpairInputDevicesWhenLoggingOut = false;
 };
