Fixed some Printf specifiers.

diff --git a/Engine/Source/Developer/DerivedDataCache/Private/DerivedDataCache.cpp b/Engine/Source/Developer/DerivedDataCache/Private/DerivedDataCache.cpp
--- a/Engine/Source/Developer/DerivedDataCache/Private/DerivedDataCache.cpp
+++ b/Engine/Source/Developer/DerivedDataCache/Private/DerivedDataCache.cpp
@@ -476,7 +476,7 @@ class FDerivedDataCache final
 
 					if (!bMatchesInSize || bDifferentMemory)
 					{
-						FString ErrMsg = FString::Printf(TEXT("There is a mismatch between the DDC data and the generated data for plugin (%s) for asset (%s). BytesInDDC:%d, BytesGenerated:%d, bDifferentMemory:%d, offset:%d"), DataDeriver->GetPluginName(), *DataDeriver->GetDebugContextString(), NumInDDC, NumGenerated, bDifferentMemory, DifferentOffset);
+						FString ErrMsg = FString::Printf(TEXT("There is a mismatch between the DDC data and the generated data for plugin (%s) for asset (%s). BytesInDDC:%" INT64_FMT ", BytesGenerated:%" INT64_FMT ", bDifferentMemory:%d, offset:%d"), DataDeriver->GetPluginName(), *DataDeriver->GetDebugContextString(), NumInDDC, NumGenerated, bDifferentMemory, DifferentOffset);
 						ensureMsgf(false, TEXT("%s"), *ErrMsg);
 						UE_LOG(LogDerivedDataCache, Error, TEXT("%s"), *ErrMsg );
 					}
diff --git a/Engine/Source/Developer/Zen/Private/ZenServerHttp.cpp b/Engine/Source/Developer/Zen/Private/ZenServerHttp.cpp
--- a/Engine/Source/Developer/Zen/Private/ZenServerHttp.cpp
+++ b/Engine/Source/Developer/Zen/Private/ZenServerHttp.cpp
@@ -411,12 +411,12 @@ namespace UE::Zen {
 			case RequestVerb::Get:
 				bSuccess = Is2xx || Is404;
 				VerbStr = TEXT("fetching");
-				AdditionalInfo = FString::Printf(TEXT("Received: %d bytes."), BytesReceived);
+				AdditionalInfo = FString::Printf(TEXT("Received: %zu bytes."), BytesReceived);
 				break;
 			case RequestVerb::Put:
 				bSuccess = Is2xx;
 				VerbStr = TEXT("updating");
-				AdditionalInfo = FString::Printf(TEXT("Sent: %d bytes."), BytesSent);
+				AdditionalInfo = FString::Printf(TEXT("Sent: %zu bytes."), BytesSent);
 				break;
 			case RequestVerb::Post:
 				bSuccess = Is2xx || Is404;
diff --git a/Engine/Source/Runtime/ApplicationCore/Private/Windows/WindowsPlatformSurvey.cpp b/Engine/Source/Runtime/ApplicationCore/Private/Windows/WindowsPlatformSurvey.cpp
--- a/Engine/Source/Runtime/ApplicationCore/Private/Windows/WindowsPlatformSurvey.cpp
+++ b/Engine/Source/Runtime/ApplicationCore/Private/Windows/WindowsPlatformSurvey.cpp
@@ -109,7 +109,7 @@ bool FWindowsPlatformSurvey::GetSurveyResults( FHardwareSurveyResults& OutResult
 			{
 				OutResults.ErrorCount++;
 				WriteFStringToResults(OutResults.LastSurveyError, TEXT("CallNtPowerInformation() failed to get processor power info"));
-				WriteFStringToResults(OutResults.LastSurveyErrorDetail, FString::Printf(TEXT("NTSTATUS: 0x%0x"), NTStatus));
+				WriteFStringToResults(OutResults.LastSurveyErrorDetail, FString::Printf(TEXT("NTSTATUS: 0x%0x"), (uint32)NTStatus));
 			}
 			delete [] PowerInfo;
 		}
@@ -211,7 +211,7 @@ bool FWindowsPlatformSurvey::GetSurveyResults( FHardwareSurveyResults& OutResult
 		UE_LOG(LogWindows, Warning, TEXT("FWindowsPlatformSurvey::TickSurveyHardware() failed to get query interface from WinSAT API") );
 		OutResults.ErrorCount++;
 		WriteFStringToResults(OutResults.LastSurveyError, TEXT("CoCreateInstance() failed to get WinSAT"));
-		WriteFStringToResults(OutResults.LastSurveyErrorDetail, FString::Printf(TEXT("HRESULT: 0x%0x"), COMResult));
+		WriteFStringToResults(OutResults.LastSurveyErrorDetail, FString::Printf(TEXT("HRESULT: 0x%0x"), (uint32)COMResult));
 
 	}
 	else
@@ -224,7 +224,7 @@ bool FWindowsPlatformSurvey::GetSurveyResults( FHardwareSurveyResults& OutResult
 			UE_LOG(LogWindows, Error, TEXT("FWindowsPlatformSurvey::TickSurveyHardware() failed to get assessment results from WinSAT API") );
 			OutResults.ErrorCount++;
 			WriteFStringToResults(OutResults.LastSurveyError, TEXT("get_Info() failed to get WinSAT assessment results"));
-			WriteFStringToResults(OutResults.LastSurveyErrorDetail, FString::Printf(TEXT("HRESULT: 0x%0x"), COMResult));
+			WriteFStringToResults(OutResults.LastSurveyErrorDetail, FString::Printf(TEXT("HRESULT: 0x%0x"), (uint32)COMResult));
 
 		}
 		else
@@ -237,7 +237,7 @@ bool FWindowsPlatformSurvey::GetSurveyResults( FHardwareSurveyResults& OutResult
 				UE_LOG(LogWindows, Error, TEXT("FWindowsPlatformSurvey::TickSurveyHardware() failed to get assessment state from WinSAT API") );
 				OutResults.ErrorCount++;
 				WriteFStringToResults(OutResults.LastSurveyError, TEXT("get_AssessmentState() failed to get WinSAT assessment state"));
-				WriteFStringToResults(OutResults.LastSurveyErrorDetail, FString::Printf(TEXT("HRESULT: 0x%0x"), COMResult));
+				WriteFStringToResults(OutResults.LastSurveyErrorDetail, FString::Printf(TEXT("HRESULT: 0x%0x"), (uint32)COMResult));
 			}
 			else
 			{
@@ -324,7 +324,7 @@ bool FWindowsPlatformSurvey::GetSurveyResults( FHardwareSurveyResults& OutResult
 		UE_LOG(LogWindows, Warning, TEXT("FWindowsPlatformSurvey::TickSurveyHardware() failed to get system power capabilities. Assuming desktop PC.") );
 		OutResults.ErrorCount++;
 		WriteFStringToResults(OutResults.LastSurveyError, TEXT("CallNtPowerInformation() failed to get system power capabilities"));
-		WriteFStringToResults(OutResults.LastSurveyErrorDetail, FString::Printf(TEXT("NTSTATUS: 0x%0x"), NTStatus));
+		WriteFStringToResults(OutResults.LastSurveyErrorDetail, FString::Printf(TEXT("NTSTATUS: 0x%0x"), (uint32)NTStatus));
 	}
 #endif	// #if USING_POWRPROF
 
@@ -346,7 +346,7 @@ bool FWindowsPlatformSurvey::GetSubComponentIndex( IProvideWinSATResultsInfo* Wi
 		UE_LOG(LogWindows, Log, TEXT("FWindowsPlatformSurvey::GetSubComponentIndex() failed to get assessment info for a sub-component from WinSAT API.") );
 		OutSurveyResults.ErrorCount++;
 		WriteFStringToResults(OutSurveyResults.LastPerformanceIndexError, FString::Printf(TEXT("GetAssessmentInfo() failed to get WinSAT assessment for sub-component %d"), SubComponent));
-		WriteFStringToResults(OutSurveyResults.LastPerformanceIndexErrorDetail, FString::Printf(TEXT("HRESULT: 0x%0x"), COMResult));
+		WriteFStringToResults(OutSurveyResults.LastPerformanceIndexErrorDetail, FString::Printf(TEXT("HRESULT: 0x%0x"), (uint32)COMResult));
 	}
 	else
 	{
@@ -357,7 +357,7 @@ bool FWindowsPlatformSurvey::GetSubComponentIndex( IProvideWinSATResultsInfo* Wi
 			UE_LOG(LogWindows, Log, TEXT("FWindowsPlatformSurvey::GetSubComponentIndex() failed to get sub-component score from WinSAT API.") );
 			OutSurveyResults.ErrorCount++;
 			WriteFStringToResults(OutSurveyResults.LastPerformanceIndexError, FString::Printf(TEXT("get_Score() failed to get WinSAT WIE score for sub-component %d"), SubComponent));
-			WriteFStringToResults(OutSurveyResults.LastPerformanceIndexErrorDetail, FString::Printf(TEXT("HRESULT: 0x%0x"), COMResult));
+			WriteFStringToResults(OutSurveyResults.LastPerformanceIndexErrorDetail, FString::Printf(TEXT("HRESULT: 0x%0x"), (uint32)COMResult));
 		}
 		else
 		{
@@ -406,7 +406,7 @@ void GetOSVersionLabels(const SYSTEM_INFO& SystemInfo, FHardwareSurveyResults& O
 		UE_LOG( LogWindows, Warning, TEXT( "FWindowsPlatformSurvey::GetOSVersionLabel() unknown Windows version info from GetVersionEx()" ) );
 		OutResults.ErrorCount++;
 		WriteFStringToResults( OutResults.LastSurveyError, TEXT( "GetVersionEx() returned unknown version" ) );
-		WriteFStringToResults( OutResults.LastSurveyErrorDetail, FString::Printf( TEXT( "dwMajorVersion: %d  dwMinorVersion: %d" ), OsVersionInfo.dwMajorVersion, OsVersionInfo.dwMinorVersion ) );
+		WriteFStringToResults( OutResults.LastSurveyErrorDetail, FString::Printf( TEXT( "dwMajorVersion: %u  dwMinorVersion: %u" ), (uint32)OsVersionInfo.dwMajorVersion, (uint32)OsVersionInfo.dwMinorVersion ) );
 	}
 
 	if( ErrorCode & FWindowsOSVersionHelper::ERROR_GETVERSIONEX_FAILED )
diff --git a/Engine/Source/Runtime/Core/Private/Compression/CompressionUtil.cpp b/Engine/Source/Runtime/Core/Private/Compression/CompressionUtil.cpp
--- a/Engine/Source/Runtime/Core/Private/Compression/CompressionUtil.cpp
+++ b/Engine/Source/Runtime/Core/Private/Compression/CompressionUtil.cpp
@@ -120,11 +120,11 @@ namespace FCompressionUtil
 			}
 			if (LineOffset >= 0)
 			{
-				Results.Add(FString::Printf(TEXT("%016X: %s"), LineOffset, *HexString));
+				Results.Add(FString::Printf(TEXT("%016" INT64_X_FMT ": %s"), LineOffset, *HexString));
 			}
 			else
 			{
-				Results.Add(FString::Printf(TEXT(" -%014X: %s"), -LineOffset, *HexString));
+				Results.Add(FString::Printf(TEXT(" -%014" INT64_X_FMT ": %s"), -LineOffset, *HexString));
 			}
 		}
 		return Results;
diff --git a/Engine/Source/Runtime/Core/Private/GenericPlatform/GenericPlatformCrashContext.cpp b/Engine/Source/Runtime/Core/Private/GenericPlatform/GenericPlatformCrashContext.cpp
--- a/Engine/Source/Runtime/Core/Private/GenericPlatform/GenericPlatformCrashContext.cpp
+++ b/Engine/Source/Runtime/Core/Private/GenericPlatform/GenericPlatformCrashContext.cpp
@@ -694,7 +694,7 @@ FGenericCrashContext::FGenericCrashContext(ECrashContextType InType, const TCHAR
 
 FString FGenericCrashContext::GetTempSessionContextFilePath(uint64 ProcessID)
 {
-	return FPlatformProcess::UserTempDir() / FString::Printf(TEXT("UECrashContext-%u.xml"), ProcessID);
+	return FPlatformProcess::UserTempDir() / FString::Printf(TEXT("UECrashContext-%" UINT64_FMT ".xml"), ProcessID);
 }
 
 void FGenericCrashContext::CleanupTempSessionContextFiles(const FTimespan& ExpirationAge)
diff --git a/Engine/Source/Runtime/Core/Private/Windows/WindowsPlatformMisc.cpp b/Engine/Source/Runtime/Core/Private/Windows/WindowsPlatformMisc.cpp
--- a/Engine/Source/Runtime/Core/Private/Windows/WindowsPlatformMisc.cpp
+++ b/Engine/Source/Runtime/Core/Private/Windows/WindowsPlatformMisc.cpp
@@ -229,7 +229,7 @@ int32 GetOSVersionsHelper( TCHAR* OutOSVersionLabel, int32 OSVersionLabelLength,
 					}
 
 					// Add the build number as displayed by the winver utility.
-					OSVersionLabel += FString::Printf(TEXT(" [%u.%u.%u.%s]"), OsVersionInfo.dwMajorVersion, OsVersionInfo.dwMinorVersion, OsVersionInfo.dwBuildNumber, *UpdateBuildRevision);
+					OSVersionLabel += FString::Printf(TEXT(" [%u.%u.%u.%s]"), (uint32)OsVersionInfo.dwMajorVersion, (uint32)OsVersionInfo.dwMinorVersion, (uint32)OsVersionInfo.dwBuildNumber, *UpdateBuildRevision);
 				}
 				break;
 			default:
@@ -3942,7 +3942,7 @@ bool FWindowsPlatformMisc::QueryRegKey( const Windows::HKEY InKey, const TCHAR*
 						::DWORD Value;
 						if (RegQueryValueEx(Key, InValueName, NULL, NULL, (LPBYTE)&Value, &Size) == ERROR_SUCCESS)
 						{
-							OutData = FString::Printf(TEXT("%d"), Value);
+							OutData = FString::Printf(TEXT("%u"), (uint32)Value);
 							bSuccess = true;
 						}
 						break;
@@ -3953,7 +3953,7 @@ bool FWindowsPlatformMisc::QueryRegKey( const Windows::HKEY InKey, const TCHAR*
 						int64 Value;
 						if (RegQueryValueEx(Key, InValueName, NULL, NULL, (LPBYTE)&Value, &Size) == ERROR_SUCCESS)
 						{
-							OutData = FString::Printf(TEXT("%lld"), Value);
+							OutData = FString::Printf(TEXT("%" INT64_FMT ), Value);
 							bSuccess = true;
 						}
 						break;
diff --git a/Engine/Source/Runtime/Core/Tests/Containers/BitArrayTests.cpp b/Engine/Source/Runtime/Core/Tests/Containers/BitArrayTests.cpp
--- a/Engine/Source/Runtime/Core/Tests/Containers/BitArrayTests.cpp
+++ b/Engine/Source/Runtime/Core/Tests/Containers/BitArrayTests.cpp
@@ -2446,7 +2446,7 @@ public:
 				FBitArrayMemory::ModularizeWordOffset(Data, Offset);
 				uint32* ExpectedData = BaseData + WordLength;
 				uint32 ExpectedOffset = BitLength;
-				CHECK_MESSAGE(FString::Printf(TEXT("ModularizeWordOffset WordLength=%d, BitLength=%d\nExpected: Data=%d, Offset=%d\nActual:  Data=%d, Offset=%d"), WordLength, BitLength, ExpectedData - BaseData, ExpectedOffset, Data - BaseData, Offset),
+				CHECK_MESSAGE(FString::Printf(TEXT("ModularizeWordOffset WordLength=%d, BitLength=%d\nExpected: Data=%d, Offset=%d\nActual:  Data=%d, Offset=%d"), WordLength, BitLength, UE_PTRDIFF_TO_INT32(ExpectedData - BaseData), ExpectedOffset, UE_PTRDIFF_TO_INT32(Data - BaseData), Offset),
 					!(Data != ExpectedData || Offset != ExpectedOffset));
 			}
 		}
diff --git a/Engine/Source/Runtime/Core/Tests/Misc/DateTimeTest.cpp b/Engine/Source/Runtime/Core/Tests/Misc/DateTimeTest.cpp
--- a/Engine/Source/Runtime/Core/Tests/Misc/DateTimeTest.cpp
+++ b/Engine/Source/Runtime/Core/Tests/Misc/DateTimeTest.cpp
@@ -10,7 +10,7 @@
 #include "TestCommon/Expectations.h"
 #endif
 
-#define TestUnixEquivalent(Desc, A, B) if ((A).ToUnixTimestamp() != (B)) FAIL_CHECK(FString::Printf(TEXT("%s - A=%d B=%d"), Desc, (A).ToUnixTimestamp(), (B)));
+#define TestUnixEquivalent(Desc, A, B) if ((A).ToUnixTimestamp() != (B)) FAIL_CHECK(FString::Printf(TEXT("%s - A=%" INT64_FMT " B=%" INT64_FMT), Desc, (A).ToUnixTimestamp(), (B)));
 #define TestYear(Desc, A, B) if ((A.GetYear()) != (B)) FAIL_CHECK(FString::Printf(TEXT("%s - A=%d B=%d"), Desc, (A.GetYear()), (B)));
 #define TestMonth(Desc, A, B) if ((A.GetMonth()) != (B)) FAIL_CHECK(FString::Printf(TEXT("%s - A=%d B=%d"), Desc, (A.GetMonth()), B));
 #define TestMonthOfYear(Desc, A, B) if ((A.GetMonthOfYear()) != (B)) FAIL_CHECK(FString::Printf(TEXT("%s - A=%d B=%d"), Desc ,(static_cast<int32>(A.GetMonthOfYear())), static_cast<int32>(B)));
diff --git a/Engine/Source/Runtime/Core/Tests/Misc/FileTest.cpp b/Engine/Source/Runtime/Core/Tests/Misc/FileTest.cpp
--- a/Engine/Source/Runtime/Core/Tests/Misc/FileTest.cpp
+++ b/Engine/Source/Runtime/Core/Tests/Misc/FileTest.cpp
@@ -54,10 +54,10 @@ TEST_CASE_NAMED(FFileTruncateTest, "System::Core::Misc::FileTruncate", "[.][Edit
 		{
 			const int64 ActualEOFPos = TestFile->Tell();
 			const int64 ExpectedEOFPos = (sizeof(int32) * 4);
-			REQUIRE_MESSAGE(FString::Printf(TEXT("File was not the expected size (got %d, expected %d): %s"), ActualEOFPos, ExpectedEOFPos, *TempFilename), ActualEOFPos == ExpectedEOFPos);
+			REQUIRE_MESSAGE(FString::Printf(TEXT("File was not the expected size (got %" INT64_FMT ", expected %" INT64_FMT "): %s"), ActualEOFPos, ExpectedEOFPos, *TempFilename), ActualEOFPos == ExpectedEOFPos);
 
 			const int64 ActualFileSize = TestFile->Size();
-			REQUIRE_MESSAGE(FString::Printf(TEXT("File was not the expected size (got %d, expected %d): %s"), ActualFileSize, ExpectedEOFPos, *TempFilename), ActualFileSize == ExpectedEOFPos);
+			REQUIRE_MESSAGE(FString::Printf(TEXT("File was not the expected size (got %" INT64_FMT ", expected %" INT64_FMT "): %s"), ActualFileSize, ExpectedEOFPos, *TempFilename), ActualFileSize == ExpectedEOFPos);
 		}
 
 		// Truncate the file at our test pos
@@ -66,7 +66,7 @@ TEST_CASE_NAMED(FFileTruncateTest, "System::Core::Misc::FileTruncate", "[.][Edit
 		// Validate that the size is reported correctly
 		{
 			const int64 ActualFileSize = TestFile->Size();
-			REQUIRE_MESSAGE(FString::Printf(TEXT("File was not the expected size after truncation (got %d, expected %d): %s"), ActualFileSize, ExpectedTruncatePos, *TempFilename), ActualFileSize == ExpectedTruncatePos);
+			REQUIRE_MESSAGE(FString::Printf(TEXT("File was not the expected size after truncation (got %" INT64_FMT ", expected %" INT64_FMT "): %s"), ActualFileSize, ExpectedTruncatePos, *TempFilename), ActualFileSize == ExpectedTruncatePos);
 		}
 
 		// Validate that we can't read past the truncation point
@@ -122,7 +122,7 @@ TEST_CASE_NAMED(FFileAppendTest, "System::Core::Misc::FileAppend", "[.][EditorCo
 		{
 			const int64 ActualEOFPos = TestFile->Tell();
 			const int64 ExpectedEOFPos = TestFile->Size();
-			REQUIRE_MESSAGE(FString::Printf(TEXT("File did not seek to the end when opening (got %d, expected %d): %s"), ActualEOFPos, ExpectedEOFPos, *TempFilename), ActualEOFPos == ExpectedEOFPos);
+			REQUIRE_MESSAGE(FString::Printf(TEXT("File did not seek to the end when opening (got %" INT64_FMT ", expected %" INT64_FMT "): %s"), ActualEOFPos, ExpectedEOFPos, *TempFilename), ActualEOFPos == ExpectedEOFPos);
 		}
 
 		TestData.Add(One);
diff --git a/Engine/Source/Runtime/CoreUObject/Private/Serialization/AsyncLoading.cpp b/Engine/Source/Runtime/CoreUObject/Private/Serialization/AsyncLoading.cpp
--- a/Engine/Source/Runtime/CoreUObject/Private/Serialization/AsyncLoading.cpp
+++ b/Engine/Source/Runtime/CoreUObject/Private/Serialization/AsyncLoading.cpp
@@ -3552,11 +3552,11 @@ void FAsyncPackage::EventDrivenSerializeExport(int32 LocalExportIndex)
 		{
 			if (Object->GetClass()->HasAnyClassFlags(CLASS_Deprecated))
 			{
-				UE_LOG(LogStreaming, Warning, TEXT("%s"), *FString::Printf(TEXT("%s: Serial size mismatch: Got %d, Expected %d"), *Object->GetFullName(), (int32)(AsyncLoader->Tell() - Export.SerialOffset), Export.SerialSize));
+				UE_LOG(LogStreaming, Warning, TEXT("%s"), *FString::Printf(TEXT("%s: Serial size mismatch: Got %d, Expected %" INT64_FMT), *Object->GetFullName(), (int32)(AsyncLoader->Tell() - Export.SerialOffset), Export.SerialSize));
 			}
 			else
 			{
-				UE_LOG(LogStreaming, Fatal, TEXT("%s"), *FString::Printf(TEXT("%s: Serial size mismatch: Got %d, Expected %d"), *Object->GetFullName(), (int32)(AsyncLoader->Tell() - Export.SerialOffset), Export.SerialSize));
+				UE_LOG(LogStreaming, Fatal, TEXT("%s"), *FString::Printf(TEXT("%s: Serial size mismatch: Got %d, Expected %" INT64_FMT), *Object->GetFullName(), (int32)(AsyncLoader->Tell() - Export.SerialOffset), Export.SerialSize));
 			}
 		}
 
diff --git a/Engine/Source/Runtime/CoreUObject/Private/UObject/LinkerLoad.cpp b/Engine/Source/Runtime/CoreUObject/Private/UObject/LinkerLoad.cpp
--- a/Engine/Source/Runtime/CoreUObject/Private/UObject/LinkerLoad.cpp
+++ b/Engine/Source/Runtime/CoreUObject/Private/UObject/LinkerLoad.cpp
@@ -3111,9 +3111,11 @@ void FLinkerLoad::GatherImportDependencies(int32 ImportIndex, TSet<FDependencyRe
 			// don't warn about the suppressed Import.XObject dependencies
 			if (!Import.XObject)
 			{
-				UE_ASSET_LOG(LogLinker, Warning, PackagePath, TEXT("VerifyImportInner failed [(%x, %d), (%x, %d)] for %s"), 
-					Import.XObject, Import.XObject ? (Import.XObject->IsNative() ? 1 : 0) : 0, 
-					Import.SourceLinker, Import.SourceIndex, 
+				UE_ASSET_LOG(LogLinker, Warning, PackagePath, TEXT("VerifyImportInner failed [(0x%" UPTRINT_X_FMT ", %d), (0x%" UPTRINT_X_FMT ", %d)] for %s"), 
+					(UPTRINT)Import.XObject,
+					Import.XObject ? (Import.XObject->IsNative() ? 1 : 0) : 0, 
+					(UPTRINT)Import.SourceLinker,
+					Import.SourceIndex, 
 					*GetImportFullName(ImportIndex));
 			}
 			EndLoad(LoadContext);
@@ -4891,11 +4893,11 @@ void FLinkerLoad::Preload( UObject* Object )
 				{
 					if (Object->GetClass()->HasAnyClassFlags(CLASS_Deprecated))
 					{
-						UE_ASSET_LOG(LogLinker, Warning, PackagePath, TEXT("%s: Serial size mismatch: Got %d, Expected %d"), *Object->GetFullName(), (int32)SizeSerialized, ExpectedSerialSize);
+						UE_ASSET_LOG(LogLinker, Warning, PackagePath, TEXT("%s: Serial size mismatch: Got %d, Expected %" INT64_FMT), *Object->GetFullName(), (int32)SizeSerialized, ExpectedSerialSize);
 					}
 					else
 					{
-						UE_ASSET_LOG(LogLinker, Fatal, PackagePath, TEXT("%s: Serial size mismatch: Got %d, Expected %d"), *Object->GetFullName(), (int32)SizeSerialized, ExpectedSerialSize);
+						UE_ASSET_LOG(LogLinker, Fatal, PackagePath, TEXT("%s: Serial size mismatch: Got %d, Expected %" INT64_FMT), *Object->GetFullName(), (int32)SizeSerialized, ExpectedSerialSize);
 					}
 				}
 
diff --git a/Engine/Source/Runtime/CoreUObject/Private/UObject/Obj.cpp b/Engine/Source/Runtime/CoreUObject/Private/UObject/Obj.cpp
--- a/Engine/Source/Runtime/CoreUObject/Private/UObject/Obj.cpp
+++ b/Engine/Source/Runtime/CoreUObject/Private/UObject/Obj.cpp
@@ -3439,7 +3439,7 @@ FString UObject::GetDefaultConfigFilename() const
 	if (OverridePlatform.Len())
 	{
 		bool bIsPlatformExtension = FPaths::DirectoryExists(FPaths::EnginePlatformExtensionDir(*OverridePlatform));
-		FString RegularPath = FString::Printf(TEXT("%s%s"), *FPaths::SourceConfigDir(), *OverridePlatform, *OverridePlatform, *GetClass()->ClassConfigName.ToString());
+		FString RegularPath = FString::Printf(TEXT("%s%s"), *FPaths::SourceConfigDir(), *OverridePlatform);
 		FString SelectedPath = RegularPath;
 
 		bool bPlatformConfigExistsInRegular = FPaths::DirectoryExists(*RegularPath);
diff --git a/Engine/Source/Runtime/Serialization/Private/Backends/CborStructDeserializerBackend.cpp b/Engine/Source/Runtime/Serialization/Private/Backends/CborStructDeserializerBackend.cpp
--- a/Engine/Source/Runtime/Serialization/Private/Backends/CborStructDeserializerBackend.cpp
+++ b/Engine/Source/Runtime/Serialization/Private/Backends/CborStructDeserializerBackend.cpp
@@ -25,7 +25,7 @@ const FString& FCborStructDeserializerBackend::GetCurrentPropertyName() const
 FString FCborStructDeserializerBackend::GetDebugString() const
 {
 	FArchive* Ar = const_cast<FArchive*>(CborReader.GetArchive());
-	return FString::Printf(TEXT("Offset: %u"), Ar ? Ar->Tell() : 0);
+	return FString::Printf(TEXT("Offset: %" UINT64_FMT), Ar ? Ar->Tell() : 0);
 }
 
 const FString& FCborStructDeserializerBackend::GetLastErrorMessage() const
