Don't allow downloads over cellular on iOS for now, pending integration with opt-in from the user

diff --git a/Engine/Source/Runtime/ApplicationCore/Private/IOS/IOSBackgroundURLSessionHandler.cpp b/Engine/Source/Runtime/ApplicationCore/Private/IOS/IOSBackgroundURLSessionHandler.cpp
--- a/Engine/Source/Runtime/ApplicationCore/Private/IOS/IOSBackgroundURLSessionHandler.cpp
+++ b/Engine/Source/Runtime/ApplicationCore/Private/IOS/IOSBackgroundURLSessionHandler.cpp
@@ -51,9 +51,9 @@ bool FBackgroundURLSessionHandler::InitBackgroundSession(const FString& SessionI
 				config.timeoutIntervalForRequest = BackgroundHttpTimeout;	        //Timeout for requests not getting a response from server at the iOS level. Will still work in the BG.
 				config.timeoutIntervalForResource = BackgroundHttpResourceTimeout;  //Will automatically retry everytime above setting causes a timeout until this many seconds pass. Then we will get
                                                                                     //an error result. No Error for server not responding until is sent by a background task until this timeout is reached.
-                //Other unused settings we might want to configure
+				config.allowsCellularAccess = NO; // don't allow cellular by default, TODO integrate opt-in dialog
+				//Other unused settings we might want to configure
 				//HTTPMaximumConnectionsPerHost
-				//allowsCellularAccess
                 
                 //Can't use on BG session, but may need to consider for any future FG session work:
 				//waitsForConnectivity  -- This is forced true on background sessions
