Updated clipboard operations in WindowsPlatformApplicationMisc to report an error instead of asserting if clipboard ownership is stolen before CloseClipboard is called

diff --git a/Engine/Source/Runtime/ApplicationCore/Private/Windows/WindowsPlatformApplicationMisc.cpp b/Engine/Source/Runtime/ApplicationCore/Private/Windows/WindowsPlatformApplicationMisc.cpp
--- a/Engine/Source/Runtime/ApplicationCore/Private/Windows/WindowsPlatformApplicationMisc.cpp
+++ b/Engine/Source/Runtime/ApplicationCore/Private/Windows/WindowsPlatformApplicationMisc.cpp
@@ -395,7 +395,10 @@ void FWindowsPlatformApplicationMisc::ClipboardCopy(const TCHAR* Str)
 		GlobalUnlock( GlobalMem );
 		if( SetClipboardData( CF_UNICODETEXT, GlobalMem ) == NULL )
 			UE_LOG(LogWindows, Fatal,TEXT("SetClipboardData failed with error code %i"), (uint32)GetLastError() );
-		verify(CloseClipboard());
+		if (!CloseClipboard())
+		{
+			UE_LOG(LogWindows, Warning, TEXT("CloseClipboard failed with error code %i"), (uint32)GetLastError());
+		}
 	}
 	else
 	{
@@ -439,7 +442,10 @@ void FWindowsPlatformApplicationMisc::ClipboardPaste(class FString& Result)
 			}
 			GlobalUnlock( GlobalMem );
 		}
-		verify(CloseClipboard());
+		if (!CloseClipboard())
+		{
+			UE_LOG(LogWindows, Warning, TEXT("CloseClipboard failed with error code %i"), (uint32)GetLastError());
+		}
 	}
 	else 
 	{
