Execute logic of change route audio core delegate on game thread instead of main thread

diff --git a/Engine/Source/Runtime/ApplicationCore/Private/IOS/IOSAppDelegate.cpp b/Engine/Source/Runtime/ApplicationCore/Private/IOS/IOSAppDelegate.cpp
--- a/Engine/Source/Runtime/ApplicationCore/Private/IOS/IOSAppDelegate.cpp
+++ b/Engine/Source/Runtime/ApplicationCore/Private/IOS/IOSAppDelegate.cpp
@@ -569,18 +569,23 @@ static IOSAppDelegate* CachedDelegate = nil;
 
 	[[NSNotificationCenter defaultCenter] addObserverForName:AVAudioSessionRouteChangeNotification object : nil queue : nil usingBlock : ^ (NSNotification *notification)
 	{
-		switch ([[[notification userInfo] objectForKey:AVAudioSessionRouteChangeReasonKey] unsignedIntegerValue])
+		NSUInteger RouteChangeReason = [[[notification userInfo] objectForKey:AVAudioSessionRouteChangeReasonKey] unsignedIntegerValue];
+		[FIOSAsyncTask CreateTaskWithBlock : ^ bool(void)
 		{
-			case AVAudioSessionRouteChangeReasonNewDeviceAvailable:
-				// headphones plugged in
-				FCoreDelegates::AudioRouteChangedDelegate.Broadcast(true);
-				break;
-
-			case AVAudioSessionRouteChangeReasonOldDeviceUnavailable:
-				// headphones unplugged
-				FCoreDelegates::AudioRouteChangedDelegate.Broadcast(false);
-				break;
-		}
+			switch (RouteChangeReason)
+			{
+				case AVAudioSessionRouteChangeReasonNewDeviceAvailable:
+					// headphones plugged in
+					FCoreDelegates::AudioRouteChangedDelegate.Broadcast(true);
+					break;
+					
+				case AVAudioSessionRouteChangeReasonOldDeviceUnavailable:
+					// headphones unplugged
+					FCoreDelegates::AudioRouteChangedDelegate.Broadcast(false);
+					break;
+			}
+			return true;
+		}];
 	}];
 
 	self.bAudioSessionInitialized = true;
