Motion input (gyro, accelerometer) with supported controllers and mobile platforms now use the same axes and units by default.

In the past, we've been inconsistent with motion-related input on controllers and mobile platforms -- they've used different coordinate spaces and different units. This means users (us included) have had to manually convert accelerometer and gyro input for each platform on a per-game basis.

In 2022 we added the *.UnifyMotionSpace CVars, which let us get motion input in Unreal's coordinate space (left-handed, z-up, etc) when set to 1, or a similar but right-handed space when set to 2. Setting to non-zero also unifies the units across platforms, giving angular velocity in radians per second and acceleration in g-units.

For a better out-of-the-box experience, this change makes all of the *.UnifyMotionSpace CVars set to 1 by default.

#tests PIE

diff --git a/Engine/Source/Runtime/ApplicationCore/Private/Android/AndroidInputInterface.cpp b/Engine/Source/Runtime/ApplicationCore/Private/Android/AndroidInputInterface.cpp
--- a/Engine/Source/Runtime/ApplicationCore/Private/Android/AndroidInputInterface.cpp
+++ b/Engine/Source/Runtime/ApplicationCore/Private/Android/AndroidInputInterface.cpp
@@ -71,11 +71,11 @@ static FAutoConsoleVariableRef CVarAndroidOldXBoxWirelessFirmware(
 	TEXT("Determines how XBox Wireless controller mapping is handled. 0 assumes new firmware, 1 will use old firmware mapping (Default: 0)"),
 	ECVF_Default);
 
-int32 AndroidUnifyMotionSpace = 0;
+int32 AndroidUnifyMotionSpace = 1;
 static FAutoConsoleVariableRef CVarAndroidUnifyMotionSpace(
 	TEXT("Android.UnifyMotionSpace"),
 	AndroidUnifyMotionSpace,
-	TEXT("If set to non-zero, acceleration, gravity, and rotation rate will all be in the same coordinate space. 0 for legacy behaviour. 1 will match Unreal's coordinate space (left-handed, z-up, etc). 2 will be right-handed by swapping x and y. Non-zero also forces rotation rate units to be radians/s and acceleration units to be g."),
+	TEXT("If set to non-zero, acceleration, gravity, and rotation rate will all be in the same coordinate space. 0 for legacy behaviour. 1 (default as of 5.5) will match Unreal's coordinate space (left-handed, z-up, etc). 2 will be right-handed by swapping x and y. Non-zero also forces rotation rate units to be radians/s and acceleration units to be g."),
 	ECVF_Default);
 
 TSharedRef< FAndroidInputInterface > FAndroidInputInterface::Create(const TSharedRef< FGenericApplicationMessageHandler >& InMessageHandler, const TSharedPtr< ICursor >& InCursor)
diff --git a/Engine/Source/Runtime/ApplicationCore/Private/IOS/IOSInputInterface.cpp b/Engine/Source/Runtime/ApplicationCore/Private/IOS/IOSInputInterface.cpp
--- a/Engine/Source/Runtime/ApplicationCore/Private/IOS/IOSInputInterface.cpp
+++ b/Engine/Source/Runtime/ApplicationCore/Private/IOS/IOSInputInterface.cpp
@@ -19,7 +19,7 @@ static TAutoConsoleVariable<float> CVarHapticsKickHeavy(TEXT("ios.VibrationHapti
 static TAutoConsoleVariable<float> CVarHapticsKickMedium(TEXT("ios.VibrationHapticsKickMediumValue"), 0.5f, TEXT("Vibation values higher than this will kick a haptics medium Impact"));
 static TAutoConsoleVariable<float> CVarHapticsKickLight(TEXT("ios.VibrationHapticsKickLightValue"), 0.3f, TEXT("Vibation values higher than this will kick a haptics light Impact"));
 static TAutoConsoleVariable<float> CVarHapticsRest(TEXT("ios.VibrationHapticsRestValue"), 0.2f, TEXT("Vibation values lower than this will allow haptics to Kick again when going over ios.VibrationHapticsKickValue"));
-static TAutoConsoleVariable<int32> CVarUnifyMotionSpace(TEXT("ios.UnifyMotionSpace"), 0, TEXT("If set to non-zero, acceleration, gravity, and rotation rate will all be in the same coordinate space. 0 for legacy behaviour. 1 will match Unreal's coordinate space (left-handed, z-up, etc). 2 will be right-handed by swapping x and y. Non-zero also forces rotation rate units to be radians/s and acceleration units to be g."));
+static TAutoConsoleVariable<int32> CVarUnifyMotionSpace(TEXT("ios.UnifyMotionSpace"), 1, TEXT("If set to non-zero, acceleration, gravity, and rotation rate will all be in the same coordinate space. 0 for legacy behaviour. 1 (default as of 5.5) will match Unreal's coordinate space (left-handed, z-up, etc). 2 will be right-handed by swapping x and y. Non-zero also forces rotation rate units to be radians/s and acceleration units to be g."));
 
 uint32 TranslateGCKeyCodeToASCII(GCKeyCode KeyCode)
 {
