iOS: Add casting for MTL::Device to fix IOS Framework build erorrs.

diff --git a/Engine/Source/Runtime/ApplicationCore/Private/IOS/IOSView.cpp b/Engine/Source/Runtime/ApplicationCore/Private/IOS/IOSView.cpp
--- a/Engine/Source/Runtime/ApplicationCore/Private/IOS/IOSView.cpp
+++ b/Engine/Source/Runtime/ApplicationCore/Private/IOS/IOSView.cpp
@@ -118,7 +118,7 @@ MTL::Device* GMetalDevice = nullptr;
 {
 #if BUILD_EMBEDDED_APP
 	SCOPED_BOOT_TIMING("MetalLayer class");
-	GMetalDevice = MTLCreateSystemDefaultDevice();
+	GMetalDevice = (__bridge MTL::Device*)MTLCreateSystemDefaultDevice();
 	return [CAMetalLayer class];
 #endif
 	
