Fixed special right input being registered twice on Apple platforms

#tests On iOS, ensured that the Special Right button works properly and does not duplicate input

diff --git a/Engine/Source/Runtime/ApplicationCore/Private/Apple/AppleControllerInterface.cpp b/Engine/Source/Runtime/ApplicationCore/Private/Apple/AppleControllerInterface.cpp
--- a/Engine/Source/Runtime/ApplicationCore/Private/Apple/AppleControllerInterface.cpp
+++ b/Engine/Source/Runtime/ApplicationCore/Private/Apple/AppleControllerInterface.cpp
@@ -187,15 +187,6 @@ void FAppleControllerInterface::HandleConnection(GCController* Controller)
         Controllers[ControllerIndex].Controller = [Controller retain];
         SetControllerType(ControllerIndex);
         
-        // Deprecated but buttonMenu behavior is unreliable since iOS/tvOS 14
-		Controllers[ControllerIndex].bPauseWasPressed = false;
-PRAGMA_DISABLE_DEPRECATION_WARNINGS
-		Controller.controllerPausedHandler = ^(GCController* Cont)
-		{
-			Controllers[ControllerIndex].bPauseWasPressed = true;
-		};
-PRAGMA_ENABLE_DEPRECATION_WARNINGS
-        
         bFoundSlot = true;
         
         IPlatformInputDeviceMapper& DeviceMapper = IPlatformInputDeviceMapper::Get();
@@ -269,15 +260,6 @@ void FAppleControllerInterface::SendControllerEvents()
 		
         GCExtendedGamepad* ExtendedGamepad = [ControllerImpl capture].extendedGamepad;
 		GCMotion* Motion = ControllerImpl.motion;
-		
-		// Workaround for unreliable buttonMenu behavior since iOS/tvOS 14
-		if (Controller.bPauseWasPressed)
-        {
-            MessageHandler->OnControllerButtonPressed(FGamepadKeyNames::SpecialRight, UserId, DeviceId, false);
-            MessageHandler->OnControllerButtonReleased(FGamepadKeyNames::SpecialRight, UserId, DeviceId, false);
-
-            Controller.bPauseWasPressed = false;
-        }
         
 		if (ExtendedGamepad != nil)
 		{
diff --git a/Engine/Source/Runtime/ApplicationCore/Private/Apple/AppleControllerInterface.h b/Engine/Source/Runtime/ApplicationCore/Private/Apple/AppleControllerInterface.h
--- a/Engine/Source/Runtime/ApplicationCore/Private/Apple/AppleControllerInterface.h
+++ b/Engine/Source/Runtime/ApplicationCore/Private/Apple/AppleControllerInterface.h
@@ -137,9 +137,6 @@ protected:
 		bool bNeedsReferenceAttitude;
 		bool bHasReferenceAttitude;
 #endif
-
-		// Workaround for unreliable buttonMenu behavior since iOS/tvOS 14
-        bool bPauseWasPressed;
 	};
 	
 	// Controller Event Callbacks are on the main thread - defer to tick processing
