Refactored iOS background downloading stack

diff --git a/Engine/Config/IOS/BaseIOSEngine.ini b/Engine/Config/IOS/BaseIOSEngine.ini
--- a/Engine/Config/IOS/BaseIOSEngine.ini
+++ b/Engine/Config/IOS/BaseIOSEngine.ini
@@ -39,14 +39,18 @@ PlatformFormat=ADPCM
 PlatformStreamingFormat=ADPCM
 
 [BackgroundHttp.iOSSettings]
-;While the app is in the forground, time out if we don't get a response in this timeframe
-ActiveReceiveTimeout=30
-;While the app is in the background, time out if we don't get a response in this timeframe
+; Use foreground-only NSURLSession 
+bUseForegroundSession=false
+; Requests will be scheduled on discretion of iOS if true, otherwise they will be scheduled ASAP if false
+bDiscretionary=false
+; Wakeup app in background and report state of current downloads, allows to retry downloads in background
+bShouldSendLaunchEvents=true
+; Time out and retry if we don't get a response in this timeframe if app is in foreground and bUseForegroundSession is true
 BackgroundReceiveTimeout=120
-;Regardless of app state, if the download hasn't finished in this long, time out
+; Time out and retry if the download hasn't finished in this long, regardless of app state 
 BackgroundHttpResourceTimeout=3600
-;If retry data exists for a task, how many times we should try using the Retry data before moving onto a different CDN. ( -1 = unlimited )
-RetryResumeDataLimit=3;
+; How many times we should try using same CDN before moving onto a different CDN ( -1 = unlimited, 0 = only try first CDN, 1 = try each CDN once )
+RetryResumeDataLimit=3
 
 [ConsoleVariables]
 Slate.CacheRenderData=0
diff --git a/Engine/Plugins/Experimental/DefaultInstallBundleManager/Source/Private/DefaultInstallBundleManager.cpp b/Engine/Plugins/Experimental/DefaultInstallBundleManager/Source/Private/DefaultInstallBundleManager.cpp
--- a/Engine/Plugins/Experimental/DefaultInstallBundleManager/Source/Private/DefaultInstallBundleManager.cpp
+++ b/Engine/Plugins/Experimental/DefaultInstallBundleManager/Source/Private/DefaultInstallBundleManager.cpp
@@ -4088,6 +4088,14 @@ void FDefaultInstallBundleManager::UpdateContentRequestFlags(TArrayView<const FN
 	}
 }
 
+void FDefaultInstallBundleManager::SetCellularPreference(int32 Value)
+{
+	for (const TPair<FInstallBundleSourceType, TSharedPtr<IInstallBundleSource>>& Pair : BundleSources)
+	{
+		Pair.Value->SetCellularPreference(Value);
+	}
+}
+
 void FDefaultInstallBundleManager::SetCacheSize(FName CacheName, uint64 CacheSize)
 {
 	if (ensureMsgf(!BundleCaches.Contains(CacheName), TEXT("FDefaultInstallBundleManager::SetCacheSize is only supported prior to initialization, for now")))
diff --git a/Engine/Plugins/Experimental/DefaultInstallBundleManager/Source/Public/DefaultInstallBundleManager.h b/Engine/Plugins/Experimental/DefaultInstallBundleManager/Source/Public/DefaultInstallBundleManager.h
--- a/Engine/Plugins/Experimental/DefaultInstallBundleManager/Source/Public/DefaultInstallBundleManager.h
+++ b/Engine/Plugins/Experimental/DefaultInstallBundleManager/Source/Public/DefaultInstallBundleManager.h
@@ -661,6 +661,7 @@ public:
 
 	virtual EInstallBundleRequestFlags GetModifyableContentRequestFlags() const override;
 	virtual void UpdateContentRequestFlags(TArrayView<const FName> BundleNames, EInstallBundleRequestFlags AddFlags, EInstallBundleRequestFlags RemoveFlags) override;
+	virtual void SetCellularPreference(int32 Value) override;
 
 	virtual void SetCacheSize(FName CacheName, uint64 CacheSize) override;
 
diff --git a/Engine/Source/Runtime/ApplicationCore/Private/IOS/IOSAppDelegate.cpp b/Engine/Source/Runtime/ApplicationCore/Private/IOS/IOSAppDelegate.cpp
--- a/Engine/Source/Runtime/ApplicationCore/Private/IOS/IOSAppDelegate.cpp
+++ b/Engine/Source/Runtime/ApplicationCore/Private/IOS/IOSAppDelegate.cpp
@@ -1619,9 +1619,9 @@ extern double GCStartTime;
     //Save off completionHandler so that a future call to FCoreDelegates::ApplicationBackgroundSessionEventsAllSentDelegate can execute it
     self.BackgroundSessionEventCompleteDelegate = completionHandler;
     
-    //Create background session with this identifier if needed to handle these events
+    // Invoke background session with this identifier if needed to handle these events
 	FString Id(identifier);
-	FBackgroundURLSessionHandler::InitBackgroundSession(Id);
+	FBackgroundURLSessionHandler::HandleEventsForBackgroundURLSession(Id);
 
 	FCoreDelegates::ApplicationBackgroundSessionEventDelegate.Broadcast(Id);
 }
diff --git a/Engine/Source/Runtime/ApplicationCore/Private/IOS/IOSBackgroundURLSessionHandler.cpp b/Engine/Source/Runtime/ApplicationCore/Private/IOS/IOSBackgroundURLSessionHandler.cpp
--- a/Engine/Source/Runtime/ApplicationCore/Private/IOS/IOSBackgroundURLSessionHandler.cpp
+++ b/Engine/Source/Runtime/ApplicationCore/Private/IOS/IOSBackgroundURLSessionHandler.cpp
@@ -1,5 +1,7 @@
 // Copyright Epic Games, Inc. All Rights Reserved.
 
+#ifndef UE_DNLD_SANDBOX
+
 #include "IOS/IOSBackgroundURLSessionHandler.h"
 
 #include "Misc/ConfigCacheIni.h"
@@ -10,264 +12,990 @@
 
 #include "IOS/IOSAppDelegate.h"
 
-FIOSBackgroundDownloadCoreDelegates::FIOSBackgroundDownload_DidFinishDownloadingToURL FIOSBackgroundDownloadCoreDelegates::OnIOSBackgroundDownload_DidFinishDownloadingToURL;
-FIOSBackgroundDownloadCoreDelegates::FIOSBackgroundDownload_DidWriteData FIOSBackgroundDownloadCoreDelegates::OnIOSBackgroundDownload_DidWriteData;
-FIOSBackgroundDownloadCoreDelegates::FIOSBackgroundDownload_DidCompleteWithError FIOSBackgroundDownloadCoreDelegates::OnIOSBackgroundDownload_DidCompleteWithError;
-FIOSBackgroundDownloadCoreDelegates::FIOSBackgroundDownload_SessionDidFinishAllEvents FIOSBackgroundDownloadCoreDelegates::OnIOSBackgroundDownload_SessionDidFinishAllEvents;
-FIOSBackgroundDownloadCoreDelegates::FIOSBackgroundDownload_DelayedBackgroundURLSessionCompleteHandler FIOSBackgroundDownloadCoreDelegates::OnDelayedBackgroundURLSessionCompleteHandler;
+#endif
+
+#include <future>
+
+// Force cancel all pending downloads.
+// Useful when testing background downloads as they persist between application sessions.
+static constexpr bool bCancelExistingDownloads = false;
+
+#define UE_DNLD_LOG(...) NSLog(@"UEBackgroundDownload " __VA_ARGS__)
+
+// --------------------------------------------------------------------------------------------------------------------
+
+// We need additional state in NSURLSessionDownloadTask to implement CDN failover.
+// Convential way to implement this would be to set a delegate on NSURLSessionDownloadTask and add properties to the delegate.
+//
+// In case of background downloads API prohibits setting a delegate or setting NSObject values.
+// Additionally we want to keep state in NSURLSessionDownloadTask itself to avoid pitfalls of maintaining separate state elsewhere.
+//
+// [NSURLSessionDownloadTask taskDescription] is one such property that can be used to track state and API maintains state between app sessions.
+@interface FBackgroundNSURLSessionDownloadTaskData : NSObject
+
+@property (nonatomic, retain) NSMutableArray<__kindof NSURL*>* URLs;
+@property (nonatomic, retain) NSMutableArray<__kindof NSNumber*>* RetryCountPerURL;
+
+// assumes all URL's have same content path but different domain
++ (FBackgroundNSURLSessionDownloadTaskData* _Nonnull)TaskDataWithURLs:(NSArray<__kindof NSURL*>*)URLs WithRetryCount:(NSInteger)RetryCount;
++ (FBackgroundNSURLSessionDownloadTaskData* _Nullable)TaskDataFromSerializedString:(NSString*)SerializedData;
+
+- (NSString*)ToSerializedString;
+
+- (NSURL* _Nonnull)GetFirstURL;
+- (NSURL* _Nullable)GetNextURL;
+- (void)ResetRetryCount:(NSInteger)RetryCount;
+- (void)Cancel;
+
+@end
+
+// --------------------------------------------------------------------------------------------------------------------
+
+@implementation FBackgroundNSURLSessionDownloadTaskData
+
+NSString* const SerializationKeyProtocolVersion = @"v";
+NSString* const SerializationKeyCDNs = @"c";
+NSString* const SerializationKeyPath = @"p";
+NSString* const SerializationKeyRetryCountPerURL = @"r";
+
++ (FBackgroundNSURLSessionDownloadTaskData* _Nonnull)TaskDataWithURLs:(NSArray<__kindof NSURL*>*)URLs WithRetryCount:(NSInteger)RetryCount
+{
+	FBackgroundNSURLSessionDownloadTaskData* Data = [[FBackgroundNSURLSessionDownloadTaskData alloc] init];
+
+	Data.URLs = [NSMutableArray arrayWithArray:URLs];
+	Data.RetryCountPerURL = [NSMutableArray arrayWithCapacity:URLs.count];
+	for (NSURL* URL in URLs)
+	{
+		(void)URL;
+		[Data.RetryCountPerURL addObject:[NSNumber numberWithInteger:RetryCount]];
+	}
+	NSAssert(Data.URLs.count > 0, @"URLs should be non empty");
+	NSAssert(Data.URLs.count == Data.RetryCountPerURL.count, @"URLs and RetryCountPerURL arrays should have same size");
+
+	return [Data autorelease];
+}
+
++ (FBackgroundNSURLSessionDownloadTaskData* _Nullable)TaskDataFromSerializedString:(NSString*)SerializedData
+{
+	FBackgroundNSURLSessionDownloadTaskData* Data = [[FBackgroundNSURLSessionDownloadTaskData alloc] init];
+
+	NSError* Error = nil;
+	NSDictionary* Dict = [NSJSONSerialization JSONObjectWithData:[SerializedData dataUsingEncoding:NSUTF8StringEncoding] options:0 error:&Error];
+
+	NSNumber* Version = [Dict valueForKey:SerializationKeyProtocolVersion];
+	NSMutableArray<__kindof NSString*>* CDNs = [Dict valueForKey:SerializationKeyCDNs];
+	NSString* Path = [Dict valueForKey:SerializationKeyPath];
+	NSArray<__kindof NSNumber*>* RetryCountPerURL = [Dict valueForKey:SerializationKeyRetryCountPerURL];
+
+	if (Error != nil || Version == nil || Version.intValue != 1 || CDNs == nil || Path == nil || RetryCountPerURL == nil)
+	{
+		UE_DNLD_LOG(@"Failed to deserialize task state '%@' due to '%@', %u, %u, %u, %u, %u",
+			SerializedData,
+			Error != nil ? Error.localizedDescription : @"",
+			Version == nil ? 1 : 0,
+			Version.intValue,
+			CDNs == nil ? 1 : 0,
+			Path == nil ? 1 : 0,
+			RetryCountPerURL == nil ? 1 : 0
+		);
+		[Data release];
+		return nil;
+	}
+
+	Data.URLs = [NSMutableArray arrayWithCapacity:CDNs.count];
+	Data.RetryCountPerURL = [NSMutableArray arrayWithArray:RetryCountPerURL];
+	for (NSString* CDN in CDNs)
+	{
+		[Data.URLs addObject:[NSURL URLWithString:[NSString stringWithFormat:@"%@%@", CDN, Path]]];
+	}
+	NSAssert(Data.URLs.count > 0, @"URLs should be non empty");
+	NSAssert(Data.URLs.count == Data.RetryCountPerURL.count, @"URLs and RetryCountPerURL arrays should have same size");
+
+	return [Data autorelease];
+}
+
+- (void)dealloc
+{
+	[_URLs release];
+	[_RetryCountPerURL release];
+	[super dealloc];
+}
+
+- (NSString *)ToSerializedString
+{
+	NSMutableDictionary* Dict = [NSMutableDictionary dictionaryWithCapacity:4];
+
+	NSString* Path = [[self.URLs firstObject] path];
+
+	NSMutableArray<__kindof NSString*>* CDNs = [NSMutableArray arrayWithCapacity:self.URLs.count];
+	for (NSURL* URL in self.URLs)
+	{
+		NSString* URLString = URL.absoluteString;
+		if (![URLString hasSuffix:Path])
+		{
+			UE_DNLD_LOG(@"Expected all URLs have same path but got '%@' when expected path '%@'", URLString, Path);
+			continue;
+		}
+		NSString* Domain = [[URLString componentsSeparatedByString:URL.path] firstObject];
+		[CDNs addObject:Domain];
+	}
+
+	[Dict setValue:[NSNumber numberWithInt:1] forKey:SerializationKeyProtocolVersion];
+	[Dict setValue:CDNs forKey:SerializationKeyCDNs];
+	[Dict setValue:Path forKey:SerializationKeyPath];
+	[Dict setValue:self.RetryCountPerURL forKey:SerializationKeyRetryCountPerURL];
+
+	NSError* Error = nil;
+	NSData* Data = [NSJSONSerialization dataWithJSONObject:Dict options:NSJSONWritingSortedKeys error:&Error];
+
+	NSString* String = [[NSString alloc] initWithData:Data encoding:NSUTF8StringEncoding];
+	return [String autorelease];
+}
+
+- (NSURL* _Nonnull)GetFirstURL
+{
+	return self.URLs.firstObject;
+}
+
+- (NSURL* _Nullable)GetNextURL
+{
+	for (NSUInteger i = 0; i < self.RetryCountPerURL.count; ++i)
+	{
+		NSUInteger RetryValue = [self.RetryCountPerURL objectAtIndex:i].integerValue;
+		if (RetryValue == 0)
+		{
+			continue;
+		}
+		else if (RetryValue <= -1) // special case for infinitely retrying same URL
+		{
+			return [self.URLs objectAtIndex:i];
+		}
+
+		RetryValue--;
+		[self.RetryCountPerURL replaceObjectAtIndex:i withObject:[NSNumber numberWithInteger:RetryValue]];
+		return [self.URLs objectAtIndex:i];
+	}
+
+	return nil;
+}
+
+- (void)ResetRetryCount:(NSInteger)RetryCount
+{
+	for (NSUInteger i = 0; i < self.RetryCountPerURL.count; ++i)
+	{
+		[self.RetryCountPerURL replaceObjectAtIndex:i withObject:[NSNumber numberWithInteger:RetryCount]];
+	}
+}
+
+- (void)Cancel
+{
+	[self ResetRetryCount:0];
+}
+
+@end
+
+// --------------------------------------------------------------------------------------------------------------------
+
+// NSURLSession wrapper focused on background downloading and CDN failover
+@interface FBackgroundNSURLSession : NSObject<NSURLSessionDelegate, NSURLSessionTaskDelegate, NSURLSessionDownloadDelegate>
+
+@property (atomic) BOOL AllowCellular;
+
++ (FBackgroundNSURLSession*)Shared;
++ (NSUInteger)GetInvalidDownloadId;
++ (NSString*)GetNSURLSessionIdentifier;
+
+- (void)Initialize;
+- (void)SetFileHashHelper:(BackgroundHttpFileHashHelperRef)HelperRef;
+- (BackgroundHttpFileHashHelperRef)GetFileHashHelper;
+- (NSString*)GetTempPathForURL:(NSURL* _Nonnull)URL;
+- (NSURLSessionDownloadTask*)CreateDownloadForURL:(NSURL* _Nonnull)URL WithPriority:(float)Priority WithTaskData:(FBackgroundNSURLSessionDownloadTaskData* _Nonnull)TaskData;
+- (NSURLSessionDownloadTask*)CreateDownloadForResumeData:(NSData* _Nonnull)ResumeData WithPriority:(float)Priority WithTaskData:(FBackgroundNSURLSessionDownloadTaskData* _Nonnull)TaskData;
+
+- (NSUInteger)CreateOrFindDownloadForURLs:(NSArray<__kindof NSString*>*)URLStrings WithPriority:(float)Priority;
+- (void)PauseDownload:(NSUInteger)DownloadId;
+- (void)ResumeDownload:(NSUInteger)DownloadId;
+- (void)CancelDownload:(NSUInteger)DownloadId;
+- (void)SetPriority:(float)Priority ForDownload:(NSUInteger)DownloadId;
+- (void)SetCurrentDownloadedBytes:(uint64)DownloadedBytes ForTask:(NSURLSessionDownloadTask*)Task;
+- (uint64)GetCurrentDownloadedBytes:(NSUInteger)DownloadId;
+- (void)RecreateDownloads;
+
+- (NSURLSessionDownloadTask*)FindDownloadTaskFor:(NSUInteger)DownloadId;
+- (NSUInteger)FindDownloadIdForTask:(NSURLSessionDownloadTask*)Task;
+- (NSUInteger)EnsureTaskIsTracked:(NSURLSessionDownloadTask*)Task;
+- (void)ReplaceTrackedTaskWith:(NSURLSessionDownloadTask*)NewTask ForDownloadId:(NSUInteger)DownloadId;
+- (void)EnsureTaskIsNotTracked:(NSURLSessionDownloadTask*)Task;
+
+- (void)SetDownloadResult:(NSInteger)HTTPCode WithTempFile:(NSString*)TempFile ForDownload:(NSURLSessionDownloadTask*)Task;
+- (NSString* _Nullable)GetDownloadResult:(NSUInteger)DownloadId OutStatus:(BOOL* _Nonnull)OutStatus OutStatusCode:(NSInteger*)OutStatusCode;
+
+// From NSURLSessionDelegate
+- (void)URLSessionDidFinishEventsForBackgroundURLSession:(NSURLSession*)Session;
+
+// From NSURLSessionTaskDelegate
+- (void)URLSession:(NSURLSession*)Session task:(NSURLSessionTask*)Task didCompleteWithError:(NSError*)Error;
+//- (void)URLSession:(NSURLSession *)session taskIsWaitingForConnectivity:(NSURLSessionTask *)Task;
+//- (void)URLSession:(NSURLSession *)Session task:(NSURLSessionTask *)Task willBeginDelayedRequest:(NSURLRequest *)Request completionHandler:(void (^)(NSURLSessionDelayedRequestDisposition Disposition, NSURLRequest* NewRequest))CompletionHandler;
+
+// From NSURLSessionDownloadDelegate
+- (void)URLSession:(NSURLSession*)Session downloadTask:(NSURLSessionDownloadTask*)Task didFinishDownloadingToURL:(NSURL*)Location;
+- (void)URLSession:(NSURLSession*)session downloadTask:(NSURLSessionDownloadTask*)Task didWriteData:(int64_t)BytesWritten totalBytesWritten:(int64_t)TotalBytesWritten totalBytesExpectedToWrite:(int64_t)TotalBytesExpectedToWrite;
+
+@end
+
+// --------------------------------------------------------------------------------------------------------------------
+
+@implementation FBackgroundNSURLSession
+{
+	NSURLSession* _Session;
+	NSMutableDictionary<__kindof NSNumber*, __kindof NSURLSessionDownloadTask*>* _AllDownloads;
+	NSUInteger _NextDownloadId;
+	std::promise<void> _AllDownloadsPromise;
+	std::future<void> _AllDownloadsFuture;
+	BackgroundHttpFileHashHelperPtr _HelperPtr;
+	int32 RetryResumeDataLimit;
+}
+
+static constexpr NSUInteger InvalidDownloadId = 0;
+
+NSString* const NSURLSessionIdentifier = @"com.epicgames.backgrounddownloads";
+
+NSProgressUserInfoKey const NSProgressDownloadCompletedBytes = @"com.epicgames.nsprogress.completedbytes";
+NSProgressUserInfoKey const NSProgressDownloadResultStatusCode = @"com.epicgames.nsprogress.resultstatuscode";
+NSProgressUserInfoKey const NSProgressDownloadResultTempFilePath = @"com.epicgames.nsprogress.tempfilepath";
 
-NSURLSession* FBackgroundURLSessionHandler::BackgroundSession = nil;
-FString FBackgroundURLSessionHandler::CachedIdentifierName = FString();
-volatile int32 FBackgroundURLSessionHandler::DelayedBackgroundURLSessionCompleteCount = 0;
+static constexpr NSInteger HTTPStatusCodeSuccessCreated = 201;
+static constexpr NSInteger HTTPStatusCodeErrorBadRequest = 500;
+static constexpr NSInteger HTTPStatusCodeErrorServer = 500;
 
-BackgroundHttpFileHashHelperRef FBackgroundURLSessionHandler::FileHashHelper = MakeShared<FBackgroundHttpFileHashHelper, ESPMode::ThreadSafe>();
++ (FBackgroundNSURLSession*)Shared
+{
+	static FBackgroundNSURLSession* Shared = nil;
+	static dispatch_once_t Once;
+	dispatch_once(&Once, ^{
+		Shared = [[self alloc] init];
+	});
+	return Shared;
+}
+
++ (NSUInteger)GetInvalidDownloadId;
+{
+	return InvalidDownloadId;
+}
+
++ (NSString*)GetNSURLSessionIdentifier
+{
+	return NSURLSessionIdentifier;
+}
+
+- (id)init
+{
+	if (self = [super init])
+	{
+		[self Initialize];
+	}
+	return self;
+}
 
-bool FBackgroundURLSessionHandler::InitBackgroundSession(const FString& SessionIdentifier)
+- (void)Initialize
 {
-	FIOSBackgroundDownloadCoreDelegates::OnDelayedBackgroundURLSessionCompleteHandler.BindStatic(&FBackgroundURLSessionHandler::OnDelayedBackgroundURLSessionCompleteHandlerCalled);
+	bool bUseForegroundSession = false;
+	bool bDiscretionary = false;
+	bool bShouldSendLaunchEvents = true;
+	int32 MaximumConnectionsPerHost = 6;
+	double TimeoutIntervalForRequest = 120.0; // Note, ignored in background sessions (if bUseForegroundSession is false).
+	double TimeoutIntervalForResource = 60.0 * 60.0;
+	RetryResumeDataLimit = 3;
+
+#ifndef UE_DNLD_SANDBOX
+	GConfig->GetBool(TEXT("BackgroundHttp.iOSSettings"), TEXT("bUseForegroundSession"), bUseForegroundSession, GEngineIni);
+	GConfig->GetBool(TEXT("BackgroundHttp.iOSSettings"), TEXT("bDiscretionary"), bDiscretionary, GEngineIni);
+	GConfig->GetBool(TEXT("BackgroundHttp.iOSSettings"), TEXT("bShouldSendLaunchEvents"), bShouldSendLaunchEvents, GEngineIni);
+	GConfig->GetInt(TEXT("BackgroundHttp"), TEXT("MaxActiveDownloads"), MaximumConnectionsPerHost, GEngineIni);
+	GConfig->GetDouble(TEXT("BackgroundHttp.iOSSettings"), TEXT("BackgroundReceiveTimeout"), TimeoutIntervalForRequest, GEngineIni);
+	GConfig->GetDouble(TEXT("BackgroundHttp.iOSSettings"), TEXT("BackgroundHttpResourceTimeout"), TimeoutIntervalForResource, GEngineIni);
+	GConfig->GetInt(TEXT("BackgroundHttp.iOSSettings"), TEXT("RetryResumeDataLimit"), RetryResumeDataLimit, GEngineIni);
+#endif
 	
-	//Make sure our BackgroundHttpFileHashHelper is initialized and loaded to handle any completed requests asking for Temp file paths
-	GetFileHashHelper()->LoadData();
+	_AllDownloads = [NSMutableDictionary new];
+	_AllDownloadsFuture = _AllDownloadsPromise.get_future();
+	_NextDownloadId = InvalidDownloadId + 1;
+
+	// Never allow cellular unless we get explicit opt-in from the user.
+	self.AllowCellular = NO;
+
+	NSURLSessionConfiguration* Configuration = bUseForegroundSession ?
+		[NSURLSessionConfiguration defaultSessionConfiguration] :
+		[NSURLSessionConfiguration backgroundSessionConfigurationWithIdentifier: NSURLSessionIdentifier];
+
+	// iOS will schedule downloads on it's own if true,
+	// otherwise all downloads will be scheduled ASAP if false
+	Configuration.discretionary = bDiscretionary;
+
+	// In case if our app gets killed in background, iOS will launch it and report finished downloads via handleEventsForBackgroundURLSession.
+	// This will help us to retry/fail-over downloads in background without waiting for user to open the game again.
+	// Note that this behavior can be disabled via Background App Refresh set to No in iOS settings.
+	Configuration.sessionSendsLaunchEvents = bShouldSendLaunchEvents;
+
+	// Set session to allow cellular and instead control this on NSMutableURLRequest level because this value cannot be changed after session is created.
+	Configuration.allowsCellularAccess = YES;
+
+	Configuration.networkServiceType = bUseForegroundSession ? NSURLNetworkServiceTypeDefault : NSURLNetworkServiceTypeBackground;
+
+	// TODO Is this any use for us? Needs entitlement.
+	//Configuration.multipathServiceType = NSURLSessionMultipathServiceTypeAggregate;
+
+	Configuration.HTTPMaximumConnectionsPerHost = MaximumConnectionsPerHost;
+
+	Configuration.timeoutIntervalForRequest = TimeoutIntervalForRequest;
+
+	Configuration.timeoutIntervalForResource = TimeoutIntervalForResource;
 
-	bool bHasSessionWithIdentifier = false;
+	_Session = [[NSURLSession sessionWithConfiguration:Configuration delegate:self delegateQueue:nil] retain];
+	UE_DNLD_LOG(@"sessionWithConfiguration '%@'", Configuration.identifier);
 
-	if (!SessionIdentifier.IsEmpty())
+	[_Session getTasksWithCompletionHandler:^(NSArray<__kindof NSURLSessionDataTask*>*, NSArray<__kindof NSURLSessionUploadTask*>*, NSArray<__kindof NSURLSessionDownloadTask*>* Downloads)
 	{
-		//if we don't already have a background session, create it
-		if (nil == BackgroundSession)
+		UE_DNLD_LOG(@"getTasksWithCompletionHandler block with %lu tasks", (unsigned long)[Downloads count]);
+		
+		if (Downloads == nil)
 		{
-			//Create actual session
+			for (NSURLSessionDownloadTask* ExistingTask in Downloads)
 			{
-				float BackgroundHttpTimeout = 90.0f;
-				GConfig->GetFloat(TEXT("BackgroundHttp.iOSSettings"), TEXT("BackgroundReceiveTimeout"), BackgroundHttpTimeout, GEngineIni);
-
-                float BackgroundHttpResourceTimeout = 3600.f;
-                GConfig->GetFloat(TEXT("BackgroundHttp.iOSSettings"), TEXT("BackgroundHttpResourceTimeout"), BackgroundHttpResourceTimeout, GEngineIni);
-                
-				NSURLSessionConfiguration *config = [NSURLSessionConfiguration backgroundSessionConfigurationWithIdentifier : SessionIdentifier.GetNSString()];
-				config.discretionary = NO;
-				config.sessionSendsLaunchEvents = YES;
-				config.shouldUseExtendedBackgroundIdleMode = YES;
-				config.timeoutIntervalForRequest = BackgroundHttpTimeout;	        //Timeout for requests not getting a response from server at the iOS level. Will still work in the BG.
-				config.timeoutIntervalForResource = BackgroundHttpResourceTimeout;  //Will automatically retry everytime above setting causes a timeout until this many seconds pass. Then we will get
-                                                                                    //an error result. No Error for server not responding until is sent by a background task until this timeout is reached.
-				config.allowsCellularAccess = NO; // don't allow cellular by default, TODO integrate opt-in dialog
-				//Other unused settings we might want to configure
-				//HTTPMaximumConnectionsPerHost
-                
-                //Can't use on BG session, but may need to consider for any future FG session work:
-				//waitsForConnectivity  -- This is forced true on background sessions
-				//URLSession:taskIsWaitingForConnectivity: method of NSURLSessionTaskDelegate -- Can't respond to this with background sessions since the above is always forced true they never call this delegate
-
-                //NSURLSessionConfiguration* config = [NSURLSessionConfiguration defaultSessionConfiguration];
-                BackgroundSession = [NSURLSession sessionWithConfiguration:config delegate:[[BackgroundDownloadDelegate alloc]init] delegateQueue:nil];
-				[BackgroundSession retain];
-			}
+				const bool bCanRestartTask =
+					(ExistingTask.state == NSURLSessionTaskStateRunning) ||
+					(ExistingTask.state == NSURLSessionTaskStateSuspended);
 
-			//test if session created
-			if (nil != BackgroundSession)
-			{
-				CachedIdentifierName = SessionIdentifier;
-				bHasSessionWithIdentifier = true;
+				if (!bCanRestartTask)
+				{
+					UE_DNLD_LOG(@"Skipping tracking for existing download task with taskIdentifier %lu because it's not in resumable state", ExistingTask.taskIdentifier);
+					continue;
+				}
+
+				const NSUInteger DownloadId = [self EnsureTaskIsTracked:ExistingTask];
 
-				//Make sure our temp directory exists where we will be putting files as they complete as we might get completion callbacks right after we init this session
-				CreateBackgroundSessionWorkingDirectory();
+				if (bCancelExistingDownloads)
+				{
+					UE_DNLD_LOG(@"Canceling existing download task with taskIdentifier %lu", ExistingTask.taskIdentifier);
+					[self CancelDownload:DownloadId];
+				}
 			}
 		}
-		//we already have one, so just see if names match
-		else
+		
+
+		_AllDownloadsPromise.set_value();
+	}];
+
+	const FString& DirectoryPath = FBackgroundHttpFileHashHelper::GetTemporaryRootPath();
+	if (ensureAlwaysMsgf(!DirectoryPath.IsEmpty(), TEXT("Invalid FBackgroundHttpFileHashHelper::GetTemporaryRootPath()")))
+	{
+		IPlatformFile& PlatformFile = FPlatformFileManager::Get().GetPlatformFile();
+
+		if (bCancelExistingDownloads)
+		{
+			PlatformFile.DeleteDirectory(*DirectoryPath);
+		}
+
+		PlatformFile.CreateDirectory(*DirectoryPath);
+
+		if (!PlatformFile.DirectoryExists(*DirectoryPath))
 		{
-			bHasSessionWithIdentifier = CachedIdentifierName.Equals(SessionIdentifier);
+			ensureAlwaysMsgf(false, TEXT("Failed to create temporary directory for background downloads"));
 		}
 	}
+}
+
+- (void)dealloc
+{
+	[_Session release];
+	[_AllDownloads release];
 
-	return ensureMsgf(bHasSessionWithIdentifier, TEXT("Unable to initialize background session with identifier %s"), *SessionIdentifier);
+	[super dealloc];
 }
 
-void FBackgroundURLSessionHandler::ShutdownBackgroundSession(bool bShouldFinishTasksFirst /* = true */)
+- (void)SetFileHashHelper:(BackgroundHttpFileHashHelperRef)HelperRef
 {
-	if (nil != BackgroundSession)
+	_HelperPtr = HelperRef.ToSharedPtr();
+}
+
+- (BackgroundHttpFileHashHelperRef)GetFileHashHelper
+{
+	// initialize a new instance in case if we get here from handleEventsForBackgroundURLSession
+	if (!_HelperPtr.IsValid())
+	{
+		_HelperPtr = MakeShared<FBackgroundHttpFileHashHelper, ESPMode::ThreadSafe>();
+		_HelperPtr->LoadData();
+	}
+
+	return _HelperPtr.ToSharedRef();
+}
+
+- (NSString*)GetTempPathForURL:(NSURL* _Nonnull)URL
+{
+	BackgroundHttpFileHashHelperRef HelperRef = [self GetFileHashHelper];
+
+	const FString TaskURL(URL.absoluteString);
+	const FString& TempFileName = HelperRef->FindOrAddTempFilenameMappingForURL(TaskURL);
+	const FString& DestinationPath = HelperRef->GetFullPathOfTempFilename(TempFileName);
+
+	IPlatformFile& PlatformFile = FPlatformFileManager::Get().GetPlatformFile();
+	const FString ConvertedPath = PlatformFile.ConvertToAbsolutePathForExternalAppForWrite(*DestinationPath);
+
+	return ConvertedPath.GetNSString();
+}
+
+- (NSURLSessionDownloadTask*)CreateDownloadForURL:(NSURL* _Nonnull)URL WithPriority:(float)Priority WithTaskData:(FBackgroundNSURLSessionDownloadTaskData* _Nonnull)TaskData;
+{
+	NSMutableURLRequest* URLRequest = [NSMutableURLRequest requestWithURL:URL];
+	URLRequest.allowsCellularAccess = self.AllowCellular;
+
+	NSURLSessionDownloadTask* Task = [_Session downloadTaskWithRequest:URLRequest];
+	[Task setPriority:Priority];
+	[Task setTaskDescription:[TaskData ToSerializedString]];
+
+	UE_DNLD_LOG(@"CreateDownloadForURL '%@' with taskIdentifier %lu", Task.taskDescription, Task.taskIdentifier);
+	return Task;
+}
+
+- (NSURLSessionDownloadTask*)CreateDownloadForResumeData:(NSData* _Nonnull)ResumeData WithPriority:(float)Priority WithTaskData:(FBackgroundNSURLSessionDownloadTaskData* _Nonnull)TaskData
+{
+	NSURLSessionDownloadTask* Task = [_Session downloadTaskWithResumeData:ResumeData];
+	[Task setPriority:Priority];
+	[Task setTaskDescription:[TaskData ToSerializedString]];
+
+	UE_DNLD_LOG(@"CreateDownloadForResumeData '%@' with taskIdentifier %lu", Task.taskDescription, Task.taskIdentifier);
+	return Task;
+}
+
+- (NSUInteger)CreateOrFindDownloadForURLs:(NSArray<__kindof NSString*>*)URLStrings WithPriority:(float)Priority
+{
+	if (_AllDownloadsFuture.wait_for(std::chrono::seconds(0)) != std::future_status::ready)
+	{
+		UE_DNLD_LOG(@"Starting wait for existing downloads status");
+		
+		_AllDownloadsFuture.wait();
+
+		UE_DNLD_LOG(@"Done waiting for existing downloads status");
+	}
+
+	// To be able to store less state we assume all strings have same asset path suffix
+	NSString* AssetPath = nil;
+	NSMutableArray<__kindof NSURL*>* URLs = [NSMutableArray arrayWithCapacity:URLStrings.count];
+	for (NSString* URLString in URLStrings)
 	{
-		if (bShouldFinishTasksFirst)
+		NSURL* URLValue = [NSURL URLWithString:URLString];
+		if (AssetPath == nil)
 		{
-			[BackgroundSession finishTasksAndInvalidate];
+			AssetPath = URLValue.path;
 		}
-		else
+		if (![URLValue.absoluteString hasSuffix:AssetPath])
 		{
-			[BackgroundSession invalidateAndCancel];
+			UE_DNLD_LOG(@"Expected all URLs have same path but got '%@' when expected path '%@'", URLString, AssetPath);
 		}
+		[URLs addObject:URLValue];
+	}
+
+	// Serialize current settings
+	FBackgroundNSURLSessionDownloadTaskData* TaskData = [FBackgroundNSURLSessionDownloadTaskData TaskDataWithURLs:URLs WithRetryCount:RetryResumeDataLimit];
 
-		[BackgroundSession release];
-		BackgroundSession = nil;
+	// Check for existing download task, could be from previous app session.
+	UE_DNLD_LOG(@"Trying to find existing download for asset path '%@'", AssetPath);
+	if (AssetPath != nil)
+	{
+		__block NSUInteger ExistingDownloadId = InvalidDownloadId;
+		__block NSURLSessionDownloadTask* ExistingTask = nil;
+		[_AllDownloads enumerateKeysAndObjectsUsingBlock:^(NSNumber* IterKey, NSURLSessionDownloadTask* IterTask, BOOL* IterStop)
+		{
+			if (IterTask.originalRequest != nil && [IterTask.originalRequest.URL.path isEqualToString:AssetPath])
+			{
+				ExistingDownloadId = IterKey.unsignedIntegerValue;
+				ExistingTask = [IterTask retain]; // Retain in case if task gets killed in another thread.
+				*IterStop = YES;
+			}
+		}];
+		
+		if (ExistingDownloadId != InvalidDownloadId && ExistingTask != nil)
+		{
+			UE_DNLD_LOG(@"Found existing download task for path '%@' with DownloadId %lu", AssetPath, ExistingDownloadId);
+
+			// Update existing task state to new one, to reset retry counters, cdn links, etc.
+			[ExistingTask setTaskDescription:[TaskData ToSerializedString]];
+
+			[ExistingTask resume]; // Resume task just in case if it was not running before.
+			[ExistingTask release];
+
+			return ExistingDownloadId;
+		}
 	}
 
-	CachedIdentifierName.Reset();
+	NSURL* URL = [TaskData GetFirstURL];
+	NSURLSessionDownloadTask* Task = [self CreateDownloadForURL:URL WithPriority:Priority WithTaskData:TaskData];
+	const NSUInteger DownloadId = [self EnsureTaskIsTracked:Task];
+	[Task resume];
+
+	return DownloadId;
 }
 
-NSURLSession* FBackgroundURLSessionHandler::GetBackgroundSession()
+- (void)PauseDownload:(NSUInteger)DownloadId
 {
-	ensureAlwaysMsgf((nullptr != BackgroundSession), TEXT("Call to GetBackgroundSession with no valid BackgroundSession! Make sure InitBackgroundSession is called first!"));
-	return BackgroundSession;
+	UE_DNLD_LOG(@"PauseDownload for DownloadId %lu", DownloadId);
+
+	NSURLSessionDownloadTask* Task = [self FindDownloadTaskFor:DownloadId];
+	if (Task != nil)
+	{
+		[Task suspend];
+	}
 }
 
-BackgroundHttpFileHashHelperRef FBackgroundURLSessionHandler::GetFileHashHelper()
+- (void)ResumeDownload:(NSUInteger)DownloadId
 {
-	return FileHashHelper;
+	UE_DNLD_LOG(@"ResumeDownload for DownloadId %lu", DownloadId);
+
+	NSURLSessionDownloadTask* Task = [self FindDownloadTaskFor:DownloadId];
+	if (Task != nil)
+	{
+		[Task resume];
+	}
 }
 
-void FBackgroundURLSessionHandler::CreateBackgroundSessionWorkingDirectory()
+- (void)CancelDownload:(NSUInteger)DownloadId
 {
-	const FString& DirectoryPath = FBackgroundHttpFileHashHelper::GetTemporaryRootPath();
-    
-	if (ensureAlwaysMsgf(!DirectoryPath.IsEmpty(), TEXT("Invalid empty BackgroundSessionWorkingDirectoryPath!")))
+	UE_DNLD_LOG(@"CancelDownload for DownloadId %lu", DownloadId);
+
+	NSURLSessionDownloadTask* Task = [self FindDownloadTaskFor:DownloadId];
+	if (Task == nil)
+	{
+		return;
+	}
+
+	FBackgroundNSURLSessionDownloadTaskData* TaskData = [FBackgroundNSURLSessionDownloadTaskData TaskDataFromSerializedString:Task.taskDescription];
+	if (TaskData != nil)
 	{
-        IPlatformFile& PlatformFile = FPlatformFileManager::Get().GetPlatformFile();
-        PlatformFile.CreateDirectory(*DirectoryPath);
-        
-        if (!PlatformFile.DirectoryExists(*DirectoryPath))
-        {
-            NSString* PathAsNSString = DirectoryPath.GetNSString();
-            //May not yet have UE_Log
-            ensureAlwaysMsgf(false, TEXT("Warning: Error in CreateBackgroundSessionWorkingDirectory"));
-        }
+		// Nuke task data, otherwise didCompleteWithError might retry request.
+		[TaskData Cancel];
+		Task.taskDescription = [TaskData ToSerializedString];
 	}
+
+	// We're done with this task.
+	[self EnsureTaskIsNotTracked:Task];
+
+	// Will invoke didCompleteWithError if task is incomplete.
+	[Task cancel];
 }
 
-void FBackgroundURLSessionHandler::CallBackgroundURLSessionCompleteHandler()
+- (void)SetPriority:(float)Priority ForDownload:(NSUInteger)DownloadId
 {
-	NSLog(@"OnIOSBackgroundDownload_SessionDidFinishAllEvents Delegates ALL Finshed. Attempting to callback into iOS to inform session work is finished.");
-	
-	IOSAppDelegate* AppDelegate = [IOSAppDelegate GetDelegate];
-	
-	if ((nullptr != AppDelegate) && (nil != AppDelegate.BackgroundSessionEventCompleteDelegate))
+	UE_DNLD_LOG(@"SetPriority %f for DownloadId %lu", Priority, DownloadId);
+
+	NSURLSessionDownloadTask* Task = [self FindDownloadTaskFor:DownloadId];
+	if (Task != nil)
 	{
-		//make a local copy of the completion handler to capture with the block
-		void(^completionHandler)() = AppDelegate.BackgroundSessionEventCompleteDelegate;
-		[completionHandler retain];
-		
-		//set the completion handler to nil to protect from anything else trying to use it and to release our copy
-		AppDelegate.BackgroundSessionEventCompleteDelegate = nil;
-		
-	   //Need to issue our stored completion callback on main thread as its a UI callback
-		[[NSOperationQueue mainQueue] addOperationWithBlock:^{
-	
-			//call the copied completion handler
-			 NSLog(@"OnIOSBackgroundDownload_SessionDidFinishAllEvents Calling iOS Completion handler.");
-			 completionHandler();
-			
-			//Release local copy so we don't leak the completion handler
-			[completionHandler release];
-		}];
+		[Task setPriority:Priority];
 	}
-	else
+}
+
+- (void)SetCurrentDownloadedBytes:(uint64)DownloadedBytes ForTask:(NSURLSessionDownloadTask*)Task
+{
+	if (Task != nil)
 	{
-		NSLog(@"OnIOSBackgroundDownload_SessionDidFinishAllEvents Can not call iOS Completion Handler as it is not set");
+		[Task.progress setUserInfoObject:[NSNumber numberWithLongLong:DownloadedBytes] forKey:NSProgressDownloadCompletedBytes];
 	}
 }
 
-void FBackgroundURLSessionHandler::AddDelayedBackgroundURLSessionComplete()
+- (uint64)GetCurrentDownloadedBytes:(NSUInteger)DownloadId
 {
-	FPlatformAtomics::InterlockedIncrement(&DelayedBackgroundURLSessionCompleteCount);
+	NSURLSessionDownloadTask* Task = [self FindDownloadTaskFor:DownloadId];
+	if (Task != nil)
+	{
+		NSNumber* CompletedBytes = [Task.progress.userInfo objectForKey:NSProgressDownloadCompletedBytes];
+		if (CompletedBytes != nil)
+		{
+			return CompletedBytes.unsignedLongLongValue;
+		}
+	}
+
+	return 0;
 }
 
-void FBackgroundURLSessionHandler::OnDelayedBackgroundURLSessionCompleteHandlerCalled()
+- (void)RecreateDownloads
 {
-	int NewDecrementedCount = FPlatformAtomics::InterlockedDecrement(&DelayedBackgroundURLSessionCompleteCount);
-	NSLog(@"OnIOSBackgroundDownload_SessionDidFinishAllEvents Delegate Finshed. Still pending: %d", NewDecrementedCount);
-	
-	if (NewDecrementedCount == 0)
+	UE_DNLD_LOG(@"RecreateDownloads started");
+
+	// Copy keys to avoid deadlocking in case if cancel/resume/etc will call delegates in-place
+	NSArray<__kindof NSNumber*>* AllKeys = nil;
+	@synchronized(_AllDownloads)
 	{
-		CallBackgroundURLSessionCompleteHandler();
+		AllKeys = [_AllDownloads allKeys];
 	}
+
+	for (NSNumber* IterKey in AllKeys)
+	{
+		const NSUInteger DownloadId = IterKey.unsignedIntegerValue;
+		NSURLSessionDownloadTask* OldTask = [_AllDownloads objectForKey:IterKey];
+		const float OldTaskPriority = OldTask.priority;
+		const NSURLSessionTaskState OldTaskState = OldTask.state;
+
+		FBackgroundNSURLSessionDownloadTaskData* NewTaskData = [FBackgroundNSURLSessionDownloadTaskData TaskDataFromSerializedString:OldTask.taskDescription];
+		if (NewTaskData == nil)
+		{
+			continue;
+		}
+
+		// Cancel old task
+		[self CancelDownload:DownloadId];
+
+		// Start a new task
+		[NewTaskData ResetRetryCount:RetryResumeDataLimit];
+		NSURLSessionDownloadTask* NewTask = [self CreateDownloadForURL:[NewTaskData GetFirstURL] WithPriority:OldTaskPriority WithTaskData:NewTaskData];
+		[self ReplaceTrackedTaskWith:NewTask ForDownloadId:DownloadId];
+		if (OldTaskState != NSURLSessionTaskStateSuspended)
+		{
+			[NewTask resume];
+		}
+	}
+
+	UE_DNLD_LOG(@"RecreateDownloads finished");
 }
 
-@implementation BackgroundDownloadDelegate
-	
--(void)URLSession:(NSURLSession *)session downloadTask : (NSURLSessionDownloadTask *)downloadTask didFinishDownloadingToURL : (NSURL *)location;
+- (NSURLSessionDownloadTask*)FindDownloadTaskFor:(NSUInteger)DownloadId
 {
-	//Go ahead and try and move the file out of IOS system temporary storage and into our own background download specific
-	//temp storage that won't be deleted by iOS
-	//We do this here because we need to handle this case even if nothing has subscribed to the delegate yet
-	
-	NSURL* destinationURL = nil;
-	NSError* downloadError = [downloadTask error];
-    const bool didTaskError = (downloadError != nullptr);
-    
-    //Get path on disk we need to move the completed iOS temp file to
-    const FString TaskURL([[[downloadTask currentRequest] URL] absoluteString]);
-	const FString& TempFileName = FBackgroundURLSessionHandler::GetFileHashHelper()->FindOrAddTempFilenameMappingForURL(TaskURL);
-	const FString& DestinationPath = FBackgroundHttpFileHashHelper::GetFullPathOfTempFilename(TempFileName);
-	
-	if (!didTaskError && (nil != location))
+	@synchronized(_AllDownloads)
 	{
-        //Convert UFS path to external path so we can pass it into iOS function to move iOS temp file
-        IPlatformFile& PlatformFile = FPlatformFileManager::Get().GetPlatformFile();
-        FString ConvertedPath = PlatformFile.ConvertToAbsolutePathForExternalAppForWrite(*DestinationPath);
-        
-		destinationURL = [NSURL fileURLWithPath:ConvertedPath.GetNSString()];
-        NSLog(@"Attempting to move iOS temp file to BackgroundHttp temp folder | iOSLocation:%s | DestinationLocation:%s", [[location path] UTF8String], [DestinationPath.GetNSString() UTF8String]);
-        
-		[[NSFileManager defaultManager] moveItemAtURL:location toURL:destinationURL error:&downloadError];
+		return [_AllDownloads objectForKey:[NSNumber numberWithUnsignedInteger:DownloadId]];
 	}
-    
-	FIOSBackgroundDownloadCoreDelegates::OnIOSBackgroundDownload_DidFinishDownloadingToURL.Broadcast(downloadTask, downloadError, DestinationPath);
 }
 
--(void)URLSession:(NSURLSession *)session downloadTask : (NSURLSessionDownloadTask *)downloadTask didWriteData:(int64_t)bytesWritten totalBytesWritten:(int64_t)totalBytesWritten totalBytesExpectedToWrite:(int64_t)totalBytesExpectedToWrite;
-{	
-	FIOSBackgroundDownloadCoreDelegates::OnIOSBackgroundDownload_DidWriteData.Broadcast(downloadTask, bytesWritten, totalBytesWritten, totalBytesExpectedToWrite);
+- (NSUInteger)FindDownloadIdForTask:(NSURLSessionDownloadTask*)Task
+{
+	// TODO this is slow, optimize if needed.
+	__block NSUInteger ExistingDownloadId = InvalidDownloadId;
+	@synchronized(_AllDownloads)
+	{
+		[_AllDownloads enumerateKeysAndObjectsUsingBlock:^(NSNumber* IterKey, NSURLSessionDownloadTask* IterTask, BOOL* IterStop)
+		{
+			if (IterTask == Task)
+			{
+				ExistingDownloadId = IterKey.unsignedIntegerValue;
+				*IterStop = YES;
+			}
+		}];
+	}
+
+	return ExistingDownloadId;
 }
 
--(void)URLSession:(NSURLSession *)session task : (NSURLSessionTask *)task didCompleteWithError : (NSError *)error;
+- (NSUInteger)EnsureTaskIsTracked:(NSURLSessionDownloadTask*)Task
 {
-    NSString* TaskURLString = [[[task currentRequest] URL] absoluteString];
-    NSLog(@"Task ending with error. TaskURL:%s", [TaskURLString UTF8String]);
+	const NSUInteger ExistingDownloadId = [self FindDownloadIdForTask:Task];
+	if (ExistingDownloadId != InvalidDownloadId)
+	{
+		return ExistingDownloadId;
+	}
+
+	@synchronized(_AllDownloads)
+	{
+		const NSUInteger DownloadId = _NextDownloadId++;
+		[_AllDownloads setObject:Task forKey:[NSNumber numberWithUnsignedInteger:DownloadId]];
+		return DownloadId;
+	}
+}
 
-    FIOSBackgroundDownloadCoreDelegates::OnIOSBackgroundDownload_DidCompleteWithError.Broadcast(task, error);
+- (void)ReplaceTrackedTaskWith:(NSURLSessionDownloadTask*)NewTask ForDownloadId:(NSUInteger)DownloadId
+{
+	@synchronized(_AllDownloads)
+	{
+		[_AllDownloads setObject:NewTask forKey:[NSNumber numberWithUnsignedInteger:DownloadId]];
+	}
 }
 
--(void)URLSessionDidFinishEventsForBackgroundURLSession:(NSURLSession *)session;
+- (void)EnsureTaskIsNotTracked:(NSURLSessionDownloadTask*)Task
 {
-	//Save any dirty hash data now that we have done all work for a given session
-	//this way any newly completed downloads will have their URL mapping saved correctly
-	FBackgroundURLSessionHandler::GetFileHashHelper()->SaveData();
-	
-	//Check if we have this delegate setup. If not our app may have not been loaded or we were evicted.
-	const bool bHasURLSessionDelegateBeenSetup = FIOSBackgroundDownloadCoreDelegates::OnIOSBackgroundDownload_SessionDidFinishAllEvents.IsBound();
+	const NSUInteger ExistingDownloadId = [self FindDownloadIdForTask:Task];
+	if (ExistingDownloadId != InvalidDownloadId)
+	{
+		@synchronized(_AllDownloads)
+		{
+			[_AllDownloads removeObjectForKey:[NSNumber numberWithUnsignedInteger:ExistingDownloadId]];
+		}
+	}
+}
+
+- (void)SetDownloadResult:(NSInteger)HTTPCode WithTempFile:(NSString*)TempFile ForDownload:(NSURLSessionDownloadTask*)Task
+{
+	const NSUInteger DownloadId = [self FindDownloadIdForTask:Task];
+	if (DownloadId == InvalidDownloadId)
+	{
+		UE_DNLD_LOG(@"Can't find DownloadId for task '%@'", Task.taskDescription);
+		return;
+	}
+
+	// We don't necessarily care if these values survive between application restarts.
+	// Otherwise we need to put them inside FBackgroundNSURLSessionDownloadTaskData.
+	[Task.progress setUserInfoObject:[NSNumber numberWithInteger:HTTPCode] forKey:NSProgressDownloadResultStatusCode];
+	[Task.progress setUserInfoObject:TempFile forKey:NSProgressDownloadResultTempFilePath];
+}
+
+- (NSString* _Nullable)GetDownloadResult:(NSUInteger)DownloadId OutStatus:(BOOL* _Nonnull)OutStatus OutStatusCode:(NSInteger*)OutStatusCode
+{
+	NSURLSessionDownloadTask* Task = [self FindDownloadTaskFor:DownloadId];
+	if (Task == nil)
+	{
+		*OutStatus = NO;
+		return nil;
+	}
+
+	NSNumber* ResultStatusCode = [Task.progress.userInfo objectForKey:NSProgressDownloadResultStatusCode];
+	if (ResultStatusCode == nil)
+	{
+		*OutStatus = NO;
+		return nil;
+	}
+
+	*OutStatus = YES;
+	*OutStatusCode = ResultStatusCode.integerValue;
+
+	return [Task.progress.userInfo objectForKey:NSProgressDownloadResultTempFilePath];
+}
+
+//- (void)URLSession:(NSURLSession*)Session didBecomeInvalidWithError:(NSError*)Error
+//{
+//	UE_DNLD_LOG(@"didBecomeInvalidWithError");
+//}
+
+- (void)URLSessionDidFinishEventsForBackgroundURLSession:(NSURLSession*)Session
+{
+	UE_DNLD_LOG(@"URLSessionDidFinishEventsForBackgroundURLSession");
+
+	IOSAppDelegate* AppDelegate = [IOSAppDelegate GetDelegate];
+	if (AppDelegate == nullptr || AppDelegate.BackgroundSessionEventCompleteDelegate == nil)
+	{
+		return;
+	}
 	
-	if (bHasURLSessionDelegateBeenSetup)
+	void(^CompletionHandler)() = [AppDelegate.BackgroundSessionEventCompleteDelegate retain];
+	AppDelegate.BackgroundSessionEventCompleteDelegate = nil;
+
+	// Completion handler has to be invoked on the main thread.
+	[[NSOperationQueue mainQueue] addOperationWithBlock:^
 	{
-		NSLog(@"Calling into OnIOSBackgroundDownload_SessionDidFinishAllEvents Waiting for Delegates to complete.");
-		
-		//We call this here and then the completion handler ourselves below this. That way we delay the complete until all systems responding to this delegate have received the broadcast.
-		FBackgroundURLSessionHandler::AddDelayedBackgroundURLSessionComplete();
-		
-		//First process other systems doing things for all events processing on URL Session
-		FIOSBackgroundDownloadCoreDelegates::OnIOSBackgroundDownload_SessionDidFinishAllEvents.Broadcast(session, FIOSBackgroundDownloadCoreDelegates::OnDelayedBackgroundURLSessionCompleteHandler);
-		
-		//Now call the completion handler to erase the delay we had above
-		FBackgroundURLSessionHandler::OnDelayedBackgroundURLSessionCompleteHandlerCalled();
+		UE_DNLD_LOG(@"URLSessionDidFinishEventsForBackgroundURLSession calling completion handler.");
+		CompletionHandler();
+		[CompletionHandler release];
+	}];
+}
+
+- (void)URLSession:(NSURLSession*)Session task:(NSURLSessionTask*)GenericTask didCompleteWithError:(NSError*)Error
+{
+	if (GenericTask.state == NSURLSessionTaskStateCompleted || ![GenericTask isKindOfClass:[NSURLSessionDownloadTask class]])
+	{
+		return;
 	}
-	//Handle restarting any suspended tasks. These can exist because we didn't resume these before our app was evicted.
-	else
+
+	NSURLSessionDownloadTask* Task = (NSURLSessionDownloadTask*)GenericTask;
+	NSString* LocalizedDescription = Error != nil ? Error.localizedDescription : @"nil";
+
+	const NSUInteger DownloadId = [self FindDownloadIdForTask:Task];
+	const bool bIsTrackedTask = DownloadId != InvalidDownloadId;
+
+	FBackgroundNSURLSessionDownloadTaskData* TaskData = [FBackgroundNSURLSessionDownloadTaskData TaskDataFromSerializedString:Task.taskDescription];
+	if (bIsTrackedTask && TaskData != nil)
 	{
-		NSLog(@"Resuming all tasks regardless of request state as we have no registered handler of OnIOSBackgroundDownload_SessionDidFinishAllEvents");
+		NSData* ResumeData = [Error.userInfo objectForKey:NSURLSessionDownloadTaskResumeData];
+		const bool bHasResumeData = ResumeData != nil && ResumeData.length > 0;
 		
-		[session getTasksWithCompletionHandler : ^ (NSArray<__kindof NSURLSessionDataTask*>* DataTasks, NSArray<__kindof NSURLSessionUploadTask*>* UploadTasks, NSArray<__kindof NSURLSessionDownloadTask*>* DownloadTasks)
+		NSURL* NextURL = [TaskData GetNextURL];
+		
+		// Continue trying if next URL is available.
+		if (NextURL != nil)
 		{
-			for (NSURLSessionDownloadTask* DownloadTask : DownloadTasks)
+			// Create resume request if our URL is the same and we have resume data.
+			if (bHasResumeData && NextURL != nil && [NextURL.absoluteString isEqualToString:Task.originalRequest.URL.absoluteString])
+			{
+				UE_DNLD_LOG(@"didCompleteWithError, task '%@' with taskIdentifier %lu failed due to '%@' and has resume data and next url is the same, retrying", Task.taskDescription, Task.taskIdentifier, LocalizedDescription);
+
+				NSURLSessionDownloadTask* NewTask = [self CreateDownloadForResumeData:ResumeData WithPriority:Task.priority WithTaskData:TaskData];
+				[self ReplaceTrackedTaskWith:NewTask ForDownloadId:DownloadId];
+				[NewTask resume];
+				return;
+			}
+			else
 			{
-				if ([DownloadTask state] == NSURLSessionTaskStateSuspended)
+				if (bHasResumeData)
+				{
+					// It should be possible to patch resume data to point to a new URL. But there is no public API to do that yet.
+					UE_DNLD_LOG(@"didCompleteWithError, task '%@' with taskIdentifier %lu failed due to '%@' and has resume data but next url is different, retrying", Task.taskDescription, Task.taskIdentifier, LocalizedDescription);
+				}
+				else
 				{
-					NSLog(@"Resuming Task for URL:%s", [[[[DownloadTask currentRequest] URL] absoluteString] UTF8String]);
-					[DownloadTask resume];
+					UE_DNLD_LOG(@"didCompleteWithError, task '%@' with taskIdentifier %lu failed due to '%@' and has no resume data or next url is different, retrying", Task.taskDescription, Task.taskIdentifier, LocalizedDescription);
 				}
+
+				NSURLSessionDownloadTask* NewTask = [self CreateDownloadForURL:NextURL WithPriority:Task.priority WithTaskData:TaskData];
+				[self ReplaceTrackedTaskWith:NewTask ForDownloadId:DownloadId];
+				[NewTask resume];
+				return;
 			}
-			
-			//Go ahead and call the completion handler now that we have resumed everything.
-			FBackgroundURLSessionHandler::OnDelayedBackgroundURLSessionCompleteHandlerCalled();
-		}];
+		}
+	}
+
+	// Can't retry anymore, fail the request
+	{
+		UE_DNLD_LOG(@"didCompleteWithError, task '%@' with taskIdentifier %lu failed due to '%@', has no retry data or no next url, failing request", Task.taskDescription, Task.taskIdentifier, LocalizedDescription);
+
+		NSURLResponse* GenericResponse = Task.response;
+		NSInteger StatusCode = HTTPStatusCodeErrorServer;
+		if ([GenericResponse isKindOfClass:[NSHTTPURLResponse class]])
+		{
+			NSHTTPURLResponse* Response = (NSHTTPURLResponse*)GenericResponse;
+			if (Response.statusCode >= HTTPStatusCodeErrorBadRequest)
+			{
+				StatusCode = Response.statusCode;
+			}
+		}
+
+		[self SetDownloadResult:StatusCode WithTempFile:nil ForDownload:Task];
+	}
+}
+
+//- (void)URLSession:(NSURLSession *)session taskIsWaitingForConnectivity:(NSURLSessionTask *)Task
+//{
+//	UE_DNLD_LOG(@"taskIsWaitingForConnectivity '%@'", Task.taskDescription);
+//}
+
+//- (void)URLSession:(NSURLSession *)Session task:(NSURLSessionTask *)Task willBeginDelayedRequest:(NSURLRequest *)Request completionHandler:(void (^)(NSURLSessionDelayedRequestDisposition Disposition, NSURLRequest* NewRequest))CompletionHandler
+//{
+//	UE_DNLD_LOG(@"willBeginDelayedRequest '%@' for '%@' with taskIdentifier %lu", Task.taskDescription, Request.debugDescription, Task.taskIdentifier);
+//	CompletionHandler(NSURLSessionDelayedRequestContinueLoading, Request);
+//}
+
+- (void)URLSession:(NSURLSession*)Session downloadTask:(NSURLSessionDownloadTask*)Task didFinishDownloadingToURL:(NSURL*)Location
+{
+	// Should not be needed, but ensure this just in case
+	[self EnsureTaskIsTracked:Task];
+
+	NSString* DestinationPath = [self GetTempPathForURL:Task.originalRequest.URL];
+
+	// Try to remove existing file in case if we have a stale file.
+	if ([[NSFileManager defaultManager] fileExistsAtPath:DestinationPath])
+	{
+		[[NSFileManager defaultManager] removeItemAtPath:DestinationPath error:nil];
 	}
+
+	NSError* Error = nil;
+	[[NSFileManager defaultManager] moveItemAtURL:Location toURL:[NSURL fileURLWithPath:DestinationPath] error:&Error];
+
+	// Update task progress just in case didWriteData was not invoked
+	const uint64 TotalBytesWritten = [[[NSFileManager defaultManager] attributesOfItemAtPath:DestinationPath error:nil] fileSize];
+	[self SetCurrentDownloadedBytes:TotalBytesWritten ForTask:Task];
+
+	if (Error != nil)
+	{
+		[self SetDownloadResult:HTTPStatusCodeErrorServer WithTempFile:nil ForDownload:Task];
+
+		UE_DNLD_LOG(@"didFinishDownloadingToURL task '%@' with taskIdentifier %lu failed to move file to '%@' due to '%@'", Task.taskDescription, Task.taskIdentifier, DestinationPath, Error.localizedDescription);
+	}
+	else
+	{
+		[self SetDownloadResult:HTTPStatusCodeSuccessCreated WithTempFile:DestinationPath ForDownload:Task];
+
+		UE_DNLD_LOG(@"didFinishDownloadingToURL task '%@' with taskIdentifier %lu move file to '%@', download finished", Task.taskDescription, Task.taskIdentifier, DestinationPath);
+	}
+}
+
+- (void)URLSession:(NSURLSession*)session downloadTask:(NSURLSessionDownloadTask*)Task didWriteData:(int64_t)BytesWritten totalBytesWritten:(int64_t)TotalBytesWritten totalBytesExpectedToWrite:(int64_t)TotalBytesExpectedToWrite
+{
+	[self SetCurrentDownloadedBytes:TotalBytesWritten ForTask:Task];
 }
 
 @end
+
+const uint64 FBackgroundURLSessionHandler::InvalidDownloadId = [FBackgroundNSURLSession GetInvalidDownloadId];
+
+void FBackgroundURLSessionHandler::AllowCellular(bool bAllow)
+{
+	const BOOL bCurrentValue = [FBackgroundNSURLSession Shared].AllowCellular;
+	const BOOL bNewValue = bAllow ? YES : NO;
+	if (bCurrentValue == bNewValue)
+	{
+		return;
+	}
+
+	[[FBackgroundNSURLSession Shared] setAllowCellular:bNewValue];
+	[[FBackgroundNSURLSession Shared] RecreateDownloads];
+}
+
+uint64 FBackgroundURLSessionHandler::CreateOrFindDownload(const TArray<FString>& URLs, const float Priority, BackgroundHttpFileHashHelperRef HelperRef)
+{
+	NSMutableArray* URLArray = [NSMutableArray arrayWithCapacity:URLs.Num()];
+	for (const FString& URL: URLs)
+	{
+		[URLArray addObject:URL.GetNSString()];
+	}
+
+	[[FBackgroundNSURLSession Shared] SetFileHashHelper:HelperRef];
+	return [[FBackgroundNSURLSession Shared] CreateOrFindDownloadForURLs:URLArray WithPriority:Priority];
+}
+
+void FBackgroundURLSessionHandler::PauseDownload(const uint64 DownloadId)
+{
+	[[FBackgroundNSURLSession Shared] PauseDownload:DownloadId];
+}
+
+void FBackgroundURLSessionHandler::ResumeDownload(const uint64 DownloadId)
+{
+	[[FBackgroundNSURLSession Shared] ResumeDownload:DownloadId];
+}
+
+void FBackgroundURLSessionHandler::CancelDownload(const uint64 DownloadId)
+{
+	[[FBackgroundNSURLSession Shared] CancelDownload:DownloadId];
+}
+
+void FBackgroundURLSessionHandler::SetPriority(const uint64 DownloadId, const float Priority)
+{
+	[[FBackgroundNSURLSession Shared] SetPriority:Priority ForDownload:DownloadId];
+}
+
+uint64 FBackgroundURLSessionHandler::GetCurrentDownloadedBytes(const uint64 DownloadId)
+{
+	return [[FBackgroundNSURLSession Shared] GetCurrentDownloadedBytes:DownloadId];
+}
+
+bool FBackgroundURLSessionHandler::IsDownloadFinished(const uint64 DownloadId, int32& OutResultHTTPCode, FString& OutTemporaryFilePath)
+{
+	BOOL Status = NO;
+	NSInteger StatusCode = 0;
+	NSString* TempFile = [[FBackgroundNSURLSession Shared] GetDownloadResult:DownloadId OutStatus:&Status OutStatusCode:&StatusCode];
+	if (!Status)
+	{
+		return false;
+	}
+
+	OutResultHTTPCode = (int32)StatusCode;
+	if (TempFile != nil)
+	{
+		OutTemporaryFilePath = FString(TempFile);
+		UE_DNLD_LOG(@"DownloadId %lli finished with status code %li and path '%@'", DownloadId, (long)StatusCode, TempFile);
+	}
+	else
+	{
+		UE_DNLD_LOG(@"DownloadId %lli finished with status code %li and no path", DownloadId, (long)StatusCode);
+	}
+
+	return true;
+}
+
+void FBackgroundURLSessionHandler::HandleEventsForBackgroundURLSession(const FString& SessionIdentifier)
+{
+	NSString* Identifier = SessionIdentifier.GetNSString();
+	if (![[FBackgroundNSURLSession GetNSURLSessionIdentifier] isEqualToString:Identifier])
+	{
+		UE_DNLD_LOG(@"HandleEventsForBackgroundURLSession ignoring session identifier '%@'", Identifier);
+		return;
+	}
+
+	UE_DNLD_LOG(@"HandleEventsForBackgroundURLSession will initializes session with identifier '%@'", Identifier);
+	[FBackgroundNSURLSession Shared];
+	// will invoke URLSessionDidFinishEventsForBackgroundURLSession internally.
+}
diff --git a/Engine/Source/Runtime/ApplicationCore/Public/IOS/IOSBackgroundURLSessionHandler.h b/Engine/Source/Runtime/ApplicationCore/Public/IOS/IOSBackgroundURLSessionHandler.h
--- a/Engine/Source/Runtime/ApplicationCore/Public/IOS/IOSBackgroundURLSessionHandler.h
+++ b/Engine/Source/Runtime/ApplicationCore/Public/IOS/IOSBackgroundURLSessionHandler.h
@@ -8,20 +8,20 @@
 #include "BackgroundHttpFileHashHelper.h"
 
 //Call backs called by the bellow FBackgroundURLSessionHandler so higher-level systems can respond to task updates.
-class APPLICATIONCORE_API FIOSBackgroundDownloadCoreDelegates
+class UE_DEPRECATED(5.5, "Use new download methods in FBackgroundURLSessionHandler.") APPLICATIONCORE_API FIOSBackgroundDownloadCoreDelegates
 {
 public:
-    DECLARE_MULTICAST_DELEGATE_ThreeParams(FIOSBackgroundDownload_DidFinishDownloadingToURL, NSURLSessionDownloadTask*, NSError*, const FString&);
-    DECLARE_MULTICAST_DELEGATE_FourParams(FIOSBackgroundDownload_DidWriteData, NSURLSessionDownloadTask*, int64_t /*Bytes Written Since Last Call */, int64_t /*Total Bytes Written */, int64_t /*Total Bytes Expedted To Write */);
-    DECLARE_MULTICAST_DELEGATE_TwoParams(FIOSBackgroundDownload_DidCompleteWithError, NSURLSessionTask*, NSError*);
+	DECLARE_MULTICAST_DELEGATE_ThreeParams(FIOSBackgroundDownload_DidFinishDownloadingToURL, NSURLSessionDownloadTask*, NSError*, const FString&);
+	DECLARE_MULTICAST_DELEGATE_FourParams(FIOSBackgroundDownload_DidWriteData, NSURLSessionDownloadTask*, int64_t /*Bytes Written Since Last Call */, int64_t /*Total Bytes Written */, int64_t /*Total Bytes Expedted To Write */);
+	DECLARE_MULTICAST_DELEGATE_TwoParams(FIOSBackgroundDownload_DidCompleteWithError, NSURLSessionTask*, NSError*);
 	DECLARE_DELEGATE(FIOSBackgroundDownload_DelayedBackgroundURLSessionCompleteHandler);
 	DECLARE_MULTICAST_DELEGATE_TwoParams(FIOSBackgroundDownload_SessionDidFinishAllEvents, NSURLSession*, FIOSBackgroundDownload_DelayedBackgroundURLSessionCompleteHandler);
-	
-	static FIOSBackgroundDownload_DidFinishDownloadingToURL OnIOSBackgroundDownload_DidFinishDownloadingToURL;
-	static FIOSBackgroundDownload_DidWriteData OnIOSBackgroundDownload_DidWriteData;
-	static FIOSBackgroundDownload_DidCompleteWithError OnIOSBackgroundDownload_DidCompleteWithError;
-	static FIOSBackgroundDownload_SessionDidFinishAllEvents OnIOSBackgroundDownload_SessionDidFinishAllEvents;
-	static FIOSBackgroundDownload_DelayedBackgroundURLSessionCompleteHandler OnDelayedBackgroundURLSessionCompleteHandler;
+
+	UE_DEPRECATED(5.5, "Use new API in FBackgroundURLSessionHandler.") static FIOSBackgroundDownload_DidFinishDownloadingToURL OnIOSBackgroundDownload_DidFinishDownloadingToURL;
+	UE_DEPRECATED(5.5, "Use new API in FBackgroundURLSessionHandler.") static FIOSBackgroundDownload_DidWriteData OnIOSBackgroundDownload_DidWriteData;
+	UE_DEPRECATED(5.5, "Use new API in FBackgroundURLSessionHandler.") static FIOSBackgroundDownload_DidCompleteWithError OnIOSBackgroundDownload_DidCompleteWithError;
+	UE_DEPRECATED(5.5, "Use new API in FBackgroundURLSessionHandler.") static FIOSBackgroundDownload_SessionDidFinishAllEvents OnIOSBackgroundDownload_SessionDidFinishAllEvents;
+	UE_DEPRECATED(5.5, "Use new API in FBackgroundURLSessionHandler.") static FIOSBackgroundDownload_DelayedBackgroundURLSessionCompleteHandler OnDelayedBackgroundURLSessionCompleteHandler;
 };
 
 //Interface for wrapping a NSURLSession configured to support background downloading of NSURLSessionDownloadTasks.
@@ -31,41 +31,56 @@ class APPLICATIONCORE_API FBackgroundURLSessionHandler
 {
 public:
 	// Initializes a BackgroundSession with the given identifier. If the current background session already exists, returns true if the identifier matches. False if identifier doesn't match or if the session fails to create.
-	static bool InitBackgroundSession(const FString& SessionIdentifier);
+	UE_DEPRECATED(5.5, "Use new API in FBackgroundURLSessionHandler.") static bool InitBackgroundSession(const FString& SessionIdentifier) {return false;}
 
 	//bShouldInvalidateExistingTasks determines if the session cancells all outstanding tasks immediately and cancels the session immediately or waits for them to finish and then invalidates the session
-	static void ShutdownBackgroundSession(bool bShouldFinishTasksFirst = true);
+	UE_DEPRECATED(5.5, "Use new API in FBackgroundURLSessionHandler.") static void ShutdownBackgroundSession(bool bShouldFinishTasksFirst = true) {}
 
 	//Gets a pointer to the current background session
-	static NSURLSession* GetBackgroundSession();
+	UE_DEPRECATED(5.5, "Use new API in FBackgroundURLSessionHandler.") static NSURLSession* GetBackgroundSession() {return nullptr;}
 
-	static void CreateBackgroundSessionWorkingDirectory();
+	UE_DEPRECATED(5.5, "Use new API in FBackgroundURLSessionHandler.") static void CreateBackgroundSessionWorkingDirectory() {}
 	
 	//Function to mark if you would like for the NSURLSession to wait to call the completion handler when
 	//OnIOSBackgroundDownload_SessionDidFinishAllEvents is called for you to call the passed completion handler
 	//NOTE: Call DURING OnIOSBackgroundDownload_SessionDidFinishAllEvents
-	static void AddDelayedBackgroundURLSessionComplete();
+	UE_DEPRECATED(5.5, "Use new API in FBackgroundURLSessionHandler.") static void AddDelayedBackgroundURLSessionComplete() {}
 	
 	//Function to handle calls to OnDelayedBackgroundURLSessionCompleteHandler
 	//The intention is to call this for every call to AddDelayedBackgroundURLSessionComplete.
 	//NOTE: Once calling this your task should be completely finished with work and ready to be backgrounded!
-	static void OnDelayedBackgroundURLSessionCompleteHandlerCalled();
-	
-	static BackgroundHttpFileHashHelperRef GetFileHashHelper();
-	
-private:
-	static NSURLSession* BackgroundSession;
-	static FString CachedIdentifierName;
-	
-	static BackgroundHttpFileHashHelperRef FileHashHelper;
-	
-	//Used to track calls to AddDelayedBackgroundURLSessionComplete vs calls to the completion handler.
-	static volatile int32 DelayedBackgroundURLSessionCompleteCount;
+	UE_DEPRECATED(5.5, "Use new API in FBackgroundURLSessionHandler.") static void OnDelayedBackgroundURLSessionCompleteHandlerCalled();
+
+	// New API for background downloading
 	
-	//Used to call the stored background url session callback
-	static void CallBackgroundURLSessionCompleteHandler();
-};
+	// Value of invalid download id which could be used to compare return value of CreateOrFindDownload 
+	static const uint64 InvalidDownloadId;
+
+	// Sets if cellular is allowed to be used for new downloads.
+	// Existing downloads will be recreated to reflect new setting value.
+	static void AllowCellular(bool bAllow);
 
-//Delegate object associated with our above NSURLSession
-@interface BackgroundDownloadDelegate : NSObject<NSURLSessionDownloadDelegate>
-@end
+	// Creates a new download or finds existing download matching URL path.
+	// All URL's should have same path and only differ in domain.
+	// Priority is a value between 0.0 to 1.0, see NSURLSessionTaskPriorityDefault.
+	// HelperRef is optional shared reference to BackgroundHttpFileHashHelperRef.
+	// In case of HandleEventsForBackgroundURLSession this subsystem will create it's own reference.
+	static uint64 CreateOrFindDownload(const TArray<FString>& URLs, const float Priority, BackgroundHttpFileHashHelperRef HelperRef);
+
+	static void PauseDownload(const uint64 DownloadId);
+
+	static void ResumeDownload(const uint64 DownloadId);
+
+	// Cancels and invalidates DownloadId.
+	static void CancelDownload(const uint64 DownloadId);
+
+	// Priority is a value between 0.0 to 1.0, see NSURLSessionTaskPriorityDefault.
+	static void SetPriority(const uint64 DownloadId, const float Priority);
+
+	static uint64 GetCurrentDownloadedBytes(const uint64 DownloadId);
+
+	static bool IsDownloadFinished(const uint64 DownloadId, int32& OutResultHTTPCode, FString& OutTemporaryFilePath);
+
+	// To be used by app delegate, call it from handleEventsForBackgroundURLSession.
+	static void HandleEventsForBackgroundURLSession(const FString& SessionIdentifier);
+};
diff --git a/Engine/Source/Runtime/InstallBundleManager/Public/InstallBundleManagerInterface.h b/Engine/Source/Runtime/InstallBundleManager/Public/InstallBundleManagerInterface.h
--- a/Engine/Source/Runtime/InstallBundleManager/Public/InstallBundleManagerInterface.h
+++ b/Engine/Source/Runtime/InstallBundleManager/Public/InstallBundleManagerInterface.h
@@ -181,7 +181,8 @@ public:
 	virtual EInstallBundleRequestFlags GetModifyableContentRequestFlags() const = 0;
 	INSTALLBUNDLEMANAGER_API void UpdateContentRequestFlags(FName BundleName, EInstallBundleRequestFlags AddFlags, EInstallBundleRequestFlags RemoveFlags);
 	virtual void UpdateContentRequestFlags(TArrayView<const FName> BundleNames, EInstallBundleRequestFlags AddFlags, EInstallBundleRequestFlags RemoveFlags) = 0;
-	
+	virtual void SetCellularPreference(int32 Value) = 0;
+
 	virtual void SetCacheSize(FName CacheName, uint64 CacheSize) = 0;
 
 	INSTALLBUNDLEMANAGER_API virtual void StartPatchCheck();
diff --git a/Engine/Source/Runtime/InstallBundleManager/Public/InstallBundleSourceInterface.h b/Engine/Source/Runtime/InstallBundleManager/Public/InstallBundleSourceInterface.h
--- a/Engine/Source/Runtime/InstallBundleManager/Public/InstallBundleSourceInterface.h
+++ b/Engine/Source/Runtime/InstallBundleManager/Public/InstallBundleSourceInterface.h
@@ -127,6 +127,7 @@ public:
 	// Currently only CheckForCellularDataUsage is supported
 	virtual EInstallBundleRequestFlags GetModifyableContentRequestFlags() const { return EInstallBundleRequestFlags::None; }
 	virtual void UpdateContentRequestFlags(TArrayView<const FName> BundleNames, EInstallBundleRequestFlags AddFlags, EInstallBundleRequestFlags RemoveFlags) {}
+	virtual void SetCellularPreference(int32 Value) {}
 
 	// Derived classes should implement this if their content install will take a significant amount of time
 	virtual TOptional<FInstallBundleSourceProgress> GetBundleProgress(FName BundleName) const { return TOptional<FInstallBundleSourceProgress>(); }
diff --git a/Engine/Source/Runtime/NullInstallBundleManager/Private/NullInstallBundleManager.cpp b/Engine/Source/Runtime/NullInstallBundleManager/Private/NullInstallBundleManager.cpp
--- a/Engine/Source/Runtime/NullInstallBundleManager/Private/NullInstallBundleManager.cpp
+++ b/Engine/Source/Runtime/NullInstallBundleManager/Private/NullInstallBundleManager.cpp
@@ -132,7 +132,9 @@ class FNullInstallBundleManager final : public IInstallBundleManager
 	}
 	virtual void UpdateContentRequestFlags(TArrayView<const FName> BundleNames, EInstallBundleRequestFlags AddFlags, EInstallBundleRequestFlags RemoveFlags) override
 	{
-
+	}
+	virtual void SetCellularPreference(int32 Value) override
+	{
 	}
 
 	virtual void SetCacheSize(FName CacheName, uint64 CacheSize) override
diff --git a/Engine/Source/Runtime/Online/BackgroundHTTP/Private/IOS/ApplePlatformBackgroundHttp.cpp b/Engine/Source/Runtime/Online/BackgroundHTTP/Private/IOS/ApplePlatformBackgroundHttp.cpp
--- a/Engine/Source/Runtime/Online/BackgroundHTTP/Private/IOS/ApplePlatformBackgroundHttp.cpp
+++ b/Engine/Source/Runtime/Online/BackgroundHTTP/Private/IOS/ApplePlatformBackgroundHttp.cpp
@@ -4,19 +4,15 @@
 #include "IOS/ApplePlatformBackgroundHttpManager.h"
 #include "IOS/ApplePlatformBackgroundHttpRequest.h"
 #include "IOS/ApplePlatformBackgroundHttpResponse.h"
-#include "IOS/IOSBackgroundURLSessionHandler.h"
 
 #include "Misc/Paths.h"
 
 void FApplePlatformBackgroundHttp::Initialize()
 {
-	const FString DefaultIdentifier = TEXT("com.epicgames.backgroundhttp");
-	ensureAlwaysMsgf(FBackgroundURLSessionHandler::InitBackgroundSession(DefaultIdentifier), TEXT("Failure to create a background download session with identifier %s"), *DefaultIdentifier);
 }
 
 void FApplePlatformBackgroundHttp::Shutdown()
 {
-	FBackgroundURLSessionHandler::ShutdownBackgroundSession();
 }
 
 FBackgroundHttpManagerPtr FApplePlatformBackgroundHttp::CreatePlatformBackgroundHttpManager()
diff --git a/Engine/Source/Runtime/Online/BackgroundHTTP/Private/IOS/ApplePlatformBackgroundHttpManager.cpp b/Engine/Source/Runtime/Online/BackgroundHTTP/Private/IOS/ApplePlatformBackgroundHttpManager.cpp
--- a/Engine/Source/Runtime/Online/BackgroundHTTP/Private/IOS/ApplePlatformBackgroundHttpManager.cpp
+++ b/Engine/Source/Runtime/Online/BackgroundHTTP/Private/IOS/ApplePlatformBackgroundHttpManager.cpp
@@ -18,943 +18,87 @@
 
 #include "PlatformBackgroundHttp.h"
 
-FApplePlatformBackgroundHttpManager::~FApplePlatformBackgroundHttpManager()
+void FApplePlatformBackgroundHttpManager::AddRequest(const FBackgroundHttpRequestPtr GenericRequest)
 {
-	bWasAppleBGHTTPInitialized = false;
-	
+	FAppleBackgroundHttpRequestPtr Request = StaticCastSharedPtr<FApplePlatformBackgroundHttpRequest>(GenericRequest);
+	if (!ensureAlwaysMsgf(Request.IsValid(), TEXT("Adding a non-Apple background request to our Apple Background Http Manager2! This is not supported or expected!")))
 	{
-		FRWScopeLock ScopeLock(UnAssociatedTasksLock, SLT_Write);
-		[UnAssociatedTasks release];
-		UnAssociatedTasks = nullptr;
+		return;
 	}
-	
-    CleanUpNSURLSessionResponseDelegates();
-}
-
-FString FApplePlatformBackgroundHttpManager::BackgroundSessionIdentifier = FString("");
-float FApplePlatformBackgroundHttpManager::ActiveTimeOutSetting = 30.0f;
-int FApplePlatformBackgroundHttpManager::RetryResumeDataLimitSetting = -1.f;
-volatile bool FApplePlatformBackgroundHttpManager::bWasAppleBGHTTPInitialized = false;
-
-FApplePlatformBackgroundHttpManager::FApplePlatformBackgroundHttpManager()
-	: bHasFinishedPopulatingUnassociatedTasks(false)
-    , bIsInBackground(false)
-    , bIsIteratingThroughSessionTasks(false)
-    , RequestsPendingRemove()
-    , NumCurrentlyActiveTasks(0)
-	, MaxNumActualTasks(MaxActiveDownloads)
-{
-}
-
-void FApplePlatformBackgroundHttpManager::Initialize()
-{
-	//Initialize UnAssociatedTasks
-	{
-		FRWScopeLock ScopeLock(UnAssociatedTasksLock, SLT_Write);
-		UnAssociatedTasks = [[NSMutableDictionary alloc] init];
-	}
-	//This has its own lock when needed, so not included above
-	PopulateUnAssociatedTasks();
-	
-    GConfig->GetFloat(TEXT("BackgroundHttp.iOSSettings"), TEXT("ActiveReceiveTimeout"), ActiveTimeOutSetting, GEngineIni);
-	GConfig->GetInt(TEXT("BackgroundHttp.iOSSettings"), TEXT("RetryResumeDataLimit"), RetryResumeDataLimitSetting, GEngineIni);
-
-	SetupNSURLSessionResponseDelegates();
-
-	FBackgroundHttpManagerImpl::Initialize();
-	
-	bWasAppleBGHTTPInitialized = true;
-}
-
-void FApplePlatformBackgroundHttpManager::PopulateUnAssociatedTasks()
-{
-	if (ensureAlwaysMsgf((nullptr != UnAssociatedTasks), TEXT("Call to PopulateUnAssociatedTasks without initializing UnAssociatedTasks Dictionary!")))
-	{
-		NSURLSession* BackgroundDownloadSession = FBackgroundURLSessionHandler::GetBackgroundSession();
-		if (ensureAlwaysMsgf((nullptr != BackgroundDownloadSession), TEXT("Invalid Background Download NSURLSession during AppleBackgroundHttp Init! Should have already Initialized the NSURLSession by this point!")))
-		{
-			[BackgroundDownloadSession getAllTasksWithCompletionHandler : ^ (NSArray<__kindof NSURLSessionTask*> *tasks)
-			{
-				FRWScopeLock ScopeLock(UnAssociatedTasksLock, SLT_Write);
-				
-				//Store all existing tasks by their URL
-				for (id task in tasks)
-				{
-					//Make sure we have a valid absolute string version of the URL to use for our task's key. Otherwise, we just disregard this task.
-					if (	(task != nullptr)
-						&&	([task currentRequest] != nullptr)
-						&&  ([[task currentRequest] URL] != nullptr)
-						&&	([[[task currentRequest] URL] absoluteString] != nullptr)
-						&&	([[[[task currentRequest] URL] absoluteString] length] != 0))
-					{
-						[UnAssociatedTasks setObject : task forKey : [[[task currentRequest] URL] absoluteString]];
-					}
-				}
-
-				bHasFinishedPopulatingUnassociatedTasks = true;
-			}];
-		}
-	}
-}
-
-void FApplePlatformBackgroundHttpManager::PauseAllUnassociatedTasks()
-{
-	FRWScopeLock ScopeLock(UnAssociatedTasksLock, SLT_ReadOnly);
-	
-    for (id Key in UnAssociatedTasks)
-    {
-        NSURLSessionDownloadTask* Task = (NSURLSessionDownloadTask*)([UnAssociatedTasks objectForKey:Key]);
-        if (nullptr != Task)
-        {
-            if ([Task state] == NSURLSessionTaskStateRunning)
-            {
-                [Task suspend];
-            }
-        }
-    }
-}
-
-void FApplePlatformBackgroundHttpManager::UnpauseAllUnassociatedTasks()
-{
-	FRWScopeLock ScopeLock(UnAssociatedTasksLock, SLT_ReadOnly);
-	
-    for (id Key in UnAssociatedTasks)
-    {
-        NSURLSessionDownloadTask* Task = (NSURLSessionDownloadTask*)([UnAssociatedTasks objectForKey:Key]);
-        if (nullptr != Task)
-        {
-            if ([Task state] == NSURLSessionTaskStateSuspended)
-            {
-                [Task resume];
-            }
-        }
-    }
-}
 
-void FApplePlatformBackgroundHttpManager::Shutdown()
-{
+	if (AssociateWithAnyExistingRequest(Request))
 	{
-		FRWScopeLock ScopeLock(UnAssociatedTasksLock, SLT_Write);
-		[UnAssociatedTasks release];
-		UnAssociatedTasks = nullptr;
-	}
-
-	CleanUpNSURLSessionResponseDelegates();
-    FBackgroundURLSessionHandler::ShutdownBackgroundSession();
-}
-
-void FApplePlatformBackgroundHttpManager::AddRequest(const FBackgroundHttpRequestPtr Request)
-{
-	UE_LOG(LogBackgroundHttpManager, Verbose, TEXT("AddRequest Called - RequestID:%s"), *Request->GetRequestID());
-    
-    //See if our request is an AppleBackgroundHttpRequest so we can do more detailed checks on it.
-    FAppleBackgroundHttpRequestPtr AppleRequest = StaticCastSharedPtr<FApplePlatformBackgroundHttpRequest>(Request);
-	if (ensureAlwaysMsgf(AppleRequest.IsValid(), TEXT("Adding a non-Apple background request to our Apple Background Http Manager! This is not supported or expected!")))
-	{
-		//If we fail to generate URLMapEntries or AssociateWithAnyExistingRequest, then we will have already sent a completion handler immediately, so only start work and monitor
-		//these requests if those didn't already complete this Request
-		if (GenerateURLMapEntriesForRequest(AppleRequest) && !AssociateWithAnyExistingRequest(Request))
-		{
-			if (!AssociateWithAnyExistingUnAssociatedTasks(Request))
-			{
-				StartRequest(AppleRequest);
-			}
-
-			FRWScopeLock ScopeLock(ActiveRequestLock, SLT_Write);
-			ActiveRequests.Add(Request);
-
-			//Increment our underlying FBackgroundHttpManagerImpl tracker for active requests as we
-			//don't implement the method it uses to increase this number.
-			// NOTE: We don't make use of this number in Apple Platform functions as all requests are "Active" but their
-			// underlying Task might not be, see NumCurrentlyActiveTasks instead to track how many current Tasks are downloading data.
-			++NumCurrentlyActiveRequests;
-		}
+		return;
 	}
-}
 
-bool FApplePlatformBackgroundHttpManager::GenerateURLMapEntriesForRequest( FAppleBackgroundHttpRequestPtr Request)
-{
-	bool bWasGenerateSuccess = true;
+	const uint64 DownloadId = FBackgroundURLSessionHandler::CreateOrFindDownload(Request->GetURLList(), Request->GetNSURLSessionPriority(), GetFileHashHelper());
+	Request->SetInternalDownloadId(DownloadId);
 
-	//Attempt to add entries for all URLs
+	// Should never happen in practice
+	if (DownloadId == FBackgroundURLSessionHandler::InvalidDownloadId)
 	{
-		FRWScopeLock ScopeLock(URLToRequestMapLock, SLT_Write);
-		for (const FString& URL : Request->GetURLList())
-		{
-			FBackgroundHttpURLMappedRequestPtr& FoundRequest = URLToRequestMap.FindOrAdd(URL);
-
-			const bool bRequestAlreadyExistsForURL = ((FoundRequest.IsValid()) && (Request != FoundRequest));
-			if (ensureAlwaysMsgf(!bRequestAlreadyExistsForURL, TEXT("URL is represented by 2 different Requests! Immediately completing new request with error.")))
-			{
-				FoundRequest = Request;
-			}
-			else
-			{
-				bWasGenerateSuccess = false;
-
-				FBackgroundHttpResponsePtr NewResponse = FPlatformBackgroundHttp::ConstructBackgroundResponse(EHttpResponseCodes::Unknown, FString());
-				Request->CompleteWithExistingResponseData(NewResponse);
-			}
-		}
+		FBackgroundHttpResponsePtr Response = FPlatformBackgroundHttp::ConstructBackgroundResponse(EHttpResponseCodes::ServerError, FString());
+		Request->CompleteWithExistingResponseData(Response);
+		return;
 	}
 
-	//if we didn't succeed, make sure we don't have any stale partial URL Map entries for this request
-	if (!bWasGenerateSuccess)
 	{
-		RemoveURLMapEntriesForRequest(Request);
+		FRWScopeLock ActiveScopeLock(ActiveRequestLock, SLT_Write);
+		ActiveRequests.Add(Request);
 	}
 
-	return bWasGenerateSuccess;
+	++NumCurrentlyActiveRequests;
 }
 
-void FApplePlatformBackgroundHttpManager::RemoveURLMapEntriesForRequest(FAppleBackgroundHttpRequestPtr Request)
+void FApplePlatformBackgroundHttpManager::RemoveRequest(const FBackgroundHttpRequestPtr GenericRequest)
 {
-    FRWScopeLock ScopeLock(URLToRequestMapLock, SLT_Write);
-    for (const FString& URL : Request->GetURLList())
-    {
-        FBackgroundHttpURLMappedRequestPtr& FoundRequest = URLToRequestMap.FindOrAdd(URL);
-        
-        if (FoundRequest == Request)
-        {
-            UE_LOG(LogBackgroundHttpManager, Verbose, TEXT("Removing URL Entry -- RequestDebugID:%s | URL:%s"), *Request->GetRequestDebugID(), *URL);
-            URLToRequestMap.Remove(URL);
-        }
-    }
+	// UpdateProgress might call RemoveRequest, so we could deadlock on ActiveRequestLock if we try to edit ActiveRequests
+	FRWScopeLock PendingScopeRemoveRequestLock(PendingRemoveRequestLock, SLT_Write);
+	PendingRemoveRequests.Add(GenericRequest);
 }
 
-void FApplePlatformBackgroundHttpManager::StartRequest(FAppleBackgroundHttpRequestPtr Request)
+void FApplePlatformBackgroundHttpManager::SetCellularPreference(int32 Value)
 {
-    //Just count it as a retry that won't increment the retry counter before giving us the URL as our RetryCount 0 should start this up.
-    RetryRequest(Request,false, nullptr);
-}
-
-void FApplePlatformBackgroundHttpManager::RemoveRequest(const FBackgroundHttpRequestPtr Request)
-{
-	FAppleBackgroundHttpRequestPtr AppleRequest = StaticCastSharedPtr<FApplePlatformBackgroundHttpRequest>(Request);
-	if (AppleRequest.IsValid())
-	{
-		RemoveSessionTasksForRequest(AppleRequest);
-	}
-	
-    RequestsPendingRemove.Add(Request);
-}
-
-void FApplePlatformBackgroundHttpManager::DeletePendingRemoveRequests()
-{
-    //Don't want to do this when background tasks might be using our request
-    ensureAlwaysMsgf(IsInGameThread(), TEXT("Called from un-expected thread! Potential error in an implementation of background downloads!"));
-    
-    for (const FBackgroundHttpRequestPtr& Request : RequestsPendingRemove)
-    {
-        FBackgroundHttpManagerImpl::RemoveRequest(Request);
-    }
-    
-    RequestsPendingRemove.Empty();
-}
-
-void FApplePlatformBackgroundHttpManager::RemoveSessionTasksForRequest(FAppleBackgroundHttpRequestPtr Request)
-{
-    //First remove map entries. That way we won't send a completion handler when we cancel
-    RemoveURLMapEntriesForRequest(Request);
-    
-    //Now cancel our active task
-    Request->CancelActiveTask();
-}
-
-void FApplePlatformBackgroundHttpManager::SetMaxActiveDownloads(int InMaxActiveDownloads)
-{
-	FBackgroundHttpManagerImpl::SetMaxActiveDownloads(InMaxActiveDownloads);
-
-	// It's possible we have more than the new maximum active right now, so we gracefully reduce MaxNumActualTasks down to MaxActiveDownloads as the extra tasks finish.
-	MaxNumActualTasks = FMath::Max(MaxActiveDownloads.Load(), FPlatformAtomics::AtomicRead(&NumCurrentlyActiveTasks));
-}
-
-bool FApplePlatformBackgroundHttpManager::AssociateWithAnyExistingUnAssociatedTasks(const FBackgroundHttpRequestPtr Request)
-{
-    bool bDidAssociateWithUnAssociatedTask = false;
-    
-    if (!bHasFinishedPopulatingUnassociatedTasks)
-    {
-        //@TODO: TRoss, might want to look at trying to associate these after bHasFinishedPopulatingUnassociatedTasks. At this point I don't think its really worth it as it SHOULD be done before we get here, however
-        //the PopulateUnAssociatedTasks() function has an asynch component so it could technically be unfinished with some tight timing.
-        UE_LOG(LogBackgroundHttpManager, Warning, TEXT("Call to AssociateWithAnyExistingRequest before we have finished populating unassociated tasks! Might have an unassociated task for this request that we won't associate with."));
-    }
-        
-    //See if our request is an AppleBackgroundHttpRequest so we can do more detailed checks on it.
-    FAppleBackgroundHttpRequestPtr AppleRequest = StaticCastSharedPtr<FApplePlatformBackgroundHttpRequest>(Request);
-    if (AppleRequest.IsValid())
-    {
-        bDidAssociateWithUnAssociatedTask = CheckForExistingUnAssociatedTask(AppleRequest);
-    }
-    
-    return bDidAssociateWithUnAssociatedTask;
-}
-
-bool FApplePlatformBackgroundHttpManager::CheckForExistingUnAssociatedTask(const FAppleBackgroundHttpRequestPtr Request)
-{
-    bool bDidFindExistingTask = false;
-    
-	TArray<FString> URLsToRemove;
-	
-	//Go through and read the UnAssociatedTasksLock and save off which URLs we need to remove as we associate with those tasks.
-	{
-		FRWScopeLock ScopeLock(UnAssociatedTasksLock, SLT_ReadOnly);
-		
-		if (ensureAlwaysMsgf(Request.IsValid(), TEXT("CheckForExistingUnAssociatedTask called with invalid Request!")))
-		{
-			const TArray<FString>& URLList = Request->GetURLList();
-			for (const FString& URL : URLList)
-			{
-				NSURLSessionTask* FoundTask = [UnAssociatedTasks valueForKey:URL.GetNSString()];
-				if (nullptr != FoundTask)
-				{
-					if ([FoundTask state] != NSURLSessionTaskStateCompleted && [FoundTask state] != NSURLSessionTaskStateCanceling)
-					{
-						UE_LOG(LogBackgroundHttpManager, Display, TEXT("Existing UnAssociateTask found for Request! Attempting to Associate! RequestDebugID:%s"), *(Request->GetRequestDebugID()));
-						
-						//Associate with task so that our Request takes over ownership of this task so we can remove it from our UnAssociated Tasks list without it getting GC'd
-						if (Request->AssociateWithTask(FoundTask))
-						{
-							//Always set our bWasTaskStartedInBG flag on our Request as true in the UnAssociated case as we don't know when it was really started
-							FPlatformAtomics::InterlockedExchange(&(Request->bWasTaskStartedInBG), true);
-
-							//Suspend task in case it was running so that we can adhere to our desired platform max tasks
-							[FoundTask suspend];
-
-							bDidFindExistingTask = true;
-							break;
-						}
-						else
-						{
-							UE_LOG(LogBackgroundHttpManager, Display, TEXT("UnAssociatedTask for request found, but failed to Associate with Task! -- RequestDebugID:%s | URL:%s"), *(Request->GetRequestDebugID()), *URL);
-						}
-					}
-					else
-					{
-						UE_LOG(LogBackgroundHttpManager, Display, TEXT("UnAssociatedTask for request found, BUT NOT USING as it was cancelling or completed already! -- RequestDebugID:%s | URL:%s"), *(Request->GetRequestDebugID()), *URL);
-					}
-					
-					//Still want to remove UnAssociatedTask even though we didn't use it as something else can now be downloading this data and we do not want duplicates
-					URLsToRemove.Add(URL);
-				}
-			}
-		}
-	}
-	
-	//Remove all URLs from UnAssociatedTasks
-	{
-		FRWScopeLock ScopeLock(UnAssociatedTasksLock, SLT_Write);
-		
-		for (const FString& URL : URLsToRemove)
-		{
-			[UnAssociatedTasks removeObjectForKey : URL.GetNSString()];
-		}
-	}
-	    
-    return bDidFindExistingTask;
-}
-
-void FApplePlatformBackgroundHttpManager::SetupNSURLSessionResponseDelegates()
-{
-	OnApp_EnteringBackgroundHandle = FCoreDelegates::ApplicationWillEnterBackgroundDelegate.AddRaw(this, &FApplePlatformBackgroundHttpManager::OnApp_EnteringBackground);
-	OnApp_EnteringForegroundHandle = FCoreDelegates::ApplicationHasEnteredForegroundDelegate.AddRaw(this, &FApplePlatformBackgroundHttpManager::OnApp_EnteringForeground);
-	OnTask_DidFinishDownloadingToURLHandle = FIOSBackgroundDownloadCoreDelegates::OnIOSBackgroundDownload_DidFinishDownloadingToURL.AddRaw(this, &FApplePlatformBackgroundHttpManager::OnTask_DidFinishDownloadingToURL);
-	OnTask_DidWriteDataHandle = FIOSBackgroundDownloadCoreDelegates::OnIOSBackgroundDownload_DidWriteData.AddRaw(this, &FApplePlatformBackgroundHttpManager::OnTask_DidWriteData);
-	OnTask_DidCompleteWithErrorHandle = FIOSBackgroundDownloadCoreDelegates::OnIOSBackgroundDownload_DidCompleteWithError.AddRaw(this, &FApplePlatformBackgroundHttpManager::OnTask_DidCompleteWithError);
-	OnSession_SessionDidFinishAllEventsHandle = FIOSBackgroundDownloadCoreDelegates::OnIOSBackgroundDownload_SessionDidFinishAllEvents.AddRaw(this, &FApplePlatformBackgroundHttpManager::OnSession_SessionDidFinishAllEvents);
-}
-
-void FApplePlatformBackgroundHttpManager::CleanUpNSURLSessionResponseDelegates()
-{
-	FCoreDelegates::ApplicationWillEnterBackgroundDelegate.Remove(OnApp_EnteringBackgroundHandle);
-	FCoreDelegates::ApplicationHasEnteredForegroundDelegate.Remove(OnApp_EnteringForegroundHandle);
-	FIOSBackgroundDownloadCoreDelegates::OnIOSBackgroundDownload_DidFinishDownloadingToURL.Remove(OnTask_DidFinishDownloadingToURLHandle);
-	FIOSBackgroundDownloadCoreDelegates::OnIOSBackgroundDownload_DidWriteData.Remove(OnTask_DidWriteDataHandle);
-	FIOSBackgroundDownloadCoreDelegates::OnIOSBackgroundDownload_DidCompleteWithError.Remove(OnTask_DidCompleteWithErrorHandle);
-	FIOSBackgroundDownloadCoreDelegates::OnIOSBackgroundDownload_SessionDidFinishAllEvents.Remove(OnSession_SessionDidFinishAllEventsHandle);
-}
-
-void FApplePlatformBackgroundHttpManager::OnApp_EnteringForeground()
-{
-	PauseAllActiveTasks();
-    FPlatformAtomics::InterlockedExchange(&bIsInBackground,false);
-}
-
-void FApplePlatformBackgroundHttpManager::OnApp_EnteringBackground()
-{
-	FPlatformAtomics::InterlockedExchange(&bIsInBackground,true);
-	ResumeTasksForBackgrounding();
-}
-
-void FApplePlatformBackgroundHttpManager::PauseAllActiveTasks()
-{
-    UE_LOG(LogBackgroundHttpManager, Display, TEXT("Attempting to Pause All Active Tasks"));
-    
-	NSURLSession* BackgroundDownloadSession = FBackgroundURLSessionHandler::GetBackgroundSession();
-	if (nullptr != BackgroundDownloadSession)
-	{
-		[BackgroundDownloadSession getTasksWithCompletionHandler:^(NSArray<__kindof NSURLSessionDataTask*>* DataTasks, NSArray<__kindof NSURLSessionUploadTask*>* UploadTasks, NSArray<__kindof NSURLSessionDownloadTask*>* DownloadTasks)
-		{
-			for (NSURLSessionDownloadTask* DownloadTask : DownloadTasks)
-			{
-				if ([DownloadTask state] == NSURLSessionTaskStateRunning)
-				{
-                    FString TaskURL = [[[DownloadTask currentRequest] URL] absoluteString];
-                    int TaskIdentifier = (int)[DownloadTask taskIdentifier];
-                    
-                    UE_LOG(LogBackgroundHttpManager, Display, TEXT("Pausing Task for URL:%s | TaskIdentifier:%d"), *TaskURL, TaskIdentifier);
-                    
-					[DownloadTask suspend];
-				}
-			}
-		}];
-	}
+	FBackgroundURLSessionHandler::AllowCellular(Value > 0);
 }
 
-void FApplePlatformBackgroundHttpManager::ResumeTasksForBackgrounding(FIOSBackgroundHttpPostSessionWorkCallback Callback)
-{
-    UE_LOG(LogBackgroundHttpManager, Display, TEXT("Attempting to resume all active tasks that don't have a corresponding request that has paused them in priority order"));
-    
-	NSURLSession* BackgroundDownloadSession = FBackgroundURLSessionHandler::GetBackgroundSession();
-	if (nullptr != BackgroundDownloadSession)
-	{
-		[BackgroundDownloadSession getTasksWithCompletionHandler : ^ (NSArray<__kindof NSURLSessionDataTask*>* DataTasks, NSArray<__kindof NSURLSessionUploadTask*>* UploadTasks, NSArray<__kindof NSURLSessionDownloadTask*>* DownloadTasks)
-		{
-			FRWScopeLock ScopeLock(URLToRequestMapLock, SLT_ReadOnly);
-		 
-			//We only want to automatically re-queue the highest priority things, so go through all tasks by priority order and stop once we have requeued something
-			bool bDidResumeATask = false;
-			for (uint8 PriorityAsInt = (uint8)EBackgroundHTTPPriority::High; (!bDidResumeATask && (PriorityAsInt < (uint8)EBackgroundHTTPPriority::Num)); ++PriorityAsInt)
-			{
-				for (NSURLSessionDownloadTask* DownloadTask : DownloadTasks)
-				{
-					FString TaskURL = [[[DownloadTask currentRequest] URL] absoluteString];
-					int TaskIdentifier = (int)[DownloadTask taskIdentifier];
-		 
-					FBackgroundHttpURLMappedRequestPtr* WeakRequestInMap = URLToRequestMap.Find(TaskURL);
-					FAppleBackgroundHttpRequestPtr FoundRequest = ((nullptr != WeakRequestInMap) && (WeakRequestInMap->IsValid())) ? WeakRequestInMap->Pin() : nullptr;
-		 
-					if (FoundRequest.IsValid())
-					{
-						bDidResumeATask = ResumeDownloadTaskForBackgroundingIfAppropriate(DownloadTask, FoundRequest, EBackgroundHTTPPriority(PriorityAsInt)) || bDidResumeATask;
-					}
-					else
-					{
-						UE_LOG(LogBackgroundHttpManager, Verbose, TEXT("Skipped Resuming Task as there is no corresponding request! URL:%s | TaskIdentifier:%d"), *TaskURL, TaskIdentifier);
-					}
-				}
-				
-				if (bDidResumeATask)
-				{
-					UE_LOG(LogBackgroundHttpManager, Display, TEXT("Tasks of Priority: %s Found. Not Resuming Any Lower Priority Tasks."), LexToString((EBackgroundHTTPPriority)PriorityAsInt));
-				}
-				else
-				{
-					UE_LOG(LogBackgroundHttpManager, Display, TEXT("No tasks found of priority %s."), LexToString((EBackgroundHTTPPriority)PriorityAsInt));
-				}
-			}
-			
-			Callback.ExecuteIfBound();
-		}];
-	}
-}
-
-bool FApplePlatformBackgroundHttpManager::ResumeDownloadTaskForBackgroundingIfAppropriate(NSURLSessionDownloadTask* DownloadTask, const FAppleBackgroundHttpRequestPtr Request, EBackgroundHTTPPriority LowestPriorityToQueue)
-{
-	if ((nullptr == DownloadTask) || (nullptr == [DownloadTask currentRequest]) || (!Request.IsValid()))
-	{
-		return false;
-	}
-
-	bool bDidResumeTask = false;
-	bool bDidFindActiveTaskOfPriority = false;
-	
-	FString TaskURL = [[[DownloadTask currentRequest] URL] absoluteString];
-	int TaskIdentifier = (int)[DownloadTask taskIdentifier];
-
-	EBackgroundHTTPPriority FoundRequestPriority = Request->GetRequestPriority();
-	if (FoundRequestPriority <= LowestPriorityToQueue)
-	{
-		const bool bIsRequestPaused = Request->bIsTaskPaused;
-		if (!bIsRequestPaused)
-		{
-			UE_LOG(LogBackgroundHttpManager, Display, TEXT("Resuming Found Task for URL:%s | TaskIdentifier:%d | TaskPriority:%s"), *TaskURL, TaskIdentifier, LexToString(FoundRequestPriority));
-			
-			const bool bIsTaskActive = Request->IsUnderlyingTaskActive();
-			
-			//We only want to resume Suspended tasks
-			if ([DownloadTask state] == NSURLSessionTaskStateSuspended)
-			{
-				[DownloadTask resume];
-				bDidResumeTask = true;
-			}
-			else if (bIsTaskActive)
-			{
-				bDidFindActiveTaskOfPriority = true;
-			}
-		}
-		else
-		{
-			UE_LOG(LogBackgroundHttpManager, Verbose, TEXT("NOT Resuming Task for URL as associated request was paused! URL:%s | TaskIdentifier:%d"), *TaskURL, TaskIdentifier);
-		}
-	}
-	else
-	{
-		UE_LOG(LogBackgroundHttpManager, Verbose, TEXT("NOT RESUMING Task for URL because its a lower priority:%s | TaskIdentifier:%d | TaskPriority:%s | LowestPriorityToQueue:%s"), *TaskURL, TaskIdentifier, LexToString(FoundRequestPriority), LexToString(LowestPriorityToQueue));
-	}
-	
-	return (bDidResumeTask || bDidFindActiveTaskOfPriority);
-}
-
-void FApplePlatformBackgroundHttpManager::OnTask_DidFinishDownloadingToURL(NSURLSessionDownloadTask* Task, NSError* Error, const FString& TempFilePath)
-{
-	FString TaskURL = [[[Task currentRequest] URL] absoluteString];
-    int TaskIdentifier = (int)[Task taskIdentifier];
-    
-    const int ErrorCode = [Error code];
-    const FString ErrorDescription = [Error localizedDescription];
-    
-    IPlatformFile& PlatformFile = FPlatformFileManager::Get().GetPlatformFile();
-	const bool bFileExists = PlatformFile.FileExists(*TempFilePath);
-	
-    UE_LOG(LogBackgroundHttpManager, Display, TEXT("Received Background Session Callback for URL:%s | TaskIdentifier:%d | bFileExists:%d | ErrorCode:%d | ErrorDescription:%s | Location:%s"), *TaskURL, TaskIdentifier, (int)(bFileExists), ErrorCode, *ErrorDescription, *TempFilePath);
-
-    if (bFileExists)
-    {
-        //Find request for this task and mark it complete
-        {
-            FRWScopeLock ScopeLock(URLToRequestMapLock, SLT_ReadOnly);
-            FBackgroundHttpURLMappedRequestPtr* WeakRequestInMap = URLToRequestMap.Find(TaskURL);
-            FAppleBackgroundHttpRequestPtr FoundRequest = ((nullptr != WeakRequestInMap) && (WeakRequestInMap->IsValid())) ? WeakRequestInMap->Pin() : nullptr;
-            
-            if (FoundRequest.IsValid())
-            {
-                FoundRequest->SetRequestAsSuccess(TempFilePath);
-            }
-            
-            UE_LOG(LogBackgroundHttpManager, Display, TEXT("Attempt To Mark Task Complete -- URL:%s | TaskIdentifier:%d |bDidFindTask:%d"), *TaskURL, TaskIdentifier, (int)(FoundRequest.IsValid()));
-        }
-    }
-    else
-    {
-        //Forward to the OnCompleteWithError as we don't have our finished file!
-        UE_LOG(LogBackgroundHttpManager, Display, TEXT("File Not Found For DidFinishDownloadingToURL. Transitioning to DidCompleteWithError -- TaskURL:%s | TaskIdentifier:%d| ErrorCode:%d | ErrorDescription:%s | Location:%s"), *TaskURL, TaskIdentifier, ErrorCode, *ErrorDescription, *TempFilePath);
-        OnTask_DidCompleteWithError(Task, Error);
-    }
-}
-
-void FApplePlatformBackgroundHttpManager::FinishRequest(FAppleBackgroundHttpRequestPtr Request)
+bool FApplePlatformBackgroundHttpManager::Tick(float DeltaTime)
 {
-    //We should only come into here from the GameThread so that if we send out a complete event our delegate subscribers don't have to worry about being thread-safe unnessecarily
-    ensureAlwaysMsgf(IsInGameThread(), TEXT("Called from un-expected thread! Potential error in an implementation of background downloads!"));
-    
-    //Make sure we another thread hasn't already finished this request
-    bool bHasAlreadyFinishedRequest = FPlatformAtomics::InterlockedExchange(&(Request->bHasAlreadyFinishedRequest), true);
-    if (!bHasAlreadyFinishedRequest)
-    {
-        //by default we will be finishing this request in this function, but some errors might prompt a retry out of this function
-        bool bIsRequestActuallyFinished = true;
-        
-        if (ensureAlwaysMsgf(Request.IsValid(), TEXT("Call to FinishRequest with invalid request!")))
-        {
-            const FString& TempFilePath = Request->CompletedTempDownloadLocation;
-            IPlatformFile& PlatformFile = FPlatformFileManager::Get().GetPlatformFile();
-            const bool bFileExists = PlatformFile.FileExists(*TempFilePath);
-
-            int ResponseCode = bFileExists ? EHttpResponseCodes::Created : EHttpResponseCodes::Unknown;
-
-            if (bFileExists)
-            {
-                UE_LOG(LogBackgroundHttpManager,Display, TEXT("Task Completed Successfully. RequestDebugID:%s TempFileLocation:%s"), *(Request->GetRequestDebugID()), *TempFilePath);
-                FBackgroundHttpResponsePtr NewResponse = FPlatformBackgroundHttp::ConstructBackgroundResponse(ResponseCode, *TempFilePath);
-                Request->CompleteWithExistingResponseData(NewResponse);
-            }
-            else
-            {
-                volatile bool bDidFail = FPlatformAtomics::AtomicRead(&(Request->bIsFailed));
-                
-                //Unexpected case where we didn't find a valid download, but we thought this completed
-                //successfully. Handle this unexpected failure by trying to retry the task.
-                if (!bDidFail)
-                {
-                    UE_LOG(LogBackgroundHttpManager,Error, TEXT("Task finished downloading, but finished temp file was not found! -- RequestDebugID:%s | TempFileLocation:%s"), *(Request->GetRequestDebugID()), *TempFilePath);
-                    
-                    //Mark our download as not completed as we hit an error so that we don't just keep trying to call FinishRequest
-                    FPlatformAtomics::InterlockedExchange(&(Request->bIsCompleted), false);
-                    FPlatformAtomics::InterlockedExchange(&(Request->bHasAlreadyFinishedRequest), false);
-                    
-                    //Just cancel the task. This will lead to it getting a callback to OnTask_DidCompleteWithError where we will re-create it
-                    Request->CancelActiveTask();
-                    
-                    bIsRequestActuallyFinished = false;
-                }
-                //Expected case where we failed, but expected to fail
-                else
-                {
-                    UE_LOG(LogBackgroundHttpManager, Display, TEXT("Task failed completely -- RequestDebugID:%s"), *(Request->GetRequestDebugID()));
-                    
-                    FBackgroundHttpResponsePtr NewResponse = FPlatformBackgroundHttp::ConstructBackgroundResponse(ResponseCode, TEXT(""));
-                    Request->CompleteWithExistingResponseData(NewResponse);
-                }
-            }
-            
-            //If we are actually finishing this request, lets decrement our NumCurrentlyActiveTasks counter
-            if (bIsRequestActuallyFinished)
-            {
-				//only decrement NumCurrentlyActiveTasks if this was an actual active task,
-				//can still be completed in the BG or from existing completed data without being active
-                const bool bIsTaskActive = Request->IsUnderlyingTaskActive();
-                if (bIsTaskActive)
-                {
-                    int NumActualTasks = FPlatformAtomics::InterlockedDecrement(&NumCurrentlyActiveTasks);
-
-					// Handle the case that SetMaxActiveDownloads reduced the MaxActiveDownloads while we had more than the new maximum in progress.
-					MaxNumActualTasks = FMath::Max(MaxNumActualTasks - 1, MaxActiveDownloads.Load());
-
-                    //Sanity check that our data is valid. Shouldn't ever trip if everything is working as intended.
-                    const bool bNumActualTasksIsValid = ((NumActualTasks >= 0) && (NumActualTasks <= MaxNumActualTasks));
-                    
-                    UE_LOG(LogBackgroundHttpManager,Display, TEXT("Finishing Request lowering Task Count: %d"), NumActualTasks);
-                    
-                    ensureMsgf(bNumActualTasksIsValid, TEXT("Number of Requests we think are active is invalid! -- NumCurrentlyActiveTasks:%d"), NumActualTasks);
-                }
-            }
-        }
-    }
-    else
-    {
-        UE_LOG(LogBackgroundHttpManager,Display, TEXT("Not finishing Request as its already sending a finish notification -- RequestDebugID:%s"), *(Request->GetRequestDebugID()));
-    }
-}
+	QUICK_SCOPE_CYCLE_COUNTER(STAT_FApplePlatformBackgroundHttpManager_Tick);
 
-void FApplePlatformBackgroundHttpManager::RetryRequest(FAppleBackgroundHttpRequestPtr Request, bool bShouldIncreaseRetryCount, NSData* RetryData)
-{
-	NSURLSessionDownloadTask* NewTask = nullptr;
+	ensureAlwaysMsgf(IsInGameThread(), TEXT("Called from un-expected thread! Potential error in an implementation of background downloads!"));
 
-	if (ensureAlwaysMsgf(Request.IsValid(), TEXT("Call to RetryRequest with an invalid request!")))
 	{
-		NSURLSession* BackgroundDownloadSession = FBackgroundURLSessionHandler::GetBackgroundSession();
-		if (ensureAlwaysMsgf((nullptr != BackgroundDownloadSession), TEXT("Invalid Background Download NSURLSession during RetryRequest! Should have already Initialized the NSURLSession by this point!")))
+		FRWScopeLock ScopeLock(ActiveRequestLock, SLT_ReadOnly);
+		for (FBackgroundHttpRequestPtr& GenericRequest : ActiveRequests)
 		{
-			//First, lets see if we should base this task of existing RetryData
-			const bool bShouldUseRetryData = ShouldUseRequestRetryData(Request, RetryData);
-			if (bShouldUseRetryData)
+			FAppleBackgroundHttpRequestPtr Request = StaticCastSharedPtr<FApplePlatformBackgroundHttpRequest>(GenericRequest);
+			if (ensureAlwaysMsgf(Request.IsValid(), TEXT("Invalid Request Pointer in ActiveRequests list!")))
 			{
-				UE_LOG(LogBackgroundHttpManager, Display, TEXT("Resuming Task With Resume Data -- RequestDebugID:%s | RetryData Length:%d"), *(Request->GetRequestDebugID()), [RetryData length]);
-				NewTask = [BackgroundDownloadSession downloadTaskWithResumeData:RetryData];
-			}
-
-			//If not retry data, lets try and just retry on the next CDN
-			if (nullptr == NewTask)
-			{
-                //Since we created a new task instead of using retry data, reset resume data's retry count on the request
-                Request->ResumeDataRetryCount.Reset();
-                
-				const FString& NewRetryURL = Request->GetURLForRetry(bShouldIncreaseRetryCount);
-				const bool bShouldStartNewRequest = !NewRetryURL.IsEmpty();
-				if (bShouldStartNewRequest)
-				{
-					NSURL* URL = [NSURL URLWithString:NewRetryURL.GetNSString()];
-					NewTask = [BackgroundDownloadSession downloadTaskWithURL:URL];
-				}
-			}
-
-			if (nullptr != NewTask)
-			{
-				Request->AssociateWithTask(NewTask);
-
-				//If we are in BG activate right now without waiting for the FG tick
-				volatile bool bCopyOfBGState = FPlatformAtomics::AtomicRead(&bIsInBackground);
-				
-				//We also want to re-activate any already activated tasks
-				const bool bIsTaskActive = Request->IsUnderlyingTaskActive();
-				
-				if (!Request->bIsTaskPaused && (bCopyOfBGState || bIsTaskActive))
-				{
-					Request->ActivateUnderlyingTask();
-				}
-
-                UE_LOG(LogBackgroundHttpManager, Display, TEXT("Created Task for Request -- RequestDebugID:%s | bIsAppInBG:%d | bIsPaused:%d"), *(Request->GetRequestDebugID()), (int)bCopyOfBGState, (int)Request->bIsTaskPaused);
-
-				//Always set our bWasTaskStartedInBG flag on our Request so we will know if we need to restart this task next FG Tick.
-				FPlatformAtomics::InterlockedExchange(&(Request->bWasTaskStartedInBG), bCopyOfBGState);
-			}
-			else
-			{
-				UE_LOG(LogBackgroundHttpManager, Display, TEXT("Marking Request Failed. Out of Retries -- RequestDebugID:%s | bShouldUseRetryData:%d"), *(Request->GetRequestDebugID()), (int)bShouldUseRetryData);
-				Request->SetRequestAsFailed();
+				Request->UpdateProgress();
 			}
 		}
 	}
-}
 
-bool FApplePlatformBackgroundHttpManager::ShouldUseRequestRetryData(FAppleBackgroundHttpRequestPtr Request, NSData* RetryData) const
-{
-	bool bShouldUseData = false;
-
-	if (ensureAlwaysMsgf(Request.IsValid(), TEXT("Call to ShouldUseRequestRetryData with an invalid request!")))
 	{
-		if (IsRetryDataValid(RetryData))
+		FRWScopeLock PendingScopeRemoveRequestLock(PendingRemoveRequestLock, SLT_Write);
+		for (FBackgroundHttpRequestPtr& GenericRequest : PendingRemoveRequests)
 		{
-			const int CurrentResumeDataRetryCount = Request->ResumeDataRetryCount.Increment();
-			if ((RetryResumeDataLimitSetting < 0) || (CurrentResumeDataRetryCount <= RetryResumeDataLimitSetting))
+			FAppleBackgroundHttpRequestPtr Request = StaticCastSharedPtr<FApplePlatformBackgroundHttpRequest>(GenericRequest);
+			if (ensureAlwaysMsgf(Request.IsValid(), TEXT("Adding a non-Apple background request to our Apple Background Http Manager2! This is not supported or expected!")))
 			{
-				bShouldUseData = true;
-			}
-		}
-	}
-    
-    return bShouldUseData;
-}
-
-bool FApplePlatformBackgroundHttpManager::IsRetryDataValid(NSData* RetryData) const
-{
-    return ((nullptr != RetryData) && ([RetryData length] > 0));
-}
-
-void FApplePlatformBackgroundHttpManager::OnTask_DidWriteData(NSURLSessionDownloadTask* Task, int64_t BytesWrittenSinceLastCall, int64_t TotalBytesWritten, int64_t TotalBytesExpectedToWrite)
-{
-    if (ensureAlwaysMsgf((nullptr != Task), TEXT("Call to DidWriteData with invalid Task!")))
-    {
-        FString TaskURL = [[[Task currentRequest] URL] absoluteString];
-        int TaskIdentifier = (int)[Task taskIdentifier];
-        
-        //Find task and update it's download progress
-        {
-            FRWScopeLock ScopeLock(URLToRequestMapLock, SLT_ReadOnly);
-            FBackgroundHttpURLMappedRequestPtr* WeakRequestInMap = URLToRequestMap.Find(TaskURL);
-            FAppleBackgroundHttpRequestPtr FoundRequest = ((nullptr != WeakRequestInMap) && (WeakRequestInMap->IsValid())) ? WeakRequestInMap->Pin() : nullptr;
-            
-            if (FoundRequest.IsValid())
-            {
-                if (FoundRequest->DownloadProgress < TotalBytesWritten)
-                {
-                    int64 DownloadProgress = FPlatformAtomics::AtomicRead(&(FoundRequest->DownloadProgress));
-                    UE_LOG(LogBackgroundHttpManager, Verbose, TEXT("Updating Task Progress! -- RequestDebugID:%s | TaskIdentifier:%d | Current Progress:%lld | New Progress:%lld"), *(FoundRequest->GetRequestDebugID()), TaskIdentifier, DownloadProgress, TotalBytesWritten);
-                }
-                else
-                {
-                    ensureAlwaysMsgf(false, TEXT("Download Progress tried to go down not up unexpectidly! This could mean a task was unknowingly duplicated! -- RequestDebugID:%s | TaskIdentifier:%d | Current Progress:%lld | New Progress:%lld"), *(FoundRequest->GetRequestDebugID()), TaskIdentifier, FoundRequest->DownloadProgress, TotalBytesWritten);
-                }
-
-                FoundRequest->UpdateDownloadProgress(TotalBytesWritten, BytesWrittenSinceLastCall);
-            }
-        }
-    }
-}
+				const uint64 DownloadId = Request->GetInternalDownloadId();
+				FBackgroundURLSessionHandler::CancelDownload(DownloadId);
+				Request->SetInternalDownloadId(FBackgroundURLSessionHandler::InvalidDownloadId);
 
-void FApplePlatformBackgroundHttpManager::OnTask_DidCompleteWithError(NSURLSessionTask* Task, NSError* Error)
-{
-	if (ensureAlwaysMsgf((nullptr != Task), TEXT("Call to OnTask_DidCompleteWithError delegate with an invalid task!")))
-	{
-		FString TaskURL = [[[Task currentRequest] URL] absoluteString];
-        int TaskIdentifier = (int)[Task taskIdentifier];
-        
-		const bool bDidCompleteWithError = (nullptr != Error);
-		const int ErrorCode = [Error code];
-        const FString ErrorDescription = [Error localizedDescription];
-        
-		NSData* ResumeData = Error ? [Error.userInfo objectForKey:NSURLSessionDownloadTaskResumeData] : nullptr;
-		const bool bHasResumeData = (ResumeData && ([ResumeData  length] > 0));
-
-		NSNumber* CancelledReasonKey = [Error.userInfo objectForKey:NSURLErrorBackgroundTaskCancelledReasonKey];
-		int CancelledReasonInt = (nullptr != CancelledReasonKey) ? [CancelledReasonKey intValue] : -1;
-
-		FString DebugRetryOverrideReason;
-
-		//We still come into the function when tasks complete successfully. Only handle actual errors
-		if (bDidCompleteWithError)
-		{
-            FRWScopeLock ScopeLock(URLToRequestMapLock, SLT_ReadOnly);
-            FBackgroundHttpURLMappedRequestPtr* WeakRequestInMap = URLToRequestMap.Find(TaskURL);
-            FAppleBackgroundHttpRequestPtr FoundRequest = ((nullptr != WeakRequestInMap) && (WeakRequestInMap->IsValid())) ? WeakRequestInMap->Pin() : nullptr;
-            const bool bDidFindValidRequest = FoundRequest.IsValid();
-         
-            //by default increase error count. Special cases below will overwrite this
-            bool bShouldRetryIncreaseRetryCount = true;
-            
-            //If we don't have internet, we don't want to move through our CDNs, but rather chain recreate download tasks until we regain internet
-            if ([Error code] == NSURLErrorNotConnectedToInternet)
-            {
-                bShouldRetryIncreaseRetryCount = false;
-                DebugRetryOverrideReason = TEXT("Not Connected To Internet");
-            }
-            
-            UE_LOG(LogBackgroundHttpManager, Display, TEXT("DidCompleteWithError for Task. -- URL:%s | TaskIdentifier:%d | bDidFindVaildRequest:%d | bDidCompleteWithError:%d | ErrorCode:%d | bHasResumeData:%d | CancelledReasonKey:%d | RetryOverrideReason:%s | bShouldRetryIncreaseRetryCount:%d | ErrorDescription:%s"), *TaskURL, TaskIdentifier, (int)bDidFindValidRequest, (int)bDidCompleteWithError, ErrorCode, (int)bHasResumeData, CancelledReasonInt, *DebugRetryOverrideReason, (int)bShouldRetryIncreaseRetryCount, *ErrorDescription);
-            
-			if (bDidFindValidRequest)
-			{
-				RetryRequest(FoundRequest, bShouldRetryIncreaseRetryCount, ResumeData);
-			}
-			else
-			{
-                //This can be a valid case because of UnAssociatedTasks, so don't error here
-                UE_LOG(LogBackgroundHttpManager, Display, TEXT("No request for completing task! -- TaskURL:%s | TaskIdentifier:%d"), *TaskURL, TaskIdentifier);
+				FBackgroundHttpManagerImpl::RemoveRequest(Request);
 			}
 		}
-	}
-}
 
-void FApplePlatformBackgroundHttpManager::OnSession_SessionDidFinishAllEvents(NSURLSession* Session, FIOSBackgroundDownloadCoreDelegates::FIOSBackgroundDownload_DelayedBackgroundURLSessionCompleteHandler Callback)
-{
-	//Let BackgroundURLSessionHandler know that it should wait until we call the callback
-	FBackgroundURLSessionHandler::AddDelayedBackgroundURLSessionComplete();
-	
-	const bool bCopyIsBG = FPlatformAtomics::AtomicRead(&bIsInBackground);
-	
-	UE_LOG(LogBackgroundHttpManager, Display, TEXT("NSURLSession done sending background events for all already queued tasks. bWasAppleBGHTTPInitialized:%d | bIsBG:%d"), bWasAppleBGHTTPInitialized, bCopyIsBG);
-	
-	//If we are in the BG, or if we have not yet been initialized, lets go ahead and just immediately send the callback
-	if (bWasAppleBGHTTPInitialized && bCopyIsBG)
-	{
-		//Now that we have finished all queued background tasks, lets resume tasks for any lower priorities that weren't started
-		ResumeTasksForBackgrounding(FIOSBackgroundDownloadCoreDelegates::FIOSBackgroundDownload_DelayedBackgroundURLSessionCompleteHandler::CreateLambda(
-		[=]()
-		{
-			Callback.ExecuteIfBound();
-		}));
-	}
-	else
-	{
-		Callback.ExecuteIfBound();
+		PendingRemoveRequests.Empty();
 	}
-}
 
-bool FApplePlatformBackgroundHttpManager::Tick(float DeltaTime)
-{
-	QUICK_SCOPE_CYCLE_COUNTER(STAT_FApplePlatformBackgroundHttpManager_Tick);
-
-    ensureAlwaysMsgf(IsInGameThread(), TEXT("Called from un-expected thread! Potential error in an implementation of background downloads!"));
-    
-    TickRequests(DeltaTime);
-    TickUnassociatedTasks(DeltaTime);
-    
 	GetFileHashHelper()->SaveData();
-	
-    //Always keep ticking
-    return true;
-}
 
-void FApplePlatformBackgroundHttpManager::TickRequests(float DeltaTime)
-{
-    ensureAlwaysMsgf(IsInGameThread(), TEXT("Called from un-expected thread! Potential error in an implementation of background downloads!"));
-    
-    //First lets go through all our Requests to see if we need to complete or recreate any requests
-    {
-		//Check to make sure we have room for more tasks to be active first
-		int CurrentCount = FPlatformAtomics::AtomicRead(&NumCurrentlyActiveTasks);
-		bool bNeedsMoreTasks = (CurrentCount < MaxActiveDownloads);
-		FAppleBackgroundHttpRequestPtr FoundRequestToStart = nullptr;
-		
-        FRWScopeLock ScopeLock(ActiveRequestLock, SLT_ReadOnly);
-        for (FBackgroundHttpRequestPtr& Request : ActiveRequests)
-        {
-            FAppleBackgroundHttpRequestPtr AppleRequest = StaticCastSharedPtr<FApplePlatformBackgroundHttpRequest>(Request);
-            if (ensureAlwaysMsgf(AppleRequest.IsValid(), TEXT("Invalid Request Pointer in ActiveRequests list!")))
-            {
-                const bool bIsTaskActive = AppleRequest->IsUnderlyingTaskActive();
-                const bool bIsTaskPaused = AppleRequest->IsUnderlyingTaskPaused();
-                const bool bIsTaskComplete = AppleRequest->IsTaskComplete();
-                const bool bWasStartedInBG = FPlatformAtomics::AtomicRead(&(AppleRequest->bWasTaskStartedInBG));
-                const bool bIsPendingCancel = FPlatformAtomics::AtomicRead(&(AppleRequest->bIsPendingCancel));
-                
-                UE_LOG(LogBackgroundHttpManager, VeryVerbose, TEXT("Checking Status of Request on Tick -- RequestDebugID::%s | bIsTaskComplete:%d | bWasStartedInBG:%d"), *(AppleRequest->GetRequestDebugID()), (int)bIsTaskComplete, (int)bWasStartedInBG);
-                
-                if (bIsTaskComplete)
-                {
-                    UE_LOG(LogBackgroundHttpManager, Display, TEXT("Calling FinishRequest On -- RequestDebugID::%s | bIsTaskComplete:%d | bWasStartedInBG:%d"), *(AppleRequest->GetRequestDebugID()), (int)bIsTaskComplete, (int)bWasStartedInBG);
-                    FinishRequest(AppleRequest);
-                }
-                else if (bWasStartedInBG && !bIsPendingCancel)
-                {
-                    UE_LOG(LogBackgroundHttpManager, Display, TEXT("Cancelling Request Created In BG To Re-Create In FG -- RequestDebugID:%s"), *(AppleRequest->GetRequestDebugID()));
-                    
-                    //reset to false so we don't run this twice while waiting on recreation
-                    FPlatformAtomics::InterlockedExchange(&(AppleRequest->bWasTaskStartedInBG), false);
-                    
-                    //Just cancel the task. This will lead to it getting a callback to OnTask_DidCompleteWithError where we will re-create it
-                    //We want to recreate any task spun up in the background as it will not respect our session settings if created in BG.
-                    AppleRequest->CancelActiveTask();
-                }
-                else if (bIsTaskActive && !bIsTaskPaused && !bIsPendingCancel)
-                {
-                    const bool bShouldTimeOut = AppleRequest->TickTimeOutTimer(DeltaTime);
-                    if (bShouldTimeOut)
-                    {
-                        UE_LOG(LogBackgroundHttpManager, Display, TEXT("Timing out Request Due To Lack of Server Response -- RequestDebugID:%s"), *(AppleRequest->GetRequestDebugID()));
-                        
-                        //Just cancel the task and let the OnTask_DidCompleteWithError callback handle retrying it if appropriate.
-                        AppleRequest->CancelActiveTask();
-                    }
-                }
-				else if (bNeedsMoreTasks && !bIsTaskActive && !bIsTaskPaused && !bIsPendingCancel)
-				{
-					//This task is a possible to start, so lets see if its the highest priority task we have found yet.
-					
-					if (!FoundRequestToStart.IsValid())
-					{
-						FoundRequestToStart = AppleRequest;
-					}
-					else
-					{
-						if (AppleRequest->GetRequestPriority() < FoundRequestToStart->GetRequestPriority())
-						{
-							FoundRequestToStart = AppleRequest;
-						}
-					}
-				}
-				
-				AppleRequest->SendDownloadProgressUpdate();
-            }
-        }
-		
-		if (FoundRequestToStart.IsValid())
-		{
-			const int32 OldTotal = FPlatformAtomics::InterlockedIncrement(&NumCurrentlyActiveTasks);
-			
-			if (OldTotal <= MaxActiveDownloads)
-			{
-				UE_LOG(LogBackgroundHttpManager, Display, TEXT("Starting Task for Request -- RequestDebugID:%s | CurrentlyActiveRequests:%d"), *(FoundRequestToStart->GetRequestDebugID()), OldTotal);
-				FoundRequestToStart->ActivateUnderlyingTask();
-			}
-			else
-			{
-				UE_LOG(LogBackgroundHttpManager, Log, TEXT("Request failed to activate as we passed the platform max from another requeset before we could resume. RequestDebugID:%s | CurrentlyActiveRequests:%d"), *(FoundRequestToStart->GetRequestDebugID()), OldTotal);
-				
-				//Don't activate and remove our increment from above because something put us over the limit before we resumed
-				FPlatformAtomics::InterlockedDecrement(&NumCurrentlyActiveTasks);
-			}
-		}
-    }
-    
-    //Now that we have gone through and finished all the requests, go ahead and delete any pending removes
-    DeletePendingRemoveRequests();
-}
-
-void FApplePlatformBackgroundHttpManager::TickUnassociatedTasks(float DeltaTime)
-{
-    ensureAlwaysMsgf(IsInGameThread(), TEXT("Called from un-expected thread! Potential error in an implementation of background downloads!"));
-    
-    //If we don't have anything queued, lets resume any un-associated tasks
-    int CurrentCount = FPlatformAtomics::AtomicRead(&NumCurrentlyActiveTasks);
-    if (CurrentCount == 0)
-    {
-        UnpauseAllUnassociatedTasks();
-    }
-    else
-    {
-        //we have something queued, lets pause unassociated tasks
-        PauseAllUnassociatedTasks();
-    }
-}
-
-//Make sure we are using the FBackgroundURLSessionHandler's version so that both our results are synced
-BackgroundHttpFileHashHelperRef FApplePlatformBackgroundHttpManager::GetFileHashHelper()
-{
-	return FBackgroundURLSessionHandler::GetFileHashHelper();
-}
-
-const BackgroundHttpFileHashHelperRef FApplePlatformBackgroundHttpManager::GetFileHashHelper() const
-{
-	return FBackgroundURLSessionHandler::GetFileHashHelper();
+	return true;
 }
diff --git a/Engine/Source/Runtime/Online/BackgroundHTTP/Private/IOS/ApplePlatformBackgroundHttpRequest.cpp b/Engine/Source/Runtime/Online/BackgroundHTTP/Private/IOS/ApplePlatformBackgroundHttpRequest.cpp
--- a/Engine/Source/Runtime/Online/BackgroundHTTP/Private/IOS/ApplePlatformBackgroundHttpRequest.cpp
+++ b/Engine/Source/Runtime/Online/BackgroundHTTP/Private/IOS/ApplePlatformBackgroundHttpRequest.cpp
@@ -2,282 +2,75 @@
 
 #include "IOS/ApplePlatformBackgroundHttpRequest.h"
 #include "IOS/ApplePlatformBackgroundHttpManager.h"
-#include "Interfaces/IBackgroundHttpResponse.h"
-#include "BackgroundHttpModule.h"
+#include "IOS/ApplePlatformBackgroundHttp.h"
+#include "IOS/IOSBackgroundURLSessionHandler.h"
+#include "HAL/PlatformFile.h"
+#include "HAL/PlatformFileManager.h"
 
 FApplePlatformBackgroundHttpRequest::FApplePlatformBackgroundHttpRequest()
-    : CompletedTempDownloadLocation()
-    , ActiveTimeOutTimer(30.f)
-    , RetryCount(0)
-	, ResumeDataRetryCount(0)
-    , FirstTask(nullptr)
-	, FirstTaskIdentifier(0)
-    , CombinedRequestID("")
-    , bIsTaskActive(false)
-    , bIsTaskPaused(false)
-    , bIsCompleted(false)
-    , bIsFailed(false)
-	, bIsRequestSwitchingTasks(false)
-    , bWasTaskStartedInBG(false)
-    , bHasAlreadyFinishedRequest(false)
-    , bIsPendingCancel(false)
-    , DownloadProgress(0)
-    , DownloadProgressSinceLastUpdateSent(0)
-{    
-}
-
-void FApplePlatformBackgroundHttpRequest::CompleteWithExistingResponseData(FBackgroundHttpResponsePtr BackgroundResponse)
+	: DownloadId(FBackgroundURLSessionHandler::InvalidDownloadId)
 {
-    if (ensureAlwaysMsgf(BackgroundResponse.IsValid(), TEXT("Call to CompleteWithExistingResponseData with an invalid response!")))
-    {
-        FBackgroundHttpRequestImpl::CompleteWithExistingResponseData(BackgroundResponse);
-        CompleteRequest_Internal(true, BackgroundResponse->GetTempContentFilePath());
-    }
 }
 
-void FApplePlatformBackgroundHttpRequest::SetRequestAsSuccess(const FString& CompletedTempDownloadLocationIn)
+void FApplePlatformBackgroundHttpRequest::PauseRequest()
 {
-    CompleteRequest_Internal(true, CompletedTempDownloadLocationIn);
+	FBackgroundURLSessionHandler::PauseDownload(DownloadId);
 }
 
-void FApplePlatformBackgroundHttpRequest::SetRequestAsFailed()
+void FApplePlatformBackgroundHttpRequest::ResumeRequest()
 {
-    CompleteRequest_Internal(false, FString());
+	FBackgroundURLSessionHandler::ResumeDownload(DownloadId);
 }
 
-void FApplePlatformBackgroundHttpRequest::CompleteRequest_Internal(bool bWasRequestSuccess, const FString& CompletedTempDownloadLocationIn)
+const FString& FApplePlatformBackgroundHttpRequest::GetRequestDebugID() const
 {
-    UE_LOG(LogBackgroundHttpRequest, Display, TEXT("Marking Request Complete -- RequestDebugID:%s | bWasRequestSuccess:%d | CompletedTempDownloadLocation:%s"), *GetRequestDebugID(), (int)bWasRequestSuccess, *CompletedTempDownloadLocationIn);
-    
-	//Purposefully avoid setting bIsTaskActive to false here as we still expect that to be set until the request is deleted.
-	
-	FPlatformAtomics::InterlockedExchange(&bIsCompleted, true);
-	FPlatformAtomics::InterlockedExchange(&bIsFailed, !bWasRequestSuccess);
-    FPlatformAtomics::InterlockedExchange(&bIsPendingCancel, false);
-    
-    if (!CompletedTempDownloadLocationIn.IsEmpty())
-    {
-        CompletedTempDownloadLocation = CompletedTempDownloadLocationIn;
-    }
-    
-	NotifyNotificationObjectOfComplete(bWasRequestSuccess);
+	static const FString EmptyString = TEXT("");
+	return EmptyString;
 }
 
-const FString& FApplePlatformBackgroundHttpRequest::GetURLForRetry(bool bShouldIncrementRetryCountFirst)
+void FApplePlatformBackgroundHttpRequest::UpdateProgress()
 {
-	const int NewRetryCount = bShouldIncrementRetryCountFirst ? RetryCount.Increment() : RetryCount.GetValue();
-	
-	//If we are out of Retries, just send an empty string
-	if (NewRetryCount > NumberOfTotalRetries)
+	if (DownloadId == FBackgroundURLSessionHandler::InvalidDownloadId)
 	{
-		UE_LOG(LogBackgroundHttpRequest, Display, TEXT("GetURLForRetry is out of Retries for Request -- RequestDebugID:%s"), *GetRequestDebugID());
-        
-        static FString EmptyResponse = TEXT("");
-        return EmptyResponse;
+		return;
 	}
-	//Still have remaining retries
-	else
-	{
-		const int URLIndex = NewRetryCount % URLList.Num();
-		const FString& URLToReturn = URLList[URLIndex];
-
-		UE_LOG(LogBackgroundHttpRequest, Display, TEXT("GetURLForRetry found valid URL for current retry -- RequestDebugID:%s | NewRetryCount:%d | URLToReturn:%s"), *GetRequestDebugID(), NewRetryCount, *URLToReturn);
-        return URLToReturn;
-    }
-}
-
-void FApplePlatformBackgroundHttpRequest::ResetProgressTracking()
-{
-    FPlatformAtomics::InterlockedExchange(&DownloadProgress, 0);
-}
-
-void FApplePlatformBackgroundHttpRequest::ActivateUnderlyingTask()
-{
-    volatile FTaskNode* ActiveTaskNode = FirstTask;
-    if (ensureAlwaysMsgf((nullptr != ActiveTaskNode), TEXT("Call to ActivateUnderlyingTask with an invalid node! Need to create underlying task(and node) before activating!")))
-    {
-        NSURLSessionTask* UnderlyingTask = ActiveTaskNode->OurTask;
-        if (ensureAlwaysMsgf((nullptr != UnderlyingTask), TEXT("Call to ActivateUnderlyingTask with an invalid task! Need to create underlying task before activating!")))
-        {
-            FString TaskURL = [[[UnderlyingTask currentRequest] URL] absoluteString];
-            int TaskIdentifier = (int)[UnderlyingTask taskIdentifier];
-            
-            UE_LOG(LogBackgroundHttpRequest, Display, TEXT("Activating Task for Request -- RequestDebugID:%s | TaskIdentifier:%d | TaskURL:%s"), *GetRequestDebugID(), TaskIdentifier, *TaskURL);
-        
-            FPlatformAtomics::InterlockedExchange(&bIsTaskActive, true);
-            FPlatformAtomics::InterlockedExchange(&bIsTaskPaused, false);
-            FPlatformAtomics::InterlockedExchange(&bIsPendingCancel, false);
-            
-            [UnderlyingTask resume];
-            
-            ResetTimeOutTimer();
-            ResetProgressTracking();
-        }
-    }
-}
-
-void FApplePlatformBackgroundHttpRequest::PauseUnderlyingTask()
-{
-    volatile FTaskNode* ActiveTaskNode = FirstTask;
-    if (ensureAlwaysMsgf((nullptr != ActiveTaskNode), TEXT("Call to PauseUnderlyingTask with an invalid node! Need to create underlying task(and node) before trying to pause!")))
-    {
-        NSURLSessionTask* UnderlyingTask = ActiveTaskNode->OurTask;
-        if (ensureAlwaysMsgf((nullptr != UnderlyingTask), TEXT("Call to PauseUnderlyingTask with an invalid task! Need to create underlying task before trying to pause!")))
-        {
-            FString TaskURL = [[[UnderlyingTask currentRequest] URL] absoluteString];
-            int TaskIdentifier = (int)[UnderlyingTask taskIdentifier];
-            
-            UE_LOG(LogBackgroundHttpRequest, Display, TEXT("Pausing Task for Request -- RequestDebugID:%s | TaskIdentifier:%d | TaskURL:%s"), *GetRequestDebugID(), TaskIdentifier, *TaskURL);
-            
-            FPlatformAtomics::InterlockedExchange(&bIsPendingCancel, false);
-            FPlatformAtomics::InterlockedExchange(&bIsTaskPaused, true);
-            
-            [UnderlyingTask suspend];
-            
-            ResetTimeOutTimer();
-            ResetProgressTracking();
-        }
-    }
-}
-
-bool FApplePlatformBackgroundHttpRequest::IsUnderlyingTaskActive()
-{
-    return FPlatformAtomics::AtomicRead(&bIsTaskActive);
-}
-
-bool FApplePlatformBackgroundHttpRequest::IsUnderlyingTaskPaused()
-{
-    return FPlatformAtomics::AtomicRead(&bIsTaskPaused);
-}
-
-bool FApplePlatformBackgroundHttpRequest::TickTimeOutTimer(float DeltaTime)
-{
-    ActiveTimeOutTimer -= DeltaTime;
-    return (ActiveTimeOutTimer <= 0.f) ? true : false;
-}
-
-void FApplePlatformBackgroundHttpRequest::ResetTimeOutTimer()
-{
-    ActiveTimeOutTimer = FApplePlatformBackgroundHttpManager::ActiveTimeOutSetting;
-}
 
-bool FApplePlatformBackgroundHttpRequest::AssociateWithTask(NSURLSessionTask* ExistingTask)
-{
-	bool bDidAssociate = false;
+	const uint64 CurrentDownloadedBytes = FBackgroundURLSessionHandler::GetCurrentDownloadedBytes(DownloadId);
+	const uint64 BytesSinceLastReport = CurrentDownloadedBytes > LastReportedDownloadedBytes ? CurrentDownloadedBytes - LastReportedDownloadedBytes : 0;
+	LastReportedDownloadedBytes = CurrentDownloadedBytes;
 
-	if (ensureAlwaysMsgf((nullptr != ExistingTask), TEXT("Call to AssociateWithTask with an invalid Task! RequestDebugID:%s"), *GetRequestDebugID()))
+	if (BytesSinceLastReport > 0)
 	{
-		const FString TaskURL = [[[ExistingTask currentRequest] URL] absoluteString];
-        int TaskIdentifier = (int)[ExistingTask taskIdentifier];
-        
-		const bool bWasAlreadySwitching = FPlatformAtomics::InterlockedExchange(&bIsRequestSwitchingTasks, true);
-		if (!bWasAlreadySwitching)
-		{
-			volatile FTaskNode* NewNode = new FTaskNode();
-			NewNode->OurTask = ExistingTask;
-
-			//Add a count to our task's reference list so it doesn't get deleted while in our Request's task list
-			[ExistingTask retain];
-
-			//Swap our new node and the first one in the list
-			NewNode->NextNode = (FTaskNode*)FPlatformAtomics::InterlockedExchangePtr((void**)(&FirstTask), (void*)NewNode);
-
-			//Save off our first task's identifier and our CombinedRequestID that includes it
-			FirstTaskIdentifier = (int)[(FirstTask->OurTask) taskIdentifier];
-            CombinedRequestID = FString::Printf(TEXT("%d.%s"), FirstTaskIdentifier, *RequestID);
-            
-            UE_LOG(LogBackgroundHttpRequest, Display, TEXT("Associated Request With New Task -- RequestDebugID:%s | TaskIdentifier:%d | TaskURL:%s"), *GetRequestDebugID(), TaskIdentifier, *TaskURL);
-
-			FPlatformAtomics::InterlockedExchange(&bIsPendingCancel, false);
+		OnProgressUpdated().ExecuteIfBound(SharedThis(this), CurrentDownloadedBytes, BytesSinceLastReport);
+	}
 
-			ResetTimeOutTimer();
-			ResetProgressTracking();
+	int32 ResultHTTPCode;
+	FString TemporaryFilePath;
+	if (FBackgroundURLSessionHandler::IsDownloadFinished(DownloadId, ResultHTTPCode, TemporaryFilePath))
+	{
+		const bool bFileExists = !TemporaryFilePath.IsEmpty() && FPlatformFileManager::Get().GetPlatformFile().FileExists(*TemporaryFilePath);
 
-			bDidAssociate = true;
-            FPlatformAtomics::InterlockedExchange(&bIsRequestSwitchingTasks, false);
-		}
-		else
+		// Fail request if we can't access the file
+		if (!bFileExists)
 		{
-            UE_LOG(LogBackgroundHttpRequest, Display, TEXT("Failed to Associate Request with new Task as there was already a pending AssociateWithTask running! -- RequestDebugID:%s | TaskIdentifier:%d | TaskURL:%s"), *GetRequestDebugID(), TaskIdentifier, *TaskURL);
+			ResultHTTPCode = EHttpResponseCodes::ServerError;
 		}
-	}
-
-	return bDidAssociate;
-}
-
-void FApplePlatformBackgroundHttpRequest::PauseRequest()
-{
-    PauseUnderlyingTask();
-}
 
-void FApplePlatformBackgroundHttpRequest::ResumeRequest()
-{
-	FPlatformAtomics::InterlockedExchange(&bIsTaskPaused, false);
-	
-	//We only want to re-activate tasks that have already been flagged as active.
-	//Otherwise let our BackgroundHTTP Manager handle activating us on a tick now that we aren't paused.
-	if (IsUnderlyingTaskActive())
-	{
-		ActivateUnderlyingTask();
-	}
-	else
-	{
-		UE_LOG(LogBackgroundHttpRequest, Display, TEXT("ResumeRequest called on a task that wasn't active -- RequestDebugID:%s"), *GetRequestDebugID());
+		FBackgroundHttpResponsePtr Response = FPlatformBackgroundHttp::ConstructBackgroundResponse(ResultHTTPCode, TemporaryFilePath);
+		CompleteWithExistingResponseData(Response); // will internally call RemoveRequest which will cancel download
 	}
 }
 
-void FApplePlatformBackgroundHttpRequest::CancelActiveTask()
+float FApplePlatformBackgroundHttpRequest::BackgroundRequestPriorityToNSURLSessionPriority(const EBackgroundHTTPPriority Priority)
 {
-    volatile FTaskNode* TaskNodeWeAreCancelling = FirstTask;
-	if (nullptr != TaskNodeWeAreCancelling)
+	switch (Priority)
 	{
-        if (nullptr != TaskNodeWeAreCancelling->OurTask)
-        {
-            FString TaskURL = [[[TaskNodeWeAreCancelling->OurTask currentRequest] URL] absoluteString];
-            int TaskIdentifier = (int)[TaskNodeWeAreCancelling->OurTask taskIdentifier];
-            
-            UE_LOG(LogBackgroundHttpRequest, Display, TEXT("Cancelling Task -- RequestDebugID:%s | TaskIdentifier:%d | TaskURL:%s"), *GetRequestDebugID(), TaskIdentifier, *TaskURL);
-            
-            FPlatformAtomics::InterlockedExchange(&bIsPendingCancel, true);
-            
-            [TaskNodeWeAreCancelling->OurTask cancel];
-        }
+	case EBackgroundHTTPPriority::High:
+		return NSURLSessionTaskPriorityHigh;
+	case EBackgroundHTTPPriority::Low:
+		return NSURLSessionTaskPriorityLow;
+	case EBackgroundHTTPPriority::Normal:
+	default:
+		return NSURLSessionTaskPriorityDefault;
 	}
 }
-
-void FApplePlatformBackgroundHttpRequest::UpdateDownloadProgress(int64_t TotalDownloaded,int64_t DownloadedSinceLastUpdate)
-{
-    UE_LOG(LogBackgroundHttpRequest, VeryVerbose, TEXT("Request Update Progress -- RequestDebugID:%s | OldProgress:%lld | NewProgress:%lld | ProgressSinceLastUpdate:%lld"), *GetRequestDebugID(), DownloadProgress, TotalDownloaded, DownloadedSinceLastUpdate);
-	
-    FPlatformAtomics::AtomicStore(&DownloadProgress, TotalDownloaded);
-    FPlatformAtomics::InterlockedAdd(&DownloadProgressSinceLastUpdateSent, DownloadedSinceLastUpdate);
-    
-    ResetTimeOutTimer();
-}
-
-void FApplePlatformBackgroundHttpRequest::SendDownloadProgressUpdate()
-{
-    volatile int64 DownloadProgressCopy = FPlatformAtomics::AtomicRead(&DownloadProgress);
-    
-    //Don't send any updates if we haven't updated anything since we last sent an update
-    if (DownloadProgressCopy > 0)
-    {
-        //Reset our DownloadProgressSinceLastUpdateSent to 0 now that we are sending a progress update
-        volatile int64 DownloadProgressSinceLastUpdateSentCopy = FPlatformAtomics::InterlockedExchange(&DownloadProgressSinceLastUpdateSent, 0);
-        
-        OnProgressUpdated().ExecuteIfBound(SharedThis(this), DownloadProgressCopy, DownloadProgressSinceLastUpdateSentCopy);
-    }
-}
-
-const FString& FApplePlatformBackgroundHttpRequest::GetRequestDebugID() const
-{
-    //We use CombinedRequestID to append a TaskIdentifier on the end of the RequestID. If we have one, use that. Otherwise fallback on what is already set
-    return CombinedRequestID.IsEmpty() ? RequestID : CombinedRequestID;
-}
-
-bool FApplePlatformBackgroundHttpRequest::IsTaskComplete() const
-{
-    const bool bDidRequestFail = FPlatformAtomics::AtomicRead(&bIsFailed);
-    const bool bDidRequestComplete = FPlatformAtomics::AtomicRead(&bIsCompleted);
-    return (bDidRequestFail || bDidRequestComplete);
-}
diff --git a/Engine/Source/Runtime/Online/BackgroundHTTP/Public/IOS/ApplePlatformBackgroundHttpManager.h b/Engine/Source/Runtime/Online/BackgroundHTTP/Public/IOS/ApplePlatformBackgroundHttpManager.h
--- a/Engine/Source/Runtime/Online/BackgroundHTTP/Public/IOS/ApplePlatformBackgroundHttpManager.h
+++ b/Engine/Source/Runtime/Online/BackgroundHTTP/Public/IOS/ApplePlatformBackgroundHttpManager.h
@@ -8,8 +8,6 @@
 #include "Interfaces/IBackgroundHttpRequest.h"
 #include "IOS/ApplePlatformBackgroundHttpRequest.h"
 
-#include "IOS/IOSBackgroundURLSessionHandler.h"
-
 typedef TWeakPtr<FApplePlatformBackgroundHttpRequest,ESPMode::ThreadSafe> FBackgroundHttpURLMappedRequestPtr;
 
 DECLARE_DELEGATE(FIOSBackgroundHttpPostSessionWorkCallback);
@@ -21,113 +19,19 @@ class BACKGROUNDHTTP_API FApplePlatformBackgroundHttpManager
 	: public FBackgroundHttpManagerImpl
 {
 public:
-	virtual void Initialize() override;
-	virtual void Shutdown() override;
+	FApplePlatformBackgroundHttpManager() = default;
+	virtual ~FApplePlatformBackgroundHttpManager() = default;
+
 	virtual void AddRequest(const FBackgroundHttpRequestPtr Request) override;
 	virtual void RemoveRequest(const FBackgroundHttpRequestPtr Request) override;
-	virtual bool IsGenericImplementation() const override { return false;  }
-	virtual void SetMaxActiveDownloads(int MaxActiveDownloads) override;
-
-	/**
-	* Constructor
-	*/
-	FApplePlatformBackgroundHttpManager();
-
-	/**
-	 * Destructor
-	 */
-	virtual ~FApplePlatformBackgroundHttpManager();
-
-    //FTSTickerObjectBase implementation
-    virtual bool Tick(float DeltaTime) override;
-    
-    //This is the identifier we expect to be associated with our NSURLDownloadSession for background downloads
-    static FString BackgroundSessionIdentifier;
-    
-    //This value is pulled from .ini settings and is used to time out requests when we are active in the foreground
-    static float ActiveTimeOutSetting;
-    
-	//This value is pulled from .ini setting and is used to limit the number of times a request will try and use the same
-	//resume data before just ignoring that data and moving onto another CDN retry
-	static int RetryResumeDataLimitSetting;
-
-protected:
-	virtual BackgroundHttpFileHashHelperRef GetFileHashHelper() override;
-	virtual const BackgroundHttpFileHashHelperRef GetFileHashHelper() const override;
-	
-	bool AssociateWithAnyExistingUnAssociatedTasks(const FBackgroundHttpRequestPtr Request);
-        
-	bool CheckForExistingUnAssociatedTask(const FAppleBackgroundHttpRequestPtr Request);
-
-    void DeletePendingRemoveRequests();
-	
-protected:
-	//This dictionary is used to hold tasks that already existed on our Background Session when our BackgroundHttpManager was initialized. See PopulateUnAssociatedTasks and CheckForExistingUnAssociatedTask
-	NSMutableDictionary<NSString*, NSURLSessionDownloadTask*>* UnAssociatedTasks;
-	FRWLock UnAssociatedTasksLock;
-	
-	//Map to hold the associated BackgroundHttpRequest for any given URL. Multiple URLs will end up pointing to the same request in this list as it stores
-    TMap<const FString, FBackgroundHttpURLMappedRequestPtr> URLToRequestMap;
-	FRWLock URLToRequestMapLock;
+	virtual void SetCellularPreference(int32 Value) override;
+	virtual bool IsGenericImplementation() const override { return false; }
+	virtual bool Tick(float DeltaTime) override;
 
-	//Used to track if we came into certain code branches through a pathway that has initialized the manager yet or not.
-	//Some of our platform callbacks can end up hitting VERY early or in background versions of our app without the engine being fully spun-up.
-	static volatile bool bWasAppleBGHTTPInitialized;
-	
+	UE_DEPRECATED(5.5, "Variable deprecated") static FString BackgroundSessionIdentifier;
+	UE_DEPRECATED(5.5, "Variable deprecated") static float ActiveTimeOutSetting;
+	UE_DEPRECATED(5.5, "Variable deprecated") static int RetryResumeDataLimitSetting;
 private:
-	//Checks for tasks that already exist on the BackgroundSession. Should only be called during Initialize. These tasks should really only exist 
-	//due to previous application sessions starting background downloads that finished after the app closed. (IE: If the user started a patch and then our app was terminated by the OS while 
-	//suspended. On next launch we will have re-associated those existing downloads to our background download session, but don't yet have a FBackgroundHttpRequest for that task.)
-	//See AssociateWithAnyExistingRequest for how these Tasks are eventually claimed by an FBackgroundHttpRequest
-	void PopulateUnAssociatedTasks();
-
-	void SetupNSURLSessionResponseDelegates();
-	void CleanUpNSURLSessionResponseDelegates();
-
-    //These are used internally to pause/resume things at the task level
-	void PauseAllActiveTasks();
-	void ResumeTasksForBackgrounding(FIOSBackgroundHttpPostSessionWorkCallback Callback = FIOSBackgroundHttpPostSessionWorkCallback());
-	bool ResumeDownloadTaskForBackgroundingIfAppropriate(NSURLSessionDownloadTask* DownloadTask, const FAppleBackgroundHttpRequestPtr Request, EBackgroundHTTPPriority LowestPriorityToQueue);
-
-    void PauseAllUnassociatedTasks();
-    void UnpauseAllUnassociatedTasks();
-    
-    void TickRequests(float DeltaTime);
-    void TickUnassociatedTasks(float DeltaTime);
-    
-    void StartRequest(FAppleBackgroundHttpRequestPtr Request);
-	void FinishRequest(FAppleBackgroundHttpRequestPtr Request);
-	void RetryRequest(FAppleBackgroundHttpRequestPtr Request, bool bShouldIncreaseRetryCount, NSData* RetryData);
-    void RemoveSessionTasksForRequest(FAppleBackgroundHttpRequestPtr Request);
-    bool GenerateURLMapEntriesForRequest(FAppleBackgroundHttpRequestPtr Request);
-    void RemoveURLMapEntriesForRequest(FAppleBackgroundHttpRequestPtr Request);
-    
-	bool IsRetryDataValid(NSData* RetryData) const;
-	bool ShouldUseRequestRetryData(FAppleBackgroundHttpRequestPtr Request, NSData* RetryData) const;
-
-	//Delegates to respond to Background Session Events
-	void OnApp_EnteringForeground();
-	void OnApp_EnteringBackground();
-	void OnTask_DidFinishDownloadingToURL(NSURLSessionDownloadTask* Task, NSError* Error, const FString& TempFilePath);
-	void OnTask_DidWriteData(NSURLSessionDownloadTask* Task, int64_t BytesWrittenSinceLastCall, int64_t TotalBytesWritten, int64_t TotalBytesExpectedToWrite);
-	void OnTask_DidCompleteWithError(NSURLSessionTask* Task, NSError* Error);
-	void OnSession_SessionDidFinishAllEvents(NSURLSession* Session, FIOSBackgroundDownloadCoreDelegates::FIOSBackgroundDownload_DelayedBackgroundURLSessionCompleteHandler Callback);
-
-	FDelegateHandle OnApp_EnteringForegroundHandle;
-	FDelegateHandle OnApp_EnteringBackgroundHandle;
-	FDelegateHandle OnTask_DidFinishDownloadingToURLHandle;
-	FDelegateHandle OnTask_DidWriteDataHandle;
-	FDelegateHandle OnTask_DidCompleteWithErrorHandle;
-	FDelegateHandle OnSession_SessionDidFinishAllEventsHandle;
-
-	volatile int bHasFinishedPopulatingUnassociatedTasks;
-	volatile int bIsInBackground;
-    volatile int bIsIteratingThroughSessionTasks;
-    
-    //Need to store requests to remove at the end of our tick as we find them during iteration
-    TArray<FBackgroundHttpRequestPtr> RequestsPendingRemove;
-    
-    /** On iOS we need to track how many Tasks we have active. This is to replace the default implementations NumCurrentlyActiveRequests **/
-    volatile int NumCurrentlyActiveTasks;
-	int MaxNumActualTasks;
+	TArray<FBackgroundHttpRequestPtr> PendingRemoveRequests;
+	FRWLock PendingRemoveRequestLock;
 };
diff --git a/Engine/Source/Runtime/Online/BackgroundHTTP/Public/IOS/ApplePlatformBackgroundHttpRequest.h b/Engine/Source/Runtime/Online/BackgroundHTTP/Public/IOS/ApplePlatformBackgroundHttpRequest.h
--- a/Engine/Source/Runtime/Online/BackgroundHTTP/Public/IOS/ApplePlatformBackgroundHttpRequest.h
+++ b/Engine/Source/Runtime/Online/BackgroundHTTP/Public/IOS/ApplePlatformBackgroundHttpRequest.h
@@ -3,8 +3,6 @@
 #pragma once
 
 #include "Interfaces/IBackgroundHttpRequest.h"
-#include "Interfaces/IHttpRequest.h"
-
 #include "BackgroundHttpRequestImpl.h"
 
 /**
@@ -15,72 +13,34 @@ class BACKGROUNDHTTP_API FApplePlatformBackgroundHttpRequest
 {
 public:
 	FApplePlatformBackgroundHttpRequest();
-	virtual ~FApplePlatformBackgroundHttpRequest() {}
+	virtual ~FApplePlatformBackgroundHttpRequest() = default;
 
-	virtual void CompleteWithExistingResponseData(FBackgroundHttpResponsePtr BackgroundResponse) override;
-    virtual void PauseRequest() override;
-    virtual void ResumeRequest() override;
-    
-    virtual bool IsTaskComplete() const;
+	virtual void PauseRequest() override;
+	virtual void ResumeRequest() override;
 
-    //Used to provide some extra debug information over normal GetRequestID()
-    //Returns string in format of X.Y where X is the underlying Task Identifier if set and Y is what was set in the SetRequestID() call if this reqeust has associated with a task.
-    //Returns the same as GetRequestID() if no task has been associated yet.
-	const FString& GetRequestDebugID() const;
+	//Used to provide some extra debug information over normal GetRequestID()
+	//Returns string in format of X.Y where X is the underlying Task Identifier if set and Y is what was set in the SetRequestID() call if this reqeust has associated with a task.
+	//Returns the same as GetRequestID() if no task has been associated yet.
+	UE_DEPRECATED(5.5, "Method deprecated") const FString& GetRequestDebugID() const;
 
-private:
-    //Super simple linked list with volatile pointers to next element
-    struct FTaskNode
-    {
-        NSURLSessionTask* OurTask;
-        volatile FTaskNode* NextNode;
-    };
-    
-private:
-    const FString& GetURLForRetry(bool bShouldIncrementRetryCountFirst);
-    bool AssociateWithTask(NSURLSessionTask* ExistingTask);
-	void SetRequestAsSuccess(const FString& CompletedTempDownloadLocation);
-    void SetRequestAsFailed();
-    void CompleteRequest_Internal(bool bWasRequestSuccess, const FString& CompletedTempDownloadLocation);
-	void CancelActiveTask();
-	
-    //on iOS we need to delay when we send progress updates as they can come from a background thread
-    //save off results in UpdateDownloadProgress to send later when we get a SendDownloadProgressUpdate
-    void UpdateDownloadProgress(int64_t TotalDownloaded, int64_t DownloadedSinceLastUpdate);
-    void SendDownloadProgressUpdate();
-    void ResetProgressTracking();
-    
-    void ActivateUnderlyingTask();
-    bool IsUnderlyingTaskActive();
-    void PauseUnderlyingTask();
-    bool IsUnderlyingTaskPaused();
-    bool TickTimeOutTimer(float DeltaTime);
-    void ResetTimeOutTimer();
-    
-    FString CompletedTempDownloadLocation;
-    volatile float ActiveTimeOutTimer;
-    
-	FThreadSafeCounter RetryCount;
-	FThreadSafeCounter ResumeDataRetryCount;
+protected:
+	void SetInternalDownloadId(uint64 Id) {DownloadId = Id;}
+	uint64 GetInternalDownloadId() const {return DownloadId;}
 
-	volatile FTaskNode* FirstTask;
-	
-    volatile int FirstTaskIdentifier;
-    FString CombinedRequestID;
+	void UpdateProgress();
+
+	static float BackgroundRequestPriorityToNSURLSessionPriority(const EBackgroundHTTPPriority Priority);
+	float GetNSURLSessionPriority() const
+	{
+		return BackgroundRequestPriorityToNSURLSessionPriority(GetRequestPriority());
+	}
 
-    volatile int32 bIsTaskActive;
-    volatile int32 bIsTaskPaused;
-	volatile int32 bIsCompleted;
-	volatile int32 bIsFailed;
-	volatile int32 bIsRequestSwitchingTasks;
-	volatile int32 bWasTaskStartedInBG;
-    volatile int32 bHasAlreadyFinishedRequest;
-    volatile int32 bIsPendingCancel;
-    
-	volatile int64 DownloadProgress;
-    volatile int64 DownloadProgressSinceLastUpdateSent;
+	friend class FApplePlatformBackgroundHttpManager;
+	
+private:
+	uint64 DownloadId;
 
-    friend class FApplePlatformBackgroundHttpManager;
+	uint64 LastReportedDownloadedBytes = 0;
 };
 
 typedef TSharedPtr<class FApplePlatformBackgroundHttpRequest, ESPMode::ThreadSafe> FAppleBackgroundHttpRequestPtr;
diff --git a/Engine/Source/Runtime/Online/BackgroundHTTPFileHash/Public/BackgroundHttpFileHashHelper.h b/Engine/Source/Runtime/Online/BackgroundHTTPFileHash/Public/BackgroundHttpFileHashHelper.h
--- a/Engine/Source/Runtime/Online/BackgroundHTTPFileHash/Public/BackgroundHttpFileHashHelper.h
+++ b/Engine/Source/Runtime/Online/BackgroundHTTPFileHash/Public/BackgroundHttpFileHashHelper.h
@@ -112,3 +112,4 @@ private:
 };
 
 typedef TSharedRef<FBackgroundHttpFileHashHelper, ESPMode::ThreadSafe> BackgroundHttpFileHashHelperRef;
+typedef TSharedPtr<FBackgroundHttpFileHashHelper, ESPMode::ThreadSafe> BackgroundHttpFileHashHelperPtr;
