Fix prohibited word use
#rnx

diff --git a/Engine/Source/Runtime/ApplicationCore/Private/IOS/IOSBackgroundURLSessionHandler.cpp b/Engine/Source/Runtime/ApplicationCore/Private/IOS/IOSBackgroundURLSessionHandler.cpp
--- a/Engine/Source/Runtime/ApplicationCore/Private/IOS/IOSBackgroundURLSessionHandler.cpp
+++ b/Engine/Source/Runtime/ApplicationCore/Private/IOS/IOSBackgroundURLSessionHandler.cpp
@@ -572,7 +572,7 @@ static constexpr NSInteger HTTPStatusCodeErrorServer = 500;
 	FBackgroundNSURLSessionDownloadTaskData* TaskData = [FBackgroundNSURLSessionDownloadTaskData TaskDataFromSerializedString:Task.taskDescription];
 	if (TaskData != nil)
 	{
-		// Nuke task data, otherwise didCompleteWithError might retry request.
+		// Remove task data from this download task, otherwise didCompleteWithError might retry the request.
 		[TaskData Cancel];
 		Task.taskDescription = [TaskData ToSerializedString];
 	}
