Partially fixing memory leak in iOS background download code and game controller poll code in app preinit

diff --git a/Engine/Source/Runtime/ApplicationCore/Private/IOS/IOSApplication.cpp b/Engine/Source/Runtime/ApplicationCore/Private/IOS/IOSApplication.cpp
--- a/Engine/Source/Runtime/ApplicationCore/Private/IOS/IOSApplication.cpp
+++ b/Engine/Source/Runtime/ApplicationCore/Private/IOS/IOSApplication.cpp
@@ -91,36 +91,39 @@ void FIOSApplication::AddExternalInputDevice(TSharedPtr<IInputDevice> InputDevic
 
 void FIOSApplication::PollGameDeviceState( const float TimeDelta )
 {
-	// initialize any externally-implemented input devices (we delay load initialize the array so any plugins have had time to load)
-	if (!bHasLoadedInputPlugins && GIsRunning)
+	@autoreleasepool
 	{
-		TArray<IInputDeviceModule*> PluginImplementations = IModularFeatures::Get().GetModularFeatureImplementations<IInputDeviceModule>(IInputDeviceModule::GetModularFeatureName());
-		for (auto InputPluginIt = PluginImplementations.CreateIterator(); InputPluginIt; ++InputPluginIt)
+		// initialize any externally-implemented input devices (we delay load initialize the array so any plugins have had time to load)
+		if (!bHasLoadedInputPlugins && GIsRunning)
 		{
-			TSharedPtr<IInputDevice> Device = (*InputPluginIt)->CreateInputDevice(MessageHandler);
-			AddExternalInputDevice(Device);
+			TArray<IInputDeviceModule*> PluginImplementations = IModularFeatures::Get().GetModularFeatureImplementations<IInputDeviceModule>(IInputDeviceModule::GetModularFeatureName());
+			for (auto InputPluginIt = PluginImplementations.CreateIterator(); InputPluginIt; ++InputPluginIt)
+			{
+				TSharedPtr<IInputDevice> Device = (*InputPluginIt)->CreateInputDevice(MessageHandler);
+				AddExternalInputDevice(Device);
+			}
+
+			bHasLoadedInputPlugins = true;
 		}
 
-		bHasLoadedInputPlugins = true;
-	}
-
-	// Poll game device state and send new events
-	{
-		QUICK_SCOPE_CYCLE_COUNTER(STAT_IOSApplication_InputInterface_Tick);
-		InputInterface->Tick(TimeDelta);
-	}
-	{
-		QUICK_SCOPE_CYCLE_COUNTER(STAT_IOSApplication_InputInterface_SendControllerEvents);
-		InputInterface->SendControllerEvents();
-	}
+		// Poll game device state and send new events
+		{
+			QUICK_SCOPE_CYCLE_COUNTER(STAT_IOSApplication_InputInterface_Tick);
+			InputInterface->Tick(TimeDelta);
+		}
+		{
+			QUICK_SCOPE_CYCLE_COUNTER(STAT_IOSApplication_InputInterface_SendControllerEvents);
+			InputInterface->SendControllerEvents();
+		}
 
-	// Poll externally-implemented devices
-	{
-		QUICK_SCOPE_CYCLE_COUNTER(STAT_IOSApplication_ExternalInputDevice);
-		for (auto DeviceIt = ExternalInputDevices.CreateIterator(); DeviceIt; ++DeviceIt)
+		// Poll externally-implemented devices
 		{
-			(*DeviceIt)->Tick(TimeDelta);
-			(*DeviceIt)->SendControllerEvents();
+			QUICK_SCOPE_CYCLE_COUNTER(STAT_IOSApplication_ExternalInputDevice);
+			for (auto DeviceIt = ExternalInputDevices.CreateIterator(); DeviceIt; ++DeviceIt)
+			{
+				(*DeviceIt)->Tick(TimeDelta);
+				(*DeviceIt)->SendControllerEvents();
+			}
 		}
 	}
 }
diff --git a/Engine/Source/Runtime/ApplicationCore/Private/IOS/IOSBackgroundURLSessionHandler.cpp b/Engine/Source/Runtime/ApplicationCore/Private/IOS/IOSBackgroundURLSessionHandler.cpp
--- a/Engine/Source/Runtime/ApplicationCore/Private/IOS/IOSBackgroundURLSessionHandler.cpp
+++ b/Engine/Source/Runtime/ApplicationCore/Private/IOS/IOSBackgroundURLSessionHandler.cpp
@@ -914,88 +914,115 @@ const uint64 FBackgroundURLSessionHandler::InvalidDownloadId = [FBackgroundNSURL
 
 void FBackgroundURLSessionHandler::AllowCellular(bool bAllow)
 {
-	const BOOL bCurrentValue = [FBackgroundNSURLSession Shared].AllowCellular;
-	const BOOL bNewValue = bAllow ? YES : NO;
-	if (bCurrentValue == bNewValue)
+	@autoreleasepool
 	{
-		return;
-	}
+		const BOOL bCurrentValue = [FBackgroundNSURLSession Shared].AllowCellular;
+		const BOOL bNewValue = bAllow ? YES : NO;
+		if (bCurrentValue == bNewValue)
+		{
+			return;
+		}
 
-	[[FBackgroundNSURLSession Shared] setAllowCellular:bNewValue];
-	[[FBackgroundNSURLSession Shared] RecreateDownloads];
+		[[FBackgroundNSURLSession Shared] setAllowCellular:bNewValue];
+		[[FBackgroundNSURLSession Shared] RecreateDownloads];
+	}
 }
 
 uint64 FBackgroundURLSessionHandler::CreateOrFindDownload(const TArray<FString>& URLs, const float Priority, BackgroundHttpFileHashHelperRef HelperRef)
 {
-	NSMutableArray* URLArray = [NSMutableArray arrayWithCapacity:URLs.Num()];
-	for (const FString& URL: URLs)
+	@autoreleasepool
 	{
-		[URLArray addObject:URL.GetNSString()];
+		NSMutableArray* URLArray = [NSMutableArray arrayWithCapacity:URLs.Num()];
+		for (const FString& URL: URLs)
+		{
+			[URLArray addObject:URL.GetNSString()];
+		}
+		
+		[[FBackgroundNSURLSession Shared] SetFileHashHelper:HelperRef];
+		return [[FBackgroundNSURLSession Shared] CreateOrFindDownloadForURLs:URLArray WithPriority:Priority];
 	}
-
-	[[FBackgroundNSURLSession Shared] SetFileHashHelper:HelperRef];
-	return [[FBackgroundNSURLSession Shared] CreateOrFindDownloadForURLs:URLArray WithPriority:Priority];
 }
 
 void FBackgroundURLSessionHandler::PauseDownload(const uint64 DownloadId)
 {
-	[[FBackgroundNSURLSession Shared] PauseDownload:DownloadId];
+	@autoreleasepool
+	{
+		[[FBackgroundNSURLSession Shared] PauseDownload:DownloadId];
+	}
 }
 
 void FBackgroundURLSessionHandler::ResumeDownload(const uint64 DownloadId)
 {
-	[[FBackgroundNSURLSession Shared] ResumeDownload:DownloadId];
+	@autoreleasepool
+	{
+		[[FBackgroundNSURLSession Shared] ResumeDownload:DownloadId];
+	}
 }
 
 void FBackgroundURLSessionHandler::CancelDownload(const uint64 DownloadId)
 {
-	[[FBackgroundNSURLSession Shared] CancelDownload:DownloadId];
+	@autoreleasepool
+	{
+		[[FBackgroundNSURLSession Shared] CancelDownload:DownloadId];
+	}
 }
 
 void FBackgroundURLSessionHandler::SetPriority(const uint64 DownloadId, const float Priority)
 {
-	[[FBackgroundNSURLSession Shared] SetPriority:Priority ForDownload:DownloadId];
+	@autoreleasepool
+	{
+		[[FBackgroundNSURLSession Shared] SetPriority:Priority ForDownload:DownloadId];
+	}
 }
 
 uint64 FBackgroundURLSessionHandler::GetCurrentDownloadedBytes(const uint64 DownloadId)
 {
-	return [[FBackgroundNSURLSession Shared] GetCurrentDownloadedBytes:DownloadId];
+	@autoreleasepool
+	{
+		return [[FBackgroundNSURLSession Shared] GetCurrentDownloadedBytes:DownloadId];
+	}
 }
 
 bool FBackgroundURLSessionHandler::IsDownloadFinished(const uint64 DownloadId, int32& OutResultHTTPCode, FString& OutTemporaryFilePath)
 {
-	BOOL Status = NO;
-	NSInteger StatusCode = 0;
-	NSString* TempFile = [[FBackgroundNSURLSession Shared] GetDownloadResult:DownloadId OutStatus:&Status OutStatusCode:&StatusCode];
-	if (!Status)
+	@autoreleasepool
 	{
-		return false;
-	}
-
-	OutResultHTTPCode = (int32)StatusCode;
-	if (TempFile != nil)
-	{
-		OutTemporaryFilePath = FString(TempFile);
-		UE_DNLD_LOG(@"DownloadId %lli finished with status code %li and path '%@'", DownloadId, (long)StatusCode, TempFile);
-	}
-	else
-	{
-		UE_DNLD_LOG(@"DownloadId %lli finished with status code %li and no path", DownloadId, (long)StatusCode);
+		BOOL Status = NO;
+		NSInteger StatusCode = 0;
+		NSString* TempFile = [[FBackgroundNSURLSession Shared] GetDownloadResult:DownloadId OutStatus:&Status OutStatusCode:&StatusCode];
+		if (!Status)
+		{
+			return false;
+		}
+		
+		OutResultHTTPCode = (int32)StatusCode;
+		if (TempFile != nil)
+		{
+			OutTemporaryFilePath = FString(TempFile);
+			UE_DNLD_LOG(@"DownloadId %lli finished with status code %li and path '%@'", DownloadId, (long)StatusCode, TempFile);
+		}
+		else
+		{
+			UE_DNLD_LOG(@"DownloadId %lli finished with status code %li and no path", DownloadId, (long)StatusCode);
+		}
+		
+		return true;
 	}
-
-	return true;
 }
 
 void FBackgroundURLSessionHandler::HandleEventsForBackgroundURLSession(const FString& SessionIdentifier)
 {
-	NSString* Identifier = SessionIdentifier.GetNSString();
-	if (![[FBackgroundNSURLSession GetNSURLSessionIdentifier] isEqualToString:Identifier])
+	@autoreleasepool
 	{
-		UE_DNLD_LOG(@"HandleEventsForBackgroundURLSession ignoring session identifier '%@'", Identifier);
-		return;
+		NSString* Identifier = SessionIdentifier.GetNSString();
+		if (![[FBackgroundNSURLSession GetNSURLSessionIdentifier] isEqualToString:Identifier])
+		{
+			UE_DNLD_LOG(@"HandleEventsForBackgroundURLSession ignoring session identifier '%@'", Identifier);
+			return;
+		}
+		
+		UE_DNLD_LOG(@"HandleEventsForBackgroundURLSession will initializes session with identifier '%@'", Identifier);
+		[FBackgroundNSURLSession Shared];
+		// will invoke URLSessionDidFinishEventsForBackgroundURLSession internally.
 	}
-
-	UE_DNLD_LOG(@"HandleEventsForBackgroundURLSession will initializes session with identifier '%@'", Identifier);
-	[FBackgroundNSURLSession Shared];
-	// will invoke URLSessionDidFinishEventsForBackgroundURLSession internally.
 }
