Adding CDN re-ordering to potentially improve chances of reply at the first try

diff --git a/Engine/Config/IOS/BaseIOSEngine.ini b/Engine/Config/IOS/BaseIOSEngine.ini
--- a/Engine/Config/IOS/BaseIOSEngine.ini
+++ b/Engine/Config/IOS/BaseIOSEngine.ini
@@ -51,6 +51,9 @@ BackgroundReceiveTimeout=120
 BackgroundHttpResourceTimeout=3600
 ; How many times we should try using same CDN before moving onto a different CDN ( -1 = unlimited, 0 = only try first CDN, 1 = try each CDN once )
 RetryResumeDataLimit=3
+; How many milliseconds to wait for HEAD replies on first CDN request to decide if CDN is reachable or not (0 = disabled).
+; Unreachable CDNs will be deprioritized in all download requests.
+CDNReorderingTimeout=400
 
 [ConsoleVariables]
 Slate.CacheRenderData=0
diff --git a/Engine/Source/Runtime/ApplicationCore/Private/IOS/IOSBackgroundURLSessionHandler.cpp b/Engine/Source/Runtime/ApplicationCore/Private/IOS/IOSBackgroundURLSessionHandler.cpp
--- a/Engine/Source/Runtime/ApplicationCore/Private/IOS/IOSBackgroundURLSessionHandler.cpp
+++ b/Engine/Source/Runtime/ApplicationCore/Private/IOS/IOSBackgroundURLSessionHandler.cpp
@@ -209,6 +209,7 @@ NSString* const SerializationKeyRetryCountPerURL = @"r";
 - (void)SetFileHashHelper:(BackgroundHttpFileHashHelperRef)HelperRef;
 - (BackgroundHttpFileHashHelperRef)GetFileHashHelper;
 - (NSString*)GetTempPathForURL:(NSURL* _Nonnull)URL;
+- (NSMutableArray<__kindof NSURL*>*)ReorderCDNsByReachability:(NSMutableArray<__kindof NSURL*>*)URLs;
 - (NSURLSessionDownloadTask*)CreateDownloadForURL:(NSURL* _Nonnull)URL WithPriority:(float)Priority WithTaskData:(FBackgroundNSURLSessionDownloadTaskData* _Nonnull)TaskData;
 - (NSURLSessionDownloadTask*)CreateDownloadForResumeData:(NSData* _Nonnull)ResumeData WithPriority:(float)Priority WithTaskData:(FBackgroundNSURLSessionDownloadTaskData* _Nonnull)TaskData;
 
@@ -253,8 +254,11 @@ NSString* const SerializationKeyRetryCountPerURL = @"r";
 	NSUInteger _NextDownloadId;
 	std::promise<void> _AllDownloadsPromise;
 	std::future<void> _AllDownloadsFuture;
+	NSSet<__kindof NSString*>* _UnreachableCDNs;
 	BackgroundHttpFileHashHelperPtr _HelperPtr;
+	int32 MaximumConnectionsPerHost;
 	int32 RetryResumeDataLimit;
+	int32 CDNReorderingTimeout;
 }
 
 static constexpr NSUInteger InvalidDownloadId = 0;
@@ -266,7 +270,7 @@ NSProgressUserInfoKey const NSProgressDownloadResultStatusCode = @"com.epicgames
 NSProgressUserInfoKey const NSProgressDownloadResultTempFilePath = @"com.epicgames.nsprogress.tempfilepath";
 
 static constexpr NSInteger HTTPStatusCodeSuccessCreated = 201;
-static constexpr NSInteger HTTPStatusCodeErrorBadRequest = 500;
+static constexpr NSInteger HTTPStatusCodeErrorBadRequest = 400;
 static constexpr NSInteger HTTPStatusCodeErrorServer = 500;
 
 + (FBackgroundNSURLSession*)Shared
@@ -303,10 +307,11 @@ static constexpr NSInteger HTTPStatusCodeErrorServer = 500;
 	bool bUseForegroundSession = false;
 	bool bDiscretionary = false;
 	bool bShouldSendLaunchEvents = true;
-	int32 MaximumConnectionsPerHost = 6;
+	MaximumConnectionsPerHost = 6;
 	double TimeoutIntervalForRequest = 120.0; // Note, ignored in background sessions (if bUseForegroundSession is false).
 	double TimeoutIntervalForResource = 60.0 * 60.0;
 	RetryResumeDataLimit = 3;
+	CDNReorderingTimeout = 400;
 
 #ifndef UE_DNLD_SANDBOX
 	GConfig->GetBool(TEXT("BackgroundHttp.iOSSettings"), TEXT("bUseForegroundSession"), bUseForegroundSession, GEngineIni);
@@ -316,11 +321,13 @@ static constexpr NSInteger HTTPStatusCodeErrorServer = 500;
 	GConfig->GetDouble(TEXT("BackgroundHttp.iOSSettings"), TEXT("BackgroundReceiveTimeout"), TimeoutIntervalForRequest, GEngineIni);
 	GConfig->GetDouble(TEXT("BackgroundHttp.iOSSettings"), TEXT("BackgroundHttpResourceTimeout"), TimeoutIntervalForResource, GEngineIni);
 	GConfig->GetInt(TEXT("BackgroundHttp.iOSSettings"), TEXT("RetryResumeDataLimit"), RetryResumeDataLimit, GEngineIni);
+	GConfig->GetInt(TEXT("BackgroundHttp.iOSSettings"), TEXT("CDNReorderingTimeout"), CDNReorderingTimeout, GEngineIni);
 #endif
 	
 	_AllDownloads = [NSMutableDictionary new];
 	_AllDownloadsFuture = _AllDownloadsPromise.get_future();
 	_NextDownloadId = InvalidDownloadId + 1;
+	_UnreachableCDNs = nil;
 
 	// Never allow cellular unless we get explicit opt-in from the user.
 	self.AllowCellular = NO;
@@ -410,6 +417,10 @@ static constexpr NSInteger HTTPStatusCodeErrorServer = 500;
 {
 	[_Session release];
 	[_AllDownloads release];
+	if (_UnreachableCDNs != nil)
+	{
+		[_UnreachableCDNs release];
+	}
 
 	[super dealloc];
 }
@@ -445,6 +456,129 @@ static constexpr NSInteger HTTPStatusCodeErrorServer = 500;
 	return ConvertedPath.GetNSString();
 }
 
+- (NSMutableArray<__kindof NSURL*>*)ReorderCDNsByReachability:(NSMutableArray<__kindof NSURL*>*)URLs
+{
+	if (CDNReorderingTimeout == 0)
+	{
+		return URLs;
+	}
+
+	@synchronized (_UnreachableCDNs)
+	{
+		if (_UnreachableCDNs == nil)
+		{
+			UE_DNLD_LOG(@"Starting to check for CDN reachability");
+
+			const double CDNReorderingTimeoutInSeconds = (double)CDNReorderingTimeout / 1000;
+
+			// Creating a temporary foreground NSURLSession for pinging CDN's
+			NSURLSessionConfiguration* Configuration = [NSURLSessionConfiguration defaultSessionConfiguration];
+			Configuration.discretionary = NO;
+			// Don't go over cellular here, during first download attempt we would try to over non-cellular connections first.
+			// Hence it makes more sense to prioritize CDN's that are reachable via non-cellular connections.
+			Configuration.allowsCellularAccess = NO;
+			Configuration.networkServiceType = NSURLNetworkServiceTypeResponsiveData;
+			Configuration.timeoutIntervalForRequest = CDNReorderingTimeoutInSeconds;
+			Configuration.timeoutIntervalForResource = CDNReorderingTimeoutInSeconds;
+			Configuration.HTTPMaximumConnectionsPerHost = MaximumConnectionsPerHost;
+
+			NSURLSession* Session = [NSURLSession sessionWithConfiguration:Configuration];
+			NSMutableSet<__kindof NSString*>* ReachableHosts = [NSMutableSet set];
+
+			std::shared_ptr<std::atomic<int32>> PendingTasks = std::make_shared<std::atomic<int32>>();
+			std::shared_ptr<std::promise<void>> PendingTasksFinished = std::make_shared<std::promise<void>>();
+
+			for (NSURL* URL in URLs)
+			{
+				NSMutableURLRequest* Request = [NSMutableURLRequest requestWithURL:URL];
+				// Use HEAD request because we want the smallest response possible from the CDN to see if the connection works at all.
+				[Request setHTTPMethod:@"HEAD"];
+
+				PendingTasks->fetch_add(1);
+
+				UE_DNLD_LOG(@"Create data task for '%@'", Request.URL.absoluteString);
+
+				// Note, completion handler might be invoked after end of this method.
+				NSURLSessionDataTask* Task = [Session dataTaskWithRequest:Request completionHandler:^(NSData* _Nullable Data, NSURLResponse* _Nullable Response, NSError* _Nullable Error)
+				{
+					if (Response != nil && [Response isKindOfClass:[NSHTTPURLResponse class]])
+					{
+						NSHTTPURLResponse* HTTPResponse = (NSHTTPURLResponse*)Response;
+						UE_DNLD_LOG(@"Finished data task for '%@' (host '%@') with status code %li", Request.URL.absoluteString, Request.URL.host, HTTPResponse.statusCode);
+
+						if (HTTPResponse.statusCode < HTTPStatusCodeErrorBadRequest)
+						{
+							[ReachableHosts addObject:Request.URL.host];
+						}
+					}
+					else
+					{
+						UE_DNLD_LOG(@"Finished data task for '%@' with error '%@'", Request.URL.absoluteString, (Error != nil ? Error.localizedDescription : @"nil"));
+					}
+
+					if (PendingTasks->fetch_add(-1) <= 1)
+					{
+						UE_DNLD_LOG(@"Finished all data tasks for CDN reachability");
+						PendingTasksFinished->set_value();
+					}
+				}];
+				[Task setPriority:NSURLSessionTaskPriorityHigh];
+				[Task resume];
+			}
+
+			PendingTasksFinished->get_future().wait_for(std::chrono::milliseconds(CDNReorderingTimeout));
+			UE_DNLD_LOG(@"Finished waiting for CDN reachability");
+
+			[Session invalidateAndCancel];
+
+			// We only store unreachable CDNs instead of reachable ones,
+			// to cover a case where a new unknown CDN comes in we would assume it's reachable rather than unreachable.
+			NSMutableSet<__kindof NSString*>* UnreachableCDNs = [NSMutableSet set];
+
+			for (NSURL* URL in URLs)
+			{
+				if ([ReachableHosts containsObject:URL.host])
+				{
+					UE_DNLD_LOG(@"CDN '%@' is reachable", URL.absoluteString);
+				}
+				else
+				{
+					UE_DNLD_LOG(@"CDN '%@' is not reachable", URL.absoluteString);
+					[UnreachableCDNs addObject:URL.host];
+				}
+			}
+
+			_UnreachableCDNs = [[NSSet setWithSet:UnreachableCDNs] retain];
+		}
+	}
+
+	if (_UnreachableCDNs != nil && _UnreachableCDNs.count > 0)
+	{
+		NSMutableArray<__kindof NSURL*>* Result = [NSMutableArray array];
+		NSMutableArray<__kindof NSURL*>* UnreachableURLs = [NSMutableArray array];
+
+		for (NSURL* URL in URLs)
+		{
+			// don't change the order if we don't know if host was unreachable
+			if (![_UnreachableCDNs containsObject:URL.host])
+			{
+				[Result addObject:URL];
+			}
+			else // otherwise put all hosts that were unreachable in the end without changing the relative order between them
+			{
+				[UnreachableURLs addObject:URL];
+			}
+		}
+
+		[Result addObjectsFromArray:UnreachableURLs];
+		return Result;
+	}
+	else
+	{
+		return URLs;
+	}
+}
+
 - (NSURLSessionDownloadTask*)CreateDownloadForURL:(NSURL* _Nonnull)URL WithPriority:(float)Priority WithTaskData:(FBackgroundNSURLSessionDownloadTaskData* _Nonnull)TaskData;
 {
 	NSMutableURLRequest* URLRequest = [NSMutableURLRequest requestWithURL:URL];
@@ -496,6 +630,8 @@ static constexpr NSInteger HTTPStatusCodeErrorServer = 500;
 		[URLs addObject:URLValue];
 	}
 
+	URLs = [self ReorderCDNsByReachability:URLs];
+
 	// Serialize current settings
 	FBackgroundNSURLSessionDownloadTaskData* TaskData = [FBackgroundNSURLSessionDownloadTaskData TaskDataWithURLs:URLs WithRetryCount:RetryResumeDataLimit];
 
