Fix a race in iOS background downloader CDN reorder check

diff --git a/Engine/Source/Runtime/ApplicationCore/Private/IOS/IOSBackgroundURLSessionHandler.cpp b/Engine/Source/Runtime/ApplicationCore/Private/IOS/IOSBackgroundURLSessionHandler.cpp
--- a/Engine/Source/Runtime/ApplicationCore/Private/IOS/IOSBackgroundURLSessionHandler.cpp
+++ b/Engine/Source/Runtime/ApplicationCore/Private/IOS/IOSBackgroundURLSessionHandler.cpp
@@ -488,14 +488,14 @@ static constexpr NSInteger HTTPStatusCodeErrorServer = 500;
 			std::shared_ptr<std::atomic<int32>> PendingTasks = std::make_shared<std::atomic<int32>>();
 			std::shared_ptr<std::promise<void>> PendingTasksFinished = std::make_shared<std::promise<void>>();
 
+			PendingTasks->fetch_add((int32)URLs.count);
+
 			for (NSURL* URL in URLs)
 			{
 				NSMutableURLRequest* Request = [NSMutableURLRequest requestWithURL:URL];
 				// Use HEAD request because we want the smallest response possible from the CDN to see if the connection works at all.
 				[Request setHTTPMethod:@"HEAD"];
 
-				PendingTasks->fetch_add(1);
-
 				UE_DNLD_LOG(@"Create data task for '%@'", Request.URL.absoluteString);
 
 				// Note, completion handler might be invoked after end of this method.
