Add URL check to update command line arguments only on expected received schema

diff --git a/Engine/Source/Runtime/ApplicationCore/Private/IOS/IOSAppDelegate.cpp b/Engine/Source/Runtime/ApplicationCore/Private/IOS/IOSAppDelegate.cpp
--- a/Engine/Source/Runtime/ApplicationCore/Private/IOS/IOSAppDelegate.cpp
+++ b/Engine/Source/Runtime/ApplicationCore/Private/IOS/IOSAppDelegate.cpp
@@ -1356,24 +1356,27 @@ static FAutoConsoleVariableRef CVarGEnableThermalsReport(
 	NSLog(@"%s", "IOSAppDelegate openURL\n");
 #endif
 
-	NSString* EncdodedURLString = [url absoluteString];
-	NSString* URLString = [EncdodedURLString stringByRemovingPercentEncoding];
-	FString CommandLineParameters(URLString);
-
-	// Strip the "URL" part of the URL before treating this like args. It comes in looking like so:
-	// "MyGame://arg1 arg2 arg3 ..."
-	// So, we're going to make it look like:
-	// "arg1 arg2 arg3 ..."
-	int32 URLTerminator = CommandLineParameters.Find( TEXT("://"), ESearchCase::CaseSensitive);
-	if ( URLTerminator > -1 )
-	{
-		CommandLineParameters.RightChopInline(URLTerminator + 3, EAllowShrinking::No);
+	if (FString(url.scheme) == FApp::GetName())
+	{
+		NSString* EncdodedURLString = [url absoluteString];
+		NSString* URLString = [EncdodedURLString stringByRemovingPercentEncoding];
+		FString CommandLineParameters(URLString);
+		
+		// Strip the "URL" part of the URL before treating this like args. It comes in looking like so:
+		// "MyGame://arg1 arg2 arg3 ..."
+		// So, we're going to make it look like:
+		// "arg1 arg2 arg3 ..."
+		int32 URLTerminator = CommandLineParameters.Find( TEXT("://"), ESearchCase::CaseSensitive);
+		if ( URLTerminator > -1 )
+		{
+			CommandLineParameters.RightChopInline(URLTerminator + 3, EAllowShrinking::No);
+		}
+		
+		FIOSCommandLineHelper::InitCommandArgs(CommandLineParameters);
+		self.bCommandLineReady = true;
+		[self.CommandLineParseTimer invalidate];
+		self.CommandLineParseTimer = nil;
 	}
-
-	FIOSCommandLineHelper::InitCommandArgs(CommandLineParameters);
-	self.bCommandLineReady = true;
-	[self.CommandLineParseTimer invalidate];
-	self.CommandLineParseTimer = nil;
 	
 	//    Save openurl infomation before engine initialize.
 	//    When engine is done ready, running like previous. ( if OnOpenUrl is bound on game source. )
@@ -1386,8 +1389,8 @@ static FAutoConsoleVariableRef CVarGEnableThermalsReport(
 #if !NO_LOGGING
 		NSLog(@"%s", "Before Engine Init receive IOSAppDelegate openURL\n");
 #endif
-			NSDictionary* openUrlParameter = [NSDictionary dictionaryWithObjectsAndKeys :
-		application, @"application",
+		NSDictionary* openUrlParameter = [NSDictionary dictionaryWithObjectsAndKeys :
+			application, @"application",
 			url, @"url",
 			sourceApplication, @"sourceApplication",
 			annotation, @"annotation",
