Android window: Fix incorrect DPI adjustment on low DPI screens.

diff --git a/Engine/Source/Runtime/ApplicationCore/Private/Android/AndroidWindow.cpp b/Engine/Source/Runtime/ApplicationCore/Private/Android/AndroidWindow.cpp
--- a/Engine/Source/Runtime/ApplicationCore/Private/Android/AndroidWindow.cpp
+++ b/Engine/Source/Runtime/ApplicationCore/Private/Android/AndroidWindow.cpp
@@ -788,6 +788,9 @@ static FAndroidDisplayInfo GetAndroidDisplayInfoFromDPITargets(int32 TargetDPI,
 	UE_CLOG(TargetDPI > NativeScreenDensityDPI, LogAndroid, Display, TEXT("AndroidDisplayInfoFromDPITargets : TargetDPI too high, using native screen DPI %d, window dims %d, %d"), NativeScreenDensityDPI, Info.WindowDims.X, Info.WindowDims.Y);
 	TargetDPI = FMath::Min(TargetDPI, NativeScreenDensityDPI);
 
+	UE_CLOG(NativeScreenDensityDPI<LowerLimit3DDPI, LogAndroid, Display, TEXT("AndroidDisplayInfoFromDPITargets : 3d scene lower limit is too high (%d) for a native screen, clamping lower limit to DPI %d"), LowerLimit3DDPI, NativeScreenDensityDPI);
+	LowerLimit3DDPI = FMath::Min(LowerLimit3DDPI, NativeScreenDensityDPI);
+
 	int DesiredPixelCount = Info.WindowDims.X * Info.WindowDims.Y;
 	if (SceneMaxDesiredPixelCount && DesiredPixelCount > SceneMaxDesiredPixelCount)
 	{
