Fix Android NonPCH compilation

diff --git a/Engine/Source/Runtime/ApplicationCore/Public/Android/AndroidApplication.h b/Engine/Source/Runtime/ApplicationCore/Public/Android/AndroidApplication.h
--- a/Engine/Source/Runtime/ApplicationCore/Public/Android/AndroidApplication.h
+++ b/Engine/Source/Runtime/ApplicationCore/Public/Android/AndroidApplication.h
@@ -3,6 +3,7 @@
 #pragma once
 
 #include "GenericPlatform/GenericApplication.h"
+#include "HAL/IConsoleManager.h"
 #include "AndroidWindow.h"
 #if USE_ANDROID_JNI
 #include "Android/AndroidJavaEnv.h"
