Remove ASIS outdated functionality

#rnx

diff --git a/Engine/Source/Runtime/ApplicationCore/Private/Android/AndroidWindow.cpp b/Engine/Source/Runtime/ApplicationCore/Private/Android/AndroidWindow.cpp
--- a/Engine/Source/Runtime/ApplicationCore/Private/Android/AndroidWindow.cpp
+++ b/Engine/Source/Runtime/ApplicationCore/Private/Android/AndroidWindow.cpp
@@ -91,8 +91,6 @@ static void ClearCachedWindowRects()
 	CachedWindowRect_EventThread = FAndroidCachedWindowRectParams();
 }
 
-static int32 GSurfaceViewX = 0;
-static int32 GSurfaceViewY = 0;
 int32 GSurfaceViewWidth = -1;
 int32 GSurfaceViewHeight = -1;
 
@@ -200,11 +198,8 @@ JNI_METHOD void Java_com_epicgames_makeaar_GameActivityForMakeAAR_nativeSetSurfa
 	ANativeWindow* prev = (ANativeWindow*)GAndroidWindowOverride;
 	if (surface != 0)
 	{
-		GSurfaceViewX = x;
-		GSurfaceViewY = y;
-
 		GAndroidWindowOverride = (ANativeWindow*)ANativeWindow_fromSurface(jenv, surface);
-		UE_LOG(LogAndroid, Log, TEXT("nativeSetSurfaceOverride applied: prev to new %p -> %p, pos(%d, %d)"), prev, GAndroidWindowOverride, GSurfaceViewX, GSurfaceViewY);
+		UE_LOG(LogAndroid, Log, TEXT("nativeSetSurfaceOverride applied: prev to new %p -> %p, pos(%d, %d)"), prev, GAndroidWindowOverride, x, y);
 
 	}
 	else
diff --git a/Engine/Source/Runtime/Launch/Private/Android/LaunchAndroid.cpp b/Engine/Source/Runtime/Launch/Private/Android/LaunchAndroid.cpp
--- a/Engine/Source/Runtime/Launch/Private/Android/LaunchAndroid.cpp
+++ b/Engine/Source/Runtime/Launch/Private/Android/LaunchAndroid.cpp
@@ -1435,12 +1435,7 @@ static void ActivateApp_EventThread()
 	DEVELOPER_LOG_COMMANDCB_CASE(ActivateApp_EventThread);
 	if (bAppIsActive_EventThread)
 	{
-		STANDALONE_DEBUG_LOGf(LogAndroid, TEXT("UNEXPECTED -- event thread, activate app, ALREADY have HW window lock according to bAppIsActive_EventThread. bReadyToProcessEvents=%d, bAppIsActive_EventThread=%d"), bReadyToProcessEvents, bAppIsActive_EventThread);
-
-#if !USE_ANDROID_STANDALONE // SUSPECT TODO should remove the #if?
-		// Seems this can occur.
 		return;
-#endif
 	}
 
 	// Unlock window when we're ready.
