- Fixing broken reuse of existing downloads
- Added ability to retry stale background downloads when we're in foreground on iOS
- Fixing race condition on background file hash helper invoked from different threads

diff --git a/Engine/Config/IOS/BaseIOSEngine.ini b/Engine/Config/IOS/BaseIOSEngine.ini
--- a/Engine/Config/IOS/BaseIOSEngine.ini
+++ b/Engine/Config/IOS/BaseIOSEngine.ini
@@ -45,6 +45,8 @@ bUseForegroundSession=false
 bDiscretionary=false
 ; Wakeup app in background and report state of current downloads, allows to retry downloads in background
 bShouldSendLaunchEvents=true
+; Maximum connenctions per host for whole session
+MaximumConnectionsPerHost=6
 ; Time out and retry if we don't get a response in this timeframe if app is in foreground and bUseForegroundSession is true
 BackgroundReceiveTimeout=120
 ; Time out and retry if the download hasn't finished in this long, regardless of app state 
@@ -54,6 +56,10 @@ RetryResumeDataLimit=3
 ; How many milliseconds to wait for HEAD replies on first CDN request to decide if CDN is reachable or not (0 = disabled).
 ; Unreachable CDNs will be deprioritized in all download requests.
 CDNReorderingTimeout=400
+; When in foreground every N seconds check if downloads are not receiving bytes (regardless of bUseForegroundSession), retry them if so.
+CheckForForegroundStaleDownloadsWithInterval=1
+; How long to wait in seconds before considering retrying download in CheckForForegroundStaleDownloadsWithInterval
+ForegroundStaleDownloadTimeout=60
 
 [ConsoleVariables]
 Slate.CacheRenderData=0
diff --git a/Engine/Source/Runtime/ApplicationCore/Private/IOS/IOSBackgroundURLSessionHandler.cpp b/Engine/Source/Runtime/ApplicationCore/Private/IOS/IOSBackgroundURLSessionHandler.cpp
--- a/Engine/Source/Runtime/ApplicationCore/Private/IOS/IOSBackgroundURLSessionHandler.cpp
+++ b/Engine/Source/Runtime/ApplicationCore/Private/IOS/IOSBackgroundURLSessionHandler.cpp
@@ -208,6 +208,7 @@ NSString* const SerializationKeyRetryCountPerURL = @"r";
 - (void)Initialize;
 - (void)SetFileHashHelper:(BackgroundHttpFileHashHelperRef)HelperRef;
 - (BackgroundHttpFileHashHelperRef)GetFileHashHelper;
+- (void)SaveFileHashHelperState;
 - (NSString*)GetTempPathForURL:(NSURL* _Nonnull)URL;
 - (NSMutableArray<__kindof NSURL*>*)ReorderCDNsByReachability:(NSMutableArray<__kindof NSURL*>*)URLs;
 - (NSURLSessionDownloadTask*)CreateDownloadForURL:(NSURL* _Nonnull)URL WithPriority:(float)Priority WithTaskData:(FBackgroundNSURLSessionDownloadTaskData* _Nonnull)TaskData;
@@ -220,8 +221,13 @@ NSString* const SerializationKeyRetryCountPerURL = @"r";
 - (void)SetPriority:(float)Priority ForDownload:(NSUInteger)DownloadId;
 - (void)SetCurrentDownloadedBytes:(uint64)DownloadedBytes ForTask:(NSURLSessionDownloadTask*)Task;
 - (uint64)GetCurrentDownloadedBytes:(NSUInteger)DownloadId;
+- (void)RecreateDownload:(NSUInteger)DownloadId ShouldResetRetryCount:(bool)ResetRetryCount;
 - (void)RecreateDownloads;
 
+- (void)StartCheckingForStaleDownloads;
+- (void)StopCheckingForStaleDownloads;
+- (void)CheckForStaleDownloads;
+
 - (NSURLSessionDownloadTask*)FindDownloadTaskFor:(NSUInteger)DownloadId;
 - (NSUInteger)FindDownloadIdForTask:(NSURLSessionDownloadTask*)Task;
 - (NSUInteger)EnsureTaskIsTracked:(NSURLSessionDownloadTask*)Task;
@@ -236,12 +242,13 @@ NSString* const SerializationKeyRetryCountPerURL = @"r";
 
 // From NSURLSessionTaskDelegate
 - (void)URLSession:(NSURLSession*)Session task:(NSURLSessionTask*)Task didCompleteWithError:(NSError*)Error;
-//- (void)URLSession:(NSURLSession *)session taskIsWaitingForConnectivity:(NSURLSessionTask *)Task;
-//- (void)URLSession:(NSURLSession *)Session task:(NSURLSessionTask *)Task willBeginDelayedRequest:(NSURLRequest *)Request completionHandler:(void (^)(NSURLSessionDelayedRequestDisposition Disposition, NSURLRequest* NewRequest))CompletionHandler;
+//- (void)URLSession:(NSURLSession*)Session taskIsWaitingForConnectivity:(NSURLSessionTask *)Task;
+//- (void)URLSession:(NSURLSession*)Session task:(NSURLSessionTask *)Task willBeginDelayedRequest:(NSURLRequest *)Request completionHandler:(void (^)(NSURLSessionDelayedRequestDisposition Disposition, NSURLRequest* NewRequest))CompletionHandler;
 
 // From NSURLSessionDownloadDelegate
 - (void)URLSession:(NSURLSession*)Session downloadTask:(NSURLSessionDownloadTask*)Task didFinishDownloadingToURL:(NSURL*)Location;
-- (void)URLSession:(NSURLSession*)session downloadTask:(NSURLSessionDownloadTask*)Task didWriteData:(int64_t)BytesWritten totalBytesWritten:(int64_t)TotalBytesWritten totalBytesExpectedToWrite:(int64_t)TotalBytesExpectedToWrite;
+- (void)URLSession:(NSURLSession*)Session downloadTask:(NSURLSessionDownloadTask*)Task didWriteData:(int64_t)BytesWritten totalBytesWritten:(int64_t)TotalBytesWritten totalBytesExpectedToWrite:(int64_t)TotalBytesExpectedToWrite;
+//- (void)URLSession:(NSURLSession*)Session downloadTask:(NSURLSessionDownloadTask*)Task didResumeAtOffset:(int64_t)FileOffset expectedTotalBytes:(int64_t)ExpectedTotalBytes;
 
 @end
 
@@ -255,16 +262,20 @@ NSString* const SerializationKeyRetryCountPerURL = @"r";
 	std::promise<void> _AllDownloadsPromise;
 	std::future<void> _AllDownloadsFuture;
 	NSSet<__kindof NSString*>* _UnreachableCDNs;
+	NSTimer* _ForegroundStaleDownloadCheckTimer;
 	BackgroundHttpFileHashHelperPtr _HelperPtr;
 	int32 MaximumConnectionsPerHost;
 	int32 RetryResumeDataLimit;
 	int32 CDNReorderingTimeout;
+	double CheckForForegroundStaleDownloadsWithInterval;
+	double ForegroundStaleDownloadTimeout;
 }
 
 static constexpr NSUInteger InvalidDownloadId = 0;
 
 NSString* const NSURLSessionIdentifier = @"com.epicgames.backgrounddownloads";
 
+NSProgressUserInfoKey const NSProgressDownloadLastUpdateTime = @"com.epicgames.nsprogress.lastupdatetime";
 NSProgressUserInfoKey const NSProgressDownloadCompletedBytes = @"com.epicgames.nsprogress.completedbytes";
 NSProgressUserInfoKey const NSProgressDownloadResultStatusCode = @"com.epicgames.nsprogress.resultstatuscode";
 NSProgressUserInfoKey const NSProgressDownloadResultTempFilePath = @"com.epicgames.nsprogress.tempfilepath";
@@ -312,22 +323,38 @@ static constexpr NSInteger HTTPStatusCodeErrorServer = 500;
 	double TimeoutIntervalForResource = 60.0 * 60.0;
 	RetryResumeDataLimit = 3;
 	CDNReorderingTimeout = 400;
+	CheckForForegroundStaleDownloadsWithInterval = 1.0; // how often to check for stale downloads, <=0.0 to disable
+	ForegroundStaleDownloadTimeout = 30.0; // If download hasn't received any bytes for this duration, cancel and retry if possible
 
 #ifndef UE_DNLD_SANDBOX
 	GConfig->GetBool(TEXT("BackgroundHttp.iOSSettings"), TEXT("bUseForegroundSession"), bUseForegroundSession, GEngineIni);
 	GConfig->GetBool(TEXT("BackgroundHttp.iOSSettings"), TEXT("bDiscretionary"), bDiscretionary, GEngineIni);
 	GConfig->GetBool(TEXT("BackgroundHttp.iOSSettings"), TEXT("bShouldSendLaunchEvents"), bShouldSendLaunchEvents, GEngineIni);
-	GConfig->GetInt(TEXT("BackgroundHttp"), TEXT("MaxActiveDownloads"), MaximumConnectionsPerHost, GEngineIni);
+	GConfig->GetInt(TEXT("BackgroundHttp.iOSSettings"), TEXT("MaximumConnectionsPerHost"), MaximumConnectionsPerHost, GEngineIni);
 	GConfig->GetDouble(TEXT("BackgroundHttp.iOSSettings"), TEXT("BackgroundReceiveTimeout"), TimeoutIntervalForRequest, GEngineIni);
 	GConfig->GetDouble(TEXT("BackgroundHttp.iOSSettings"), TEXT("BackgroundHttpResourceTimeout"), TimeoutIntervalForResource, GEngineIni);
 	GConfig->GetInt(TEXT("BackgroundHttp.iOSSettings"), TEXT("RetryResumeDataLimit"), RetryResumeDataLimit, GEngineIni);
 	GConfig->GetInt(TEXT("BackgroundHttp.iOSSettings"), TEXT("CDNReorderingTimeout"), CDNReorderingTimeout, GEngineIni);
+	GConfig->GetDouble(TEXT("BackgroundHttp.iOSSettings"), TEXT("CheckForForegroundStaleDownloadsWithInterval"), CheckForForegroundStaleDownloadsWithInterval, GEngineIni);
+	GConfig->GetDouble(TEXT("BackgroundHttp.iOSSettings"), TEXT("ForegroundStaleDownloadTimeout"), ForegroundStaleDownloadTimeout, GEngineIni);
 #endif
-	
+
+	UE_DNLD_LOG(@"bUseForegroundSession=%u", bUseForegroundSession ? 1 : 0);
+	UE_DNLD_LOG(@"bDiscretionary=%u", bDiscretionary ? 1 : 0);
+	UE_DNLD_LOG(@"bShouldSendLaunchEvents=%u", bShouldSendLaunchEvents ? 1 : 0);
+	UE_DNLD_LOG(@"MaximumConnectionsPerHost=%i", MaximumConnectionsPerHost);
+	UE_DNLD_LOG(@"TimeoutIntervalForRequest=%f", TimeoutIntervalForRequest);
+	UE_DNLD_LOG(@"TimeoutIntervalForResource=%f", TimeoutIntervalForResource);
+	UE_DNLD_LOG(@"RetryResumeDataLimit=%i", RetryResumeDataLimit);
+	UE_DNLD_LOG(@"CDNReorderingTimeout=%i", CDNReorderingTimeout);
+	UE_DNLD_LOG(@"CheckForForegroundStaleDownloadsWithInterval=%f", CheckForForegroundStaleDownloadsWithInterval);
+	UE_DNLD_LOG(@"ForegroundStaleDownloadTimeout=%f", ForegroundStaleDownloadTimeout);
+
 	_AllDownloads = [NSMutableDictionary new];
 	_AllDownloadsFuture = _AllDownloadsPromise.get_future();
 	_NextDownloadId = InvalidDownloadId + 1;
 	_UnreachableCDNs = nil;
+	_ForegroundStaleDownloadCheckTimer = nil;
 
 	// Never allow cellular unless we get explicit opt-in from the user.
 	self.AllowCellular = NO;
@@ -366,7 +393,7 @@ static constexpr NSInteger HTTPStatusCodeErrorServer = 500;
 	{
 		UE_DNLD_LOG(@"getTasksWithCompletionHandler block with %lu tasks", (unsigned long)[Downloads count]);
 		
-		if (Downloads == nil)
+		if (Downloads != nil)
 		{
 			for (NSURLSessionDownloadTask* ExistingTask in Downloads)
 			{
@@ -387,9 +414,10 @@ static constexpr NSInteger HTTPStatusCodeErrorServer = 500;
 					UE_DNLD_LOG(@"Canceling existing download task with taskIdentifier %lu", ExistingTask.taskIdentifier);
 					[self CancelDownload:DownloadId];
 				}
+
+				UE_DNLD_LOG(@"Existing task '%@'", ExistingTask.taskDescription);
 			}
 		}
-		
 
 		_AllDownloadsPromise.set_value();
 	}];
@@ -421,6 +449,10 @@ static constexpr NSInteger HTTPStatusCodeErrorServer = 500;
 	{
 		[_UnreachableCDNs release];
 	}
+	if (_ForegroundStaleDownloadCheckTimer != nil)
+	{
+		[_ForegroundStaleDownloadCheckTimer release];
+	}
 
 	[super dealloc];
 }
@@ -442,18 +474,29 @@ static constexpr NSInteger HTTPStatusCodeErrorServer = 500;
 	return _HelperPtr.ToSharedRef();
 }
 
+- (void)SaveFileHashHelperState
+{
+	@synchronized (self)
+	{
+		[self GetFileHashHelper]->SaveData();
+	}
+}
+
 - (NSString*)GetTempPathForURL:(NSURL* _Nonnull)URL
 {
-	BackgroundHttpFileHashHelperRef HelperRef = [self GetFileHashHelper];
+	@synchronized (self)
+	{
+		BackgroundHttpFileHashHelperRef HelperRef = [self GetFileHashHelper];
 
-	const FString TaskURL(URL.absoluteString);
-	const FString& TempFileName = HelperRef->FindOrAddTempFilenameMappingForURL(TaskURL);
-	const FString& DestinationPath = HelperRef->GetFullPathOfTempFilename(TempFileName);
+		const FString TaskURL(URL.absoluteString);
+		const FString& TempFileName = HelperRef->FindOrAddTempFilenameMappingForURL(TaskURL);
+		const FString& DestinationPath = HelperRef->GetFullPathOfTempFilename(TempFileName);
 
-	IPlatformFile& PlatformFile = FPlatformFileManager::Get().GetPlatformFile();
-	const FString ConvertedPath = PlatformFile.ConvertToAbsolutePathForExternalAppForWrite(*DestinationPath);
+		IPlatformFile& PlatformFile = FPlatformFileManager::Get().GetPlatformFile();
+		const FString ConvertedPath = PlatformFile.ConvertToAbsolutePathForExternalAppForWrite(*DestinationPath);
 
-	return ConvertedPath.GetNSString();
+		return ConvertedPath.GetNSString();
+	}
 }
 
 - (NSMutableArray<__kindof NSURL*>*)ReorderCDNsByReachability:(NSMutableArray<__kindof NSURL*>*)URLs
@@ -639,29 +682,33 @@ static constexpr NSInteger HTTPStatusCodeErrorServer = 500;
 	UE_DNLD_LOG(@"Trying to find existing download for asset path '%@'", AssetPath);
 	if (AssetPath != nil)
 	{
-		__block NSUInteger ExistingDownloadId = InvalidDownloadId;
-		__block NSURLSessionDownloadTask* ExistingTask = nil;
-		[_AllDownloads enumerateKeysAndObjectsUsingBlock:^(NSNumber* IterKey, NSURLSessionDownloadTask* IterTask, BOOL* IterStop)
+		@synchronized(_AllDownloads)
 		{
-			if (IterTask.originalRequest != nil && [IterTask.originalRequest.URL.path isEqualToString:AssetPath])
+			__block NSUInteger ExistingDownloadId = InvalidDownloadId;
+			__block NSURLSessionDownloadTask* ExistingTask = nil;
+			[_AllDownloads enumerateKeysAndObjectsUsingBlock:^(NSNumber* IterKey, NSURLSessionDownloadTask* IterTask, BOOL* IterStop)
+			 {
+				if ((IterTask.originalRequest != nil && [IterTask.originalRequest.URL.path isEqualToString:AssetPath])
+					|| (IterTask.currentRequest != nil && [IterTask.currentRequest.URL.path isEqualToString:AssetPath]))
+				{
+					ExistingDownloadId = IterKey.unsignedIntegerValue;
+					ExistingTask = [IterTask retain]; // Retain in case if task gets killed in another thread.
+					*IterStop = YES;
+				}
+			}];
+
+			if (ExistingDownloadId != InvalidDownloadId && ExistingTask != nil)
 			{
-				ExistingDownloadId = IterKey.unsignedIntegerValue;
-				ExistingTask = [IterTask retain]; // Retain in case if task gets killed in another thread.
-				*IterStop = YES;
+				UE_DNLD_LOG(@"Found existing download task for path '%@' with DownloadId %lu", AssetPath, ExistingDownloadId);
+				
+				// Update existing task state to new one, to reset retry counters, cdn links, etc.
+				[ExistingTask setTaskDescription:[TaskData ToSerializedString]];
+				
+				[ExistingTask resume]; // Resume task just in case if it was not running before.
+				[ExistingTask release];
+				
+				return ExistingDownloadId;
 			}
-		}];
-		
-		if (ExistingDownloadId != InvalidDownloadId && ExistingTask != nil)
-		{
-			UE_DNLD_LOG(@"Found existing download task for path '%@' with DownloadId %lu", AssetPath, ExistingDownloadId);
-
-			// Update existing task state to new one, to reset retry counters, cdn links, etc.
-			[ExistingTask setTaskDescription:[TaskData ToSerializedString]];
-
-			[ExistingTask resume]; // Resume task just in case if it was not running before.
-			[ExistingTask release];
-
-			return ExistingDownloadId;
 		}
 	}
 
@@ -735,6 +782,7 @@ static constexpr NSInteger HTTPStatusCodeErrorServer = 500;
 {
 	if (Task != nil)
 	{
+		[Task.progress setUserInfoObject:[NSNumber numberWithDouble:[[NSDate date] timeIntervalSince1970]] forKey:NSProgressDownloadLastUpdateTime];
 		[Task.progress setUserInfoObject:[NSNumber numberWithLongLong:DownloadedBytes] forKey:NSProgressDownloadCompletedBytes];
 	}
 }
@@ -754,6 +802,35 @@ static constexpr NSInteger HTTPStatusCodeErrorServer = 500;
 	return 0;
 }
 
+- (void)RecreateDownload:(NSUInteger)DownloadId ShouldResetRetryCount:(bool)ResetRetryCount
+{
+	NSURLSessionDownloadTask* OldTask = [self FindDownloadTaskFor:DownloadId];
+	const float OldTaskPriority = OldTask.priority;
+	const NSURLSessionTaskState OldTaskState = OldTask.state;
+
+	FBackgroundNSURLSessionDownloadTaskData* NewTaskData = [FBackgroundNSURLSessionDownloadTaskData TaskDataFromSerializedString:OldTask.taskDescription];
+	if (NewTaskData == nil)
+	{
+		return;
+	}
+
+	// Cancel old task
+	[self CancelDownload:DownloadId];
+
+	if (ResetRetryCount)
+	{
+		[NewTaskData ResetRetryCount:RetryResumeDataLimit];
+	}
+
+	// Start a new task
+	NSURLSessionDownloadTask* NewTask = [self CreateDownloadForURL:[NewTaskData GetFirstURL] WithPriority:OldTaskPriority WithTaskData:NewTaskData];
+	[self ReplaceTrackedTaskWith:NewTask ForDownloadId:DownloadId];
+	if (OldTaskState != NSURLSessionTaskStateSuspended)
+	{
+		[NewTask resume];
+	}
+}
+
 - (void)RecreateDownloads
 {
 	UE_DNLD_LOG(@"RecreateDownloads started");
@@ -768,30 +845,84 @@ static constexpr NSInteger HTTPStatusCodeErrorServer = 500;
 	for (NSNumber* IterKey in AllKeys)
 	{
 		const NSUInteger DownloadId = IterKey.unsignedIntegerValue;
-		NSURLSessionDownloadTask* OldTask = [_AllDownloads objectForKey:IterKey];
-		const float OldTaskPriority = OldTask.priority;
-		const NSURLSessionTaskState OldTaskState = OldTask.state;
+		[self RecreateDownload:DownloadId ShouldResetRetryCount:true];
+	}
+
+	UE_DNLD_LOG(@"RecreateDownloads finished");
+}
+
+- (void)StartCheckingForStaleDownloads
+{
+	if (_ForegroundStaleDownloadCheckTimer == nil && CheckForForegroundStaleDownloadsWithInterval > 0.0)
+	{
+		_ForegroundStaleDownloadCheckTimer = [[NSTimer
+											   scheduledTimerWithTimeInterval:CheckForForegroundStaleDownloadsWithInterval
+											   target:self
+											   selector:@selector(CheckForStaleDownloads)
+											   userInfo:nil
+											   repeats:YES]
+											  retain];
+		_ForegroundStaleDownloadCheckTimer.tolerance = CheckForForegroundStaleDownloadsWithInterval * 0.5;
+
+		UE_DNLD_LOG(@"Start checking for stale downloads");
+	}
+}
+
+- (void)StopCheckingForStaleDownloads
+{
+	if (_ForegroundStaleDownloadCheckTimer != nil)
+	{
+		[_ForegroundStaleDownloadCheckTimer invalidate];
+		[_ForegroundStaleDownloadCheckTimer release];
+		_ForegroundStaleDownloadCheckTimer = nil;
+
+		UE_DNLD_LOG(@"Stop checking for stale downloads");
+	}
+}
+
+- (void)CheckForStaleDownloads
+{
+	// Copy keys to avoid deadlocking in case if cancel/resume/etc will call delegates in-place
+	NSArray<__kindof NSNumber*>* AllKeys = nil;
+	@synchronized(_AllDownloads)
+	{
+		AllKeys = [_AllDownloads allKeys];
+	}
+
+	const NSTimeInterval CurrentTime = [[NSDate date] timeIntervalSince1970];
+
+	for (NSNumber* IterKey in AllKeys)
+	{
+		const NSUInteger DownloadId = IterKey.unsignedIntegerValue;
+		NSURLSessionDownloadTask* Task = [_AllDownloads objectForKey:IterKey];
 
-		FBackgroundNSURLSessionDownloadTaskData* NewTaskData = [FBackgroundNSURLSessionDownloadTaskData TaskDataFromSerializedString:OldTask.taskDescription];
-		if (NewTaskData == nil)
+		NSNumber* ResultStatusCode = [Task.progress.userInfo objectForKey:NSProgressDownloadResultStatusCode];
+		if (ResultStatusCode != nil)
 		{
+			// Task is finished downloading, skip.
 			continue;
 		}
 
-		// Cancel old task
-		[self CancelDownload:DownloadId];
+		NSNumber* LastUpdateTimeNumber = [Task.progress.userInfo objectForKey:NSProgressDownloadLastUpdateTime];
+		if (LastUpdateTimeNumber == nil)
+		{
+			// There is no known last update time, skip.
+			// As we don't get any notification from background tasks on when they started,
+			// we can't distinguish between task pending processing vs task stuck on getting first byte.
+			continue;
+		}
 
-		// Start a new task
-		[NewTaskData ResetRetryCount:RetryResumeDataLimit];
-		NSURLSessionDownloadTask* NewTask = [self CreateDownloadForURL:[NewTaskData GetFirstURL] WithPriority:OldTaskPriority WithTaskData:NewTaskData];
-		[self ReplaceTrackedTaskWith:NewTask ForDownloadId:DownloadId];
-		if (OldTaskState != NSURLSessionTaskStateSuspended)
+		const NSTimeInterval TimeSinceLastUpdate = CurrentTime - LastUpdateTimeNumber.doubleValue;
+		if (TimeSinceLastUpdate >= ForegroundStaleDownloadTimeout)
 		{
-			[NewTask resume];
+			UE_DNLD_LOG(@"Task '%@' with taskIdentifier %lu is considering stale, retrying", Task.taskDescription, Task.taskIdentifier);
+
+			// Clear last update property to avoid canceling task twice in next tick.
+			[Task.progress setUserInfoObject:nil forKey:NSProgressDownloadLastUpdateTime];
+
+			[self RecreateDownload:DownloadId ShouldResetRetryCount:false];
 		}
 	}
-
-	UE_DNLD_LOG(@"RecreateDownloads finished");
 }
 
 - (NSURLSessionDownloadTask*)FindDownloadTaskFor:(NSUInteger)DownloadId
@@ -833,6 +964,10 @@ static constexpr NSInteger HTTPStatusCodeErrorServer = 500;
 	{
 		const NSUInteger DownloadId = _NextDownloadId++;
 		[_AllDownloads setObject:Task forKey:[NSNumber numberWithUnsignedInteger:DownloadId]];
+
+		// start timer as long as we have active tasks
+		[self StartCheckingForStaleDownloads];
+
 		return DownloadId;
 	}
 }
@@ -853,6 +988,12 @@ static constexpr NSInteger HTTPStatusCodeErrorServer = 500;
 		@synchronized(_AllDownloads)
 		{
 			[_AllDownloads removeObjectForKey:[NSNumber numberWithUnsignedInteger:ExistingDownloadId]];
+
+			if (_AllDownloads.count == 0)
+			{
+				// stop timer when we have no ongoing tasks
+				[self StopCheckingForStaleDownloads];
+			}
 		}
 	}
 }
@@ -994,12 +1135,12 @@ static constexpr NSInteger HTTPStatusCodeErrorServer = 500;
 	}
 }
 
-//- (void)URLSession:(NSURLSession *)session taskIsWaitingForConnectivity:(NSURLSessionTask *)Task
+//- (void)URLSession:(NSURLSession*)Session taskIsWaitingForConnectivity:(NSURLSessionTask*)Task
 //{
 //	UE_DNLD_LOG(@"taskIsWaitingForConnectivity '%@'", Task.taskDescription);
 //}
 
-//- (void)URLSession:(NSURLSession *)Session task:(NSURLSessionTask *)Task willBeginDelayedRequest:(NSURLRequest *)Request completionHandler:(void (^)(NSURLSessionDelayedRequestDisposition Disposition, NSURLRequest* NewRequest))CompletionHandler
+//- (void)URLSession:(NSURLSession*)Session task:(NSURLSessionTask*)Task willBeginDelayedRequest:(NSURLRequest*)Request completionHandler:(void (^)(NSURLSessionDelayedRequestDisposition Disposition, NSURLRequest* NewRequest))CompletionHandler
 //{
 //	UE_DNLD_LOG(@"willBeginDelayedRequest '%@' for '%@' with taskIdentifier %lu", Task.taskDescription, Request.debugDescription, Task.taskIdentifier);
 //	CompletionHandler(NSURLSessionDelayedRequestContinueLoading, Request);
@@ -1039,11 +1180,17 @@ static constexpr NSInteger HTTPStatusCodeErrorServer = 500;
 	}
 }
 
-- (void)URLSession:(NSURLSession*)session downloadTask:(NSURLSessionDownloadTask*)Task didWriteData:(int64_t)BytesWritten totalBytesWritten:(int64_t)TotalBytesWritten totalBytesExpectedToWrite:(int64_t)TotalBytesExpectedToWrite
+- (void)URLSession:(NSURLSession*)Session downloadTask:(NSURLSessionDownloadTask*)Task didWriteData:(int64_t)BytesWritten totalBytesWritten:(int64_t)TotalBytesWritten totalBytesExpectedToWrite:(int64_t)TotalBytesExpectedToWrite
 {
+	//UE_DNLD_LOG(@"didWriteData taskIdentifier %lu wrote %lli", Task.taskIdentifier, TotalBytesWritten);
 	[self SetCurrentDownloadedBytes:TotalBytesWritten ForTask:Task];
 }
 
+//- (void)URLSession:(NSURLSession*)Session downloadTask:(NSURLSessionDownloadTask*)Task didResumeAtOffset:(int64_t)FileOffset expectedTotalBytes:(int64_t)ExpectedTotalBytes
+//{
+//	UE_DNLD_LOG(@"didResumeAtOffset taskIdentifier %lu offset %lli total %lli", Task.taskIdentifier, FileOffset, ExpectedTotalBytes);
+//}
+
 @end
 
 const uint64 FBackgroundURLSessionHandler::InvalidDownloadId = [FBackgroundNSURLSession GetInvalidDownloadId];
@@ -1162,3 +1309,11 @@ void FBackgroundURLSessionHandler::HandleEventsForBackgroundURLSession(const FSt
 		// will invoke URLSessionDidFinishEventsForBackgroundURLSession internally.
 	}
 }
+
+void FBackgroundURLSessionHandler::SaveBackgroundHttpFileHashHelperState()
+{
+	@autoreleasepool
+	{
+		[[FBackgroundNSURLSession Shared] SaveFileHashHelperState];
+	}
+}
diff --git a/Engine/Source/Runtime/ApplicationCore/Public/IOS/IOSBackgroundURLSessionHandler.h b/Engine/Source/Runtime/ApplicationCore/Public/IOS/IOSBackgroundURLSessionHandler.h
--- a/Engine/Source/Runtime/ApplicationCore/Public/IOS/IOSBackgroundURLSessionHandler.h
+++ b/Engine/Source/Runtime/ApplicationCore/Public/IOS/IOSBackgroundURLSessionHandler.h
@@ -83,4 +83,7 @@ public:
 
 	// To be used by app delegate, call it from handleEventsForBackgroundURLSession.
 	static void HandleEventsForBackgroundURLSession(const FString& SessionIdentifier);
+
+	// To be used by ApplePlatformBackgroundHttpManager
+	static void SaveBackgroundHttpFileHashHelperState();
 };
diff --git a/Engine/Source/Runtime/Online/BackgroundHTTP/Private/IOS/ApplePlatformBackgroundHttpManager.cpp b/Engine/Source/Runtime/Online/BackgroundHTTP/Private/IOS/ApplePlatformBackgroundHttpManager.cpp
--- a/Engine/Source/Runtime/Online/BackgroundHTTP/Private/IOS/ApplePlatformBackgroundHttpManager.cpp
+++ b/Engine/Source/Runtime/Online/BackgroundHTTP/Private/IOS/ApplePlatformBackgroundHttpManager.cpp
@@ -98,7 +98,7 @@ bool FApplePlatformBackgroundHttpManager::Tick(float DeltaTime)
 		PendingRemoveRequests.Empty();
 	}
 
-	GetFileHashHelper()->SaveData();
+	FBackgroundURLSessionHandler::SaveBackgroundHttpFileHashHelperState();
 
 	return true;
 }
