Adding native resolution info to Core.ClientPerformance and Core.AppStartup

#rnx

diff --git a/Engine/Source/Runtime/ApplicationCore/Private/Android/AndroidApplication.cpp b/Engine/Source/Runtime/ApplicationCore/Private/Android/AndroidApplication.cpp
--- a/Engine/Source/Runtime/ApplicationCore/Private/Android/AndroidApplication.cpp
+++ b/Engine/Source/Runtime/ApplicationCore/Private/Android/AndroidApplication.cpp
@@ -155,6 +155,17 @@ bool FAndroidApplication::IsGamepadAttached() const
 	return false;
 }
 
+bool FAndroidApplication::GetNativeWindowResolution(int32_t& OutWidth, int32_t& OutHeight) const
+{
+	if (Windows.IsEmpty()) return false;
+
+	const TSharedPtr<FAndroidWindow>& Window = Windows[0];
+
+	if (!Window.IsValid()) return false;
+
+	return Window->GetNativeWindowResolution(OutWidth, OutHeight);
+}
+
 void FDisplayMetrics::RebuildDisplayMetrics( FDisplayMetrics& OutDisplayMetrics )
 {
 	// Get screen rect
diff --git a/Engine/Source/Runtime/ApplicationCore/Private/Android/AndroidWindow.cpp b/Engine/Source/Runtime/ApplicationCore/Private/Android/AndroidWindow.cpp
--- a/Engine/Source/Runtime/ApplicationCore/Private/Android/AndroidWindow.cpp
+++ b/Engine/Source/Runtime/ApplicationCore/Private/Android/AndroidWindow.cpp
@@ -415,6 +415,28 @@ void FAndroidWindow::EventManagerUpdateWindowDimensions(int32 Width, int32 Heigh
 	}
 }
 
+bool FAndroidWindow::GetNativeWindowResolution(int32_t& OutWidth, int32_t& OutHeight) const
+{
+#if USE_ANDROID_JNI
+
+	if (NativeWindow == nullptr) return false;
+
+	ANativeWindow* AndroidWindow = static_cast<ANativeWindow*>(NativeWindow);
+
+	OutWidth = ANativeWindow_getWidth(AndroidWindow);
+	OutHeight = ANativeWindow_getHeight(AndroidWindow);
+
+	return true;
+
+#else
+
+	// Android without JNI?
+	// Just making the function work for unexpected platforms
+	return FPlatformMisc::GetOverrideResolution(OutWidth, OutHeight);
+
+#endif
+}
+
 void* FAndroidWindow::WaitForHardwareWindow()
 {
 	// Sleep if the hardware window isn't currently available.
diff --git a/Engine/Source/Runtime/ApplicationCore/Public/Android/AndroidApplication.h b/Engine/Source/Runtime/ApplicationCore/Public/Android/AndroidApplication.h
--- a/Engine/Source/Runtime/ApplicationCore/Public/Android/AndroidApplication.h
+++ b/Engine/Source/Runtime/ApplicationCore/Public/Android/AndroidApplication.h
@@ -105,6 +105,8 @@ public:
 
 	virtual bool IsGamepadAttached() const override;
 
+	bool GetNativeWindowResolution(int32_t& OutWidth, int32_t& OutHeight) const;
+
 	static TAutoConsoleVariable<bool> CVarAndroidSupportsTimestampQueries;
 	static TAutoConsoleVariable<bool> CVarAndroidSupportsDynamicResolution;
 
diff --git a/Engine/Source/Runtime/ApplicationCore/Public/Android/AndroidWindow.h b/Engine/Source/Runtime/ApplicationCore/Public/Android/AndroidWindow.h
--- a/Engine/Source/Runtime/ApplicationCore/Public/Android/AndroidWindow.h
+++ b/Engine/Source/Runtime/ApplicationCore/Public/Android/AndroidWindow.h
@@ -69,6 +69,8 @@ public:
 	// Called by the event manager to update the cached window dimensions to match the event it is processing.
 	static APPLICATIONCORE_API void EventManagerUpdateWindowDimensions(int32 Width, int32 Height);
 
+	APPLICATIONCORE_API bool GetNativeWindowResolution(int32_t& OutWidth, int32_t& OutHeight) const;
+
 protected:
 	/** @return true if the native window is currently in fullscreen mode, false otherwise */
 	virtual EWindowMode::Type GetWindowMode() const override { return EWindowMode::Fullscreen; }
