Fixing FBackgroundURLSessionHandler::CancelDownload crash

diff --git a/Engine/Source/Runtime/ApplicationCore/Private/IOS/IOSBackgroundURLSessionHandler.cpp b/Engine/Source/Runtime/ApplicationCore/Private/IOS/IOSBackgroundURLSessionHandler.cpp
--- a/Engine/Source/Runtime/ApplicationCore/Private/IOS/IOSBackgroundURLSessionHandler.cpp
+++ b/Engine/Source/Runtime/ApplicationCore/Private/IOS/IOSBackgroundURLSessionHandler.cpp
@@ -290,7 +290,7 @@ static constexpr NSInteger HTTPStatusCodeErrorServer = 500;
 + (FBackgroundNSURLSession*)Shared
 {
 	static FBackgroundNSURLSession* Shared = nil;
-	static dispatch_once_t Once;
+	static dispatch_once_t Once = {};
 	dispatch_once(&Once, ^{
 		Shared = [[self alloc] init];
 	});
@@ -447,14 +447,21 @@ static constexpr NSInteger HTTPStatusCodeErrorServer = 500;
 - (void)dealloc
 {
 	[_Session release];
+	_Session = nil;
+
 	[_AllDownloads release];
+	_AllDownloads = nil;
+
 	if (_UnreachableCDNs != nil)
 	{
 		[_UnreachableCDNs release];
+		_UnreachableCDNs = nil;
 	}
+
 	if (_ForegroundStaleDownloadCheckTimer != nil)
 	{
 		[_ForegroundStaleDownloadCheckTimer release];
+		_ForegroundStaleDownloadCheckTimer = nil;
 	}
 
 	[super dealloc];
@@ -763,11 +770,17 @@ static constexpr NSInteger HTTPStatusCodeErrorServer = 500;
 		Task.taskDescription = [TaskData ToSerializedString];
 	}
 
+	// Must retain otherwise EnsureTaskIsNotTracked will free the object.
+	[Task retain];
+
 	// We're done with this task.
 	[self EnsureTaskIsNotTracked:Task];
 
 	// Will invoke didCompleteWithError if task is incomplete.
 	[Task cancel];
+	
+	// Finally release the object.
+	[Task release];
 }
 
 - (void)SetPriority:(float)Priority ForDownload:(NSUInteger)DownloadId
@@ -818,6 +831,7 @@ static constexpr NSInteger HTTPStatusCodeErrorServer = 500;
 	}
 
 	// Cancel old task
+	OldTask = nil;
 	[self CancelDownload:DownloadId];
 
 	if (ResetRetryCount)
@@ -951,6 +965,7 @@ static constexpr NSInteger HTTPStatusCodeErrorServer = 500;
 
 			// Clear last update property to avoid canceling task twice in next tick.
 			[Task.progress setUserInfoObject:nil forKey:NSProgressDownloadLastUpdateTime];
+			Task = nil;
 
 			[self RecreateDownload:DownloadId ShouldResetRetryCount:false];
 		}
@@ -1124,7 +1139,10 @@ static constexpr NSInteger HTTPStatusCodeErrorServer = 500;
 				UE_DNLD_LOG(@"didCompleteWithError, task '%@' with taskIdentifier %lu failed due to '%@' and has resume data and next url is the same, retrying", Task.taskDescription, Task.taskIdentifier, LocalizedDescription);
 
 				NSURLSessionDownloadTask* NewTask = [self CreateDownloadForResumeData:ResumeData WithPriority:Task.priority WithTaskData:TaskData];
+
+				Task = nil;
 				[self ReplaceTrackedTaskWith:NewTask ForDownloadId:DownloadId];
+
 				[NewTask resume];
 				return;
 			}
@@ -1141,7 +1159,10 @@ static constexpr NSInteger HTTPStatusCodeErrorServer = 500;
 				}
 
 				NSURLSessionDownloadTask* NewTask = [self CreateDownloadForURL:NextURL WithPriority:Task.priority WithTaskData:TaskData];
+
+				Task = nil;
 				[self ReplaceTrackedTaskWith:NewTask ForDownloadId:DownloadId];
+
 				[NewTask resume];
 				return;
 			}
