Pass iOS background download log to UE instead of NSLog

diff --git a/Engine/Source/Runtime/ApplicationCore/Private/IOS/IOSBackgroundURLSessionHandler.cpp b/Engine/Source/Runtime/ApplicationCore/Private/IOS/IOSBackgroundURLSessionHandler.cpp
--- a/Engine/Source/Runtime/ApplicationCore/Private/IOS/IOSBackgroundURLSessionHandler.cpp
+++ b/Engine/Source/Runtime/ApplicationCore/Private/IOS/IOSBackgroundURLSessionHandler.cpp
@@ -20,7 +20,25 @@
 // Useful when testing background downloads as they persist between application sessions.
 static constexpr bool bCancelExistingDownloads = false;
 
-#define UE_DNLD_LOG(...) NSLog(@"UEBackgroundDownload " __VA_ARGS__)
+// Always report via NSLog, useful for debugging.
+static constexpr bool bReportToNSLog = false;
+
+DEFINE_LOG_CATEGORY_STATIC(LogIOSBackgroundDownload, Log, All);
+
+static void LogIOSBackgroundDownloadMessage(NSString* Message)
+{
+	if (bReportToNSLog)
+	{
+		NSLog(@"LogIOSBackgroundDownload %@", Message);
+	}
+	else if (UE_LOG_ACTIVE(LogIOSBackgroundDownload, Log))
+	{
+		const FString MessageWrapper(Message);
+		UE_LOG(LogIOSBackgroundDownload, Log, TEXT("%s"), *MessageWrapper);
+	}
+} 
+
+#define UE_DNLD_LOG(...) LogIOSBackgroundDownloadMessage([NSString stringWithFormat:__VA_ARGS__])
 
 // --------------------------------------------------------------------------------------------------------------------
 
