Mac: Fix for ensure that fires when moving assets to folders

diff --git a/Engine/Source/Runtime/ApplicationCore/Private/Mac/MacApplication.cpp b/Engine/Source/Runtime/ApplicationCore/Private/Mac/MacApplication.cpp
--- a/Engine/Source/Runtime/ApplicationCore/Private/Mac/MacApplication.cpp
+++ b/Engine/Source/Runtime/ApplicationCore/Private/Mac/MacApplication.cpp
@@ -1286,10 +1286,13 @@ bool FMacApplication::OnWindowDestroyed(TSharedRef<FMacWindow> DestroyedWindow)
 		FScopeLock Lock(&WindowsMutex);
 		Windows.Remove(DestroyedWindow);
 	}
-
-	if (!CocoaWindowsToClose.Contains(WindowHandle))
+	
 	{
-		CocoaWindowsToClose.Add(WindowHandle);
+		FScopeLock Lock(&WindowsToCloseMutex);
+		if (!CocoaWindowsToClose.Contains(WindowHandle))
+		{
+			CocoaWindowsToClose.Add(WindowHandle);
+		}
 	}
 
 	TSharedPtr<FMacWindow> WindowToActivate;
@@ -2157,11 +2160,14 @@ void FMacApplication::CloseQueuedWindows()
 {
 	// OnWindowClose may call PumpMessages, which would reenter this function, so make a local copy of SlateWindowsToClose array to avoid infinite recursive calls
 	TArray<TSharedRef<FMacWindow>> LocalWindowsToClose;
+	TArray<FCocoaWindow*> LocalCocoaWindowsToClose;
 
 	{
 		FScopeLock Lock(&WindowsToCloseMutex);
 		LocalWindowsToClose = SlateWindowsToClose;
+		LocalCocoaWindowsToClose = CocoaWindowsToClose;
 		SlateWindowsToClose.Empty();
+		CocoaWindowsToClose.Empty();
 	}
 
 	if (LocalWindowsToClose.Num() > 0)
@@ -2176,14 +2182,12 @@ void FMacApplication::CloseQueuedWindows()
 	{
 		MainThreadCall(^{
 			SCOPED_AUTORELEASE_POOL;
-			for (FCocoaWindow* Window : CocoaWindowsToClose)
+			for (FCocoaWindow* Window : LocalCocoaWindowsToClose)
 			{
 				[Window close];
 				[Window release];
 			}
 		}, UnrealCloseEventMode, true);
-
-		CocoaWindowsToClose.Empty();
 	}
 }
 
