Wire NotifyNotificationObjectOfComplete to didFinishDownloadingToURL to fix local push notification

diff --git a/Engine/Source/Runtime/ApplicationCore/Private/IOS/IOSBackgroundURLSessionHandler.cpp b/Engine/Source/Runtime/ApplicationCore/Private/IOS/IOSBackgroundURLSessionHandler.cpp
--- a/Engine/Source/Runtime/ApplicationCore/Private/IOS/IOSBackgroundURLSessionHandler.cpp
+++ b/Engine/Source/Runtime/ApplicationCore/Private/IOS/IOSBackgroundURLSessionHandler.cpp
@@ -285,11 +285,12 @@ NSString* const SerializationKeyRetryCountPerURL = @"r";
 	NSSet<__kindof NSString*>* _UnreachableCDNs;
 	NSTimer* _ForegroundStaleDownloadCheckTimer;
 	BackgroundHttpFileHashHelperPtr _HelperPtr;
-	int32 MaximumConnectionsPerHost;
-	int32 RetryResumeDataLimit;
-	int32 CDNReorderingTimeout;
-	double CheckForForegroundStaleDownloadsWithInterval;
-	double ForegroundStaleDownloadTimeout;
+	int32 _MaximumConnectionsPerHost;
+	int32 _RetryResumeDataLimit;
+	int32 _CDNReorderingTimeout;
+	double _CheckForForegroundStaleDownloadsWithInterval;
+	double _ForegroundStaleDownloadTimeout;
+	std::atomic<bool> _bAnyTaskDidCompleteWithError;
 }
 
 static constexpr NSUInteger InvalidDownloadId = 0;
@@ -339,43 +340,44 @@ static constexpr NSInteger HTTPStatusCodeErrorServer = 500;
 	bool bUseForegroundSession = false;
 	bool bDiscretionary = false;
 	bool bShouldSendLaunchEvents = true;
-	MaximumConnectionsPerHost = 6;
+	_MaximumConnectionsPerHost = 6;
 	double TimeoutIntervalForRequest = 120.0; // Note, ignored in background sessions (if bUseForegroundSession is false).
 	double TimeoutIntervalForResource = 60.0 * 60.0;
-	RetryResumeDataLimit = 3;
-	CDNReorderingTimeout = 400;
-	CheckForForegroundStaleDownloadsWithInterval = 1.0; // how often to check for stale downloads, <=0.0 to disable
-	ForegroundStaleDownloadTimeout = 30.0; // If download hasn't received any bytes for this duration, cancel and retry if possible
+	_RetryResumeDataLimit = 3;
+	_CDNReorderingTimeout = 400;
+	_CheckForForegroundStaleDownloadsWithInterval = 1.0; // how often to check for stale downloads, <=0.0 to disable
+	_ForegroundStaleDownloadTimeout = 30.0; // If download hasn't received any bytes for this duration, cancel and retry if possible
 
 #ifndef UE_DNLD_SANDBOX
 	GConfig->GetBool(TEXT("BackgroundHttp.iOSSettings"), TEXT("bUseForegroundSession"), bUseForegroundSession, GEngineIni);
 	GConfig->GetBool(TEXT("BackgroundHttp.iOSSettings"), TEXT("bDiscretionary"), bDiscretionary, GEngineIni);
 	GConfig->GetBool(TEXT("BackgroundHttp.iOSSettings"), TEXT("bShouldSendLaunchEvents"), bShouldSendLaunchEvents, GEngineIni);
-	GConfig->GetInt(TEXT("BackgroundHttp.iOSSettings"), TEXT("MaximumConnectionsPerHost"), MaximumConnectionsPerHost, GEngineIni);
+	GConfig->GetInt(TEXT("BackgroundHttp.iOSSettings"), TEXT("MaximumConnectionsPerHost"), _MaximumConnectionsPerHost, GEngineIni);
 	GConfig->GetDouble(TEXT("BackgroundHttp.iOSSettings"), TEXT("BackgroundReceiveTimeout"), TimeoutIntervalForRequest, GEngineIni);
 	GConfig->GetDouble(TEXT("BackgroundHttp.iOSSettings"), TEXT("BackgroundHttpResourceTimeout"), TimeoutIntervalForResource, GEngineIni);
-	GConfig->GetInt(TEXT("BackgroundHttp.iOSSettings"), TEXT("RetryResumeDataLimit"), RetryResumeDataLimit, GEngineIni);
-	GConfig->GetInt(TEXT("BackgroundHttp.iOSSettings"), TEXT("CDNReorderingTimeout"), CDNReorderingTimeout, GEngineIni);
-	GConfig->GetDouble(TEXT("BackgroundHttp.iOSSettings"), TEXT("CheckForForegroundStaleDownloadsWithInterval"), CheckForForegroundStaleDownloadsWithInterval, GEngineIni);
-	GConfig->GetDouble(TEXT("BackgroundHttp.iOSSettings"), TEXT("ForegroundStaleDownloadTimeout"), ForegroundStaleDownloadTimeout, GEngineIni);
+	GConfig->GetInt(TEXT("BackgroundHttp.iOSSettings"), TEXT("RetryResumeDataLimit"), _RetryResumeDataLimit, GEngineIni);
+	GConfig->GetInt(TEXT("BackgroundHttp.iOSSettings"), TEXT("CDNReorderingTimeout"), _CDNReorderingTimeout, GEngineIni);
+	GConfig->GetDouble(TEXT("BackgroundHttp.iOSSettings"), TEXT("CheckForForegroundStaleDownloadsWithInterval"), _CheckForForegroundStaleDownloadsWithInterval, GEngineIni);
+	GConfig->GetDouble(TEXT("BackgroundHttp.iOSSettings"), TEXT("ForegroundStaleDownloadTimeout"), _ForegroundStaleDownloadTimeout, GEngineIni);
 #endif
 
 	UE_DNLD_LOG(@"bUseForegroundSession=%u", bUseForegroundSession ? 1 : 0);
 	UE_DNLD_LOG(@"bDiscretionary=%u", bDiscretionary ? 1 : 0);
 	UE_DNLD_LOG(@"bShouldSendLaunchEvents=%u", bShouldSendLaunchEvents ? 1 : 0);
-	UE_DNLD_LOG(@"MaximumConnectionsPerHost=%i", MaximumConnectionsPerHost);
+	UE_DNLD_LOG(@"MaximumConnectionsPerHost=%i", _MaximumConnectionsPerHost);
 	UE_DNLD_LOG(@"TimeoutIntervalForRequest=%f", TimeoutIntervalForRequest);
 	UE_DNLD_LOG(@"TimeoutIntervalForResource=%f", TimeoutIntervalForResource);
-	UE_DNLD_LOG(@"RetryResumeDataLimit=%i", RetryResumeDataLimit);
-	UE_DNLD_LOG(@"CDNReorderingTimeout=%i", CDNReorderingTimeout);
-	UE_DNLD_LOG(@"CheckForForegroundStaleDownloadsWithInterval=%f", CheckForForegroundStaleDownloadsWithInterval);
-	UE_DNLD_LOG(@"ForegroundStaleDownloadTimeout=%f", ForegroundStaleDownloadTimeout);
+	UE_DNLD_LOG(@"RetryResumeDataLimit=%i", _RetryResumeDataLimit);
+	UE_DNLD_LOG(@"CDNReorderingTimeout=%i", _CDNReorderingTimeout);
+	UE_DNLD_LOG(@"CheckForForegroundStaleDownloadsWithInterval=%f", _CheckForForegroundStaleDownloadsWithInterval);
+	UE_DNLD_LOG(@"ForegroundStaleDownloadTimeout=%f", _ForegroundStaleDownloadTimeout);
 
 	_AllDownloads = [NSMutableDictionary new];
 	_AllDownloadsFuture = _AllDownloadsPromise.get_future();
 	_NextDownloadId = InvalidDownloadId + 1;
 	_UnreachableCDNs = nil;
 	_ForegroundStaleDownloadCheckTimer = nil;
+	_bAnyTaskDidCompleteWithError = false;
 
 	// Never allow cellular unless we get explicit opt-in from the user.
 	self.AllowCellular = NO;
@@ -401,7 +403,7 @@ static constexpr NSInteger HTTPStatusCodeErrorServer = 500;
 	// TODO Is this any use for us? Needs entitlement.
 	//Configuration.multipathServiceType = NSURLSessionMultipathServiceTypeAggregate;
 
-	Configuration.HTTPMaximumConnectionsPerHost = MaximumConnectionsPerHost;
+	Configuration.HTTPMaximumConnectionsPerHost = _MaximumConnectionsPerHost;
 
 	Configuration.timeoutIntervalForRequest = TimeoutIntervalForRequest;
 
@@ -529,7 +531,7 @@ static constexpr NSInteger HTTPStatusCodeErrorServer = 500;
 
 - (NSMutableArray<__kindof NSURL*>*)ReorderCDNsByReachability:(NSMutableArray<__kindof NSURL*>*)URLs
 {
-	if (CDNReorderingTimeout == 0)
+	if (_CDNReorderingTimeout == 0)
 	{
 		return URLs;
 	}
@@ -540,7 +542,7 @@ static constexpr NSInteger HTTPStatusCodeErrorServer = 500;
 		{
 			UE_DNLD_LOG(@"Starting to check for CDN reachability");
 
-			const double CDNReorderingTimeoutInSeconds = (double)CDNReorderingTimeout / 1000;
+			const double CDNReorderingTimeoutInSeconds = (double)_CDNReorderingTimeout / 1000;
 
 			// Creating a temporary foreground NSURLSession for pinging CDN's
 			NSURLSessionConfiguration* Configuration = [NSURLSessionConfiguration defaultSessionConfiguration];
@@ -551,7 +553,7 @@ static constexpr NSInteger HTTPStatusCodeErrorServer = 500;
 			Configuration.networkServiceType = NSURLNetworkServiceTypeResponsiveData;
 			Configuration.timeoutIntervalForRequest = CDNReorderingTimeoutInSeconds;
 			Configuration.timeoutIntervalForResource = CDNReorderingTimeoutInSeconds;
-			Configuration.HTTPMaximumConnectionsPerHost = MaximumConnectionsPerHost;
+			Configuration.HTTPMaximumConnectionsPerHost = _MaximumConnectionsPerHost;
 
 			NSURLSession* Session = [NSURLSession sessionWithConfiguration:Configuration];
 			NSMutableSet<__kindof NSString*>* ReachableHosts = [NSMutableSet set];
@@ -597,7 +599,7 @@ static constexpr NSInteger HTTPStatusCodeErrorServer = 500;
 				[Task resume];
 			}
 
-			PendingTasksFinished->get_future().wait_for(std::chrono::milliseconds(CDNReorderingTimeout));
+			PendingTasksFinished->get_future().wait_for(std::chrono::milliseconds(_CDNReorderingTimeout));
 			UE_DNLD_LOG(@"Finished waiting for CDN reachability");
 
 			[Session invalidateAndCancel];
@@ -704,7 +706,7 @@ static constexpr NSInteger HTTPStatusCodeErrorServer = 500;
 	URLs = [self ReorderCDNsByReachability:URLs];
 
 	// Serialize current settings
-	FBackgroundNSURLSessionDownloadTaskData* TaskData = [FBackgroundNSURLSessionDownloadTaskData TaskDataWithURLs:URLs WithRetryCount:RetryResumeDataLimit];
+	FBackgroundNSURLSessionDownloadTaskData* TaskData = [FBackgroundNSURLSessionDownloadTaskData TaskDataWithURLs:URLs WithRetryCount:_RetryResumeDataLimit];
 
 	// Check for existing download task, could be from previous app session.
 	UE_DNLD_LOG(@"Trying to find existing download for asset path '%@'", AssetPath);
@@ -854,7 +856,7 @@ static constexpr NSInteger HTTPStatusCodeErrorServer = 500;
 
 	if (ResetRetryCount)
 	{
-		[NewTaskData ResetRetryCount:RetryResumeDataLimit];
+		[NewTaskData ResetRetryCount:_RetryResumeDataLimit];
 	}
 
 	// Start a new task
@@ -917,16 +919,16 @@ static constexpr NSInteger HTTPStatusCodeErrorServer = 500;
 
 - (void)StartCheckingForStaleDownloads
 {
-	if (_ForegroundStaleDownloadCheckTimer == nil && CheckForForegroundStaleDownloadsWithInterval > 0.0)
+	if (_ForegroundStaleDownloadCheckTimer == nil && _CheckForForegroundStaleDownloadsWithInterval > 0.0)
 	{
 		_ForegroundStaleDownloadCheckTimer = [[NSTimer
-											   scheduledTimerWithTimeInterval:CheckForForegroundStaleDownloadsWithInterval
+											   scheduledTimerWithTimeInterval:_CheckForForegroundStaleDownloadsWithInterval
 											   target:self
 											   selector:@selector(CheckForStaleDownloads)
 											   userInfo:nil
 											   repeats:YES]
 											  retain];
-		_ForegroundStaleDownloadCheckTimer.tolerance = CheckForForegroundStaleDownloadsWithInterval * 0.5;
+		_ForegroundStaleDownloadCheckTimer.tolerance = _CheckForForegroundStaleDownloadsWithInterval * 0.5;
 
 		UE_DNLD_LOG(@"Start checking for stale downloads");
 	}
@@ -977,7 +979,7 @@ static constexpr NSInteger HTTPStatusCodeErrorServer = 500;
 		}
 
 		const NSTimeInterval TimeSinceLastUpdate = CurrentTime - LastUpdateTimeNumber.doubleValue;
-		if (TimeSinceLastUpdate >= ForegroundStaleDownloadTimeout)
+		if (TimeSinceLastUpdate >= _ForegroundStaleDownloadTimeout)
 		{
 			UE_DNLD_LOG(@"Task '%@' with taskIdentifier %lu is considering stale, retrying", Task.taskDescription, Task.taskIdentifier);
 
@@ -1076,6 +1078,9 @@ static constexpr NSInteger HTTPStatusCodeErrorServer = 500;
 	// Otherwise we need to put them inside FBackgroundNSURLSessionDownloadTaskData.
 	[Task.progress setUserInfoObject:[NSNumber numberWithInteger:HTTPCode] forKey:NSProgressDownloadResultStatusCode];
 	[Task.progress setUserInfoObject:TempFile forKey:NSProgressDownloadResultTempFilePath];
+
+	const bool bDownloadSuccess = TempFile != nil;
+	FBackgroundURLSessionHandler::OnDownloadCompleted.Broadcast(DownloadId, bDownloadSuccess);
 }
 
 - (NSString* _Nullable)GetDownloadResult:(NSUInteger)DownloadId OutStatus:(BOOL* _Nonnull)OutStatus OutStatusCode:(NSInteger*)OutStatusCode
@@ -1122,6 +1127,10 @@ static constexpr NSInteger HTTPStatusCodeErrorServer = 500;
 	[[NSOperationQueue mainQueue] addOperationWithBlock:^
 	{
 		UE_DNLD_LOG(@"URLSessionDidFinishEventsForBackgroundURLSession calling completion handler.");
+
+		FBackgroundURLSessionHandler::OnDownloadsCompletedWhileAppWasNotRunning.Broadcast(!_bAnyTaskDidCompleteWithError);
+		_bAnyTaskDidCompleteWithError = false;
+
 		CompletionHandler();
 		[CompletionHandler release];
 	}];
@@ -1134,6 +1143,10 @@ static constexpr NSInteger HTTPStatusCodeErrorServer = 500;
 		return;
 	}
 
+	// Set it even if we will retry the download, as the only use of this variable is to report it URLSessionDidFinishEventsForBackgroundURLSession,
+	// in that context retrying any download means that all downloads hasn't been completed yet.
+	_bAnyTaskDidCompleteWithError = true;
+
 	NSURLSessionDownloadTask* Task = (NSURLSessionDownloadTask*)GenericTask;
 	NSString* LocalizedDescription = Error != nil ? Error.localizedDescription : @"nil";
 
@@ -1266,6 +1279,10 @@ static constexpr NSInteger HTTPStatusCodeErrorServer = 500;
 
 const uint64 FBackgroundURLSessionHandler::InvalidDownloadId = [FBackgroundNSURLSession GetInvalidDownloadId];
 
+FBackgroundURLSessionHandler::FOnDownloadCompleted FBackgroundURLSessionHandler::OnDownloadCompleted;
+
+FBackgroundURLSessionHandler::FOnDownloadsCompletedWhileAppWasNotRunning FBackgroundURLSessionHandler::OnDownloadsCompletedWhileAppWasNotRunning;
+
 void FBackgroundURLSessionHandler::AllowCellular(bool bAllow)
 {
 	@autoreleasepool
@@ -1291,7 +1308,7 @@ uint64 FBackgroundURLSessionHandler::CreateOrFindDownload(const TArray<FString>&
 		{
 			[URLArray addObject:URL.GetNSString()];
 		}
-		
+
 		[[FBackgroundNSURLSession Shared] SetFileHashHelper:HelperRef];
 		return [[FBackgroundNSURLSession Shared] CreateOrFindDownloadForURLs:URLArray WithPriority:Priority];
 	}
diff --git a/Engine/Source/Runtime/ApplicationCore/Public/IOS/IOSBackgroundURLSessionHandler.h b/Engine/Source/Runtime/ApplicationCore/Public/IOS/IOSBackgroundURLSessionHandler.h
--- a/Engine/Source/Runtime/ApplicationCore/Public/IOS/IOSBackgroundURLSessionHandler.h
+++ b/Engine/Source/Runtime/ApplicationCore/Public/IOS/IOSBackgroundURLSessionHandler.h
@@ -52,10 +52,20 @@ public:
 	UE_DEPRECATED(5.5, "Use new API in FBackgroundURLSessionHandler.") static void OnDelayedBackgroundURLSessionCompleteHandlerCalled();
 
 	// New API for background downloading
-	
-	// Value of invalid download id which could be used to compare return value of CreateOrFindDownload 
+
+	// Value of invalid download id which could be used to compare return value of CreateOrFindDownload. 
 	static const uint64 InvalidDownloadId;
 
+	// Will be invoked from didFinishDownloadingToURL or didCompleteWithError.
+	DECLARE_MULTICAST_DELEGATE_TwoParams(FOnDownloadCompleted, const uint64 /*DownloadId*/, const bool /*bSuccess*/);
+	static FOnDownloadCompleted OnDownloadCompleted;
+
+	// Will be invoked from handleEventsForBackgroundURLSession application delegate. Needs to be registered very early, e.g. from static constructor.
+	// handleEventsForBackgroundURLSession is only invoked if app was killed by OS while in background and then relaunched to notify that downloads were completed.
+	// Is not invoked in any other scenario.
+	DECLARE_MULTICAST_DELEGATE_OneParam(FOnDownloadsCompletedWhileAppWasNotRunning, const bool /*bSuccess*/);
+	static FOnDownloadsCompletedWhileAppWasNotRunning OnDownloadsCompletedWhileAppWasNotRunning;
+
 	// Sets if cellular is allowed to be used for new downloads.
 	// Existing downloads will be recreated to reflect new setting value.
 	static void AllowCellular(bool bAllow);
diff --git a/Engine/Source/Runtime/Online/BackgroundHTTP/Private/BackgroundHttpNotificationObject.cpp b/Engine/Source/Runtime/Online/BackgroundHTTP/Private/BackgroundHttpNotificationObject.cpp
--- a/Engine/Source/Runtime/Online/BackgroundHTTP/Private/BackgroundHttpNotificationObject.cpp
+++ b/Engine/Source/Runtime/Online/BackgroundHTTP/Private/BackgroundHttpNotificationObject.cpp
@@ -82,8 +82,8 @@ FBackgroundHttpNotificationObject::~FBackgroundHttpNotificationObject()
 	{
 		if (!bNotifyOnlyOnFullSuccess || (NumFailedDownloads == 0))
 		{
-			//Setting the datetime to 0 forcing the local notifcation to be sent as soon as possible
-			FDateTime TargetTime = FDateTime();
+			// Schedule notification slightly in the future to pass checks preventing scheduling the past.
+			FDateTime TargetTime = FDateTime::Now() + FTimespan::FromSeconds(1);
 			if (nullptr != PlatformNotificationService)
 			{
 				PlatformNotificationService->ScheduleLocalNotificationAtTimeOverrideId(TargetTime, true, NotificationTitle, NotificationBody, NotificationAction, NotificationActivationString, IdOverride);
diff --git a/Engine/Source/Runtime/Online/BackgroundHTTP/Private/IOS/ApplePlatformBackgroundHttpManager.cpp b/Engine/Source/Runtime/Online/BackgroundHTTP/Private/IOS/ApplePlatformBackgroundHttpManager.cpp
--- a/Engine/Source/Runtime/Online/BackgroundHTTP/Private/IOS/ApplePlatformBackgroundHttpManager.cpp
+++ b/Engine/Source/Runtime/Online/BackgroundHTTP/Private/IOS/ApplePlatformBackgroundHttpManager.cpp
@@ -18,6 +18,19 @@
 
 #include "PlatformBackgroundHttp.h"
 
+FApplePlatformBackgroundHttpManager::FApplePlatformBackgroundHttpManager()
+	: FBackgroundHttpManagerImpl()
+	, PendingRemoveRequests()
+	, PendingRemoveRequestLock()
+{
+	OnDownloadCompletedHandle = FBackgroundURLSessionHandler::OnDownloadCompleted.AddRaw(this, &FApplePlatformBackgroundHttpManager::OnDownloadCompleted);
+}
+
+FApplePlatformBackgroundHttpManager::~FApplePlatformBackgroundHttpManager()
+{
+	FBackgroundURLSessionHandler::OnDownloadCompleted.Remove(OnDownloadCompletedHandle);
+}
+
 void FApplePlatformBackgroundHttpManager::AddRequest(const FBackgroundHttpRequestPtr GenericRequest)
 {
 	FAppleBackgroundHttpRequestPtr Request = StaticCastSharedPtr<FApplePlatformBackgroundHttpRequest>(GenericRequest);
@@ -102,3 +115,23 @@ bool FApplePlatformBackgroundHttpManager::Tick(float DeltaTime)
 
 	return true;
 }
+
+void FApplePlatformBackgroundHttpManager::OnDownloadCompleted(const uint64 DownloadId, const bool bSuccess)
+{
+	// Assuming this function will be called rarily (once a few seconds) and overall amount of downloads is limited (<500),
+	// it's then faster to just iterate over array to find the corresponding request.
+
+	FRWScopeLock ScopeLock(ActiveRequestLock, SLT_ReadOnly);
+	for (FBackgroundHttpRequestPtr& GenericRequest : ActiveRequests)
+	{
+		FAppleBackgroundHttpRequestPtr Request = StaticCastSharedPtr<FApplePlatformBackgroundHttpRequest>(GenericRequest);
+		if (ensureAlwaysMsgf(Request.IsValid(), TEXT("Invalid Request Pointer in ActiveRequests list!")))
+		{
+			if (Request->GetInternalDownloadId() == DownloadId)
+			{
+				Request->NotifyNotificationObjectOfComplete(bSuccess);
+				return;
+			}
+		}
+	}
+}
\ No newline at end of file
diff --git a/Engine/Source/Runtime/Online/BackgroundHTTP/Public/IOS/ApplePlatformBackgroundHttpManager.h b/Engine/Source/Runtime/Online/BackgroundHTTP/Public/IOS/ApplePlatformBackgroundHttpManager.h
--- a/Engine/Source/Runtime/Online/BackgroundHTTP/Public/IOS/ApplePlatformBackgroundHttpManager.h
+++ b/Engine/Source/Runtime/Online/BackgroundHTTP/Public/IOS/ApplePlatformBackgroundHttpManager.h
@@ -19,8 +19,8 @@ class BACKGROUNDHTTP_API FApplePlatformBackgroundHttpManager
 	: public FBackgroundHttpManagerImpl
 {
 public:
-	FApplePlatformBackgroundHttpManager() = default;
-	virtual ~FApplePlatformBackgroundHttpManager() = default;
+	FApplePlatformBackgroundHttpManager();
+	virtual ~FApplePlatformBackgroundHttpManager();
 
 	virtual void AddRequest(const FBackgroundHttpRequestPtr Request) override;
 	virtual void RemoveRequest(const FBackgroundHttpRequestPtr Request) override;
@@ -34,4 +34,6 @@ public:
 private:
 	TArray<FBackgroundHttpRequestPtr> PendingRemoveRequests;
 	FRWLock PendingRemoveRequestLock;
+	FDelegateHandle OnDownloadCompletedHandle;
+	void OnDownloadCompleted(const uint64 DownloadId, const bool bSuccess);
 };
