ApplicationCore Windows: Replace magic number for default screen DPI with official USER_DEFAULT_SCREEN_DPI define

diff --git a/Engine/Source/Runtime/ApplicationCore/Private/Windows/WindowsApplication.cpp b/Engine/Source/Runtime/ApplicationCore/Private/Windows/WindowsApplication.cpp
--- a/Engine/Source/Runtime/ApplicationCore/Private/Windows/WindowsApplication.cpp
+++ b/Engine/Source/Runtime/ApplicationCore/Private/Windows/WindowsApplication.cpp
@@ -2588,7 +2588,7 @@ int32 FWindowsApplication::ProcessDeferredMessage( const FDeferredWindowsMessage
 
 					if (!CurrentNativeEventWindowPtr->IsManualManageDPIChanges())
 					{
-						CurrentNativeEventWindowPtr->SetDPIScaleFactor(LOWORD(wParam) / 96.0f);
+						CurrentNativeEventWindowPtr->SetDPIScaleFactor(LOWORD(wParam) / static_cast<float>(USER_DEFAULT_SCREEN_DPI));
 
 						LPRECT NewRect = (LPRECT)lParam;
 						SetWindowPos(hwnd, nullptr, NewRect->left, NewRect->top, NewRect->right - NewRect->left, NewRect->bottom - NewRect->top, SWP_NOZORDER | SWP_NOACTIVATE);
diff --git a/Engine/Source/Runtime/ApplicationCore/Private/Windows/WindowsPlatformApplicationMisc.cpp b/Engine/Source/Runtime/ApplicationCore/Private/Windows/WindowsPlatformApplicationMisc.cpp
--- a/Engine/Source/Runtime/ApplicationCore/Private/Windows/WindowsPlatformApplicationMisc.cpp
+++ b/Engine/Source/Runtime/ApplicationCore/Private/Windows/WindowsPlatformApplicationMisc.cpp
@@ -282,7 +282,7 @@ bool FWindowsPlatformApplicationMisc::GetWindowTitleMatchingText(const TCHAR* Ti
 
 int32 FWindowsPlatformApplicationMisc::GetMonitorDPI(const FMonitorInfo& MonitorInfo)
 {
-	int32 DisplayDPI = 96;
+	int32 DisplayDPI = USER_DEFAULT_SCREEN_DPI;
 
 	if (IsHighDPIAwarenessEnabled())
 	{
@@ -360,7 +360,7 @@ float FWindowsPlatformApplicationMisc::GetDPIScaleFactorAtPoint(float X, float Y
 				uint32 DPIY = 0;
 				if (SUCCEEDED(GetDpiForMonitor(Monitor, 0 /*MDT_EFFECTIVE_DPI*/, &DPIX, &DPIY)))
 				{
-					Scale = (float)DPIX / 96.0f;
+					Scale = (float)DPIX / USER_DEFAULT_SCREEN_DPI;
 				}
 			}
 		}
@@ -368,7 +368,7 @@ float FWindowsPlatformApplicationMisc::GetDPIScaleFactorAtPoint(float X, float Y
 		{
 			HDC Context = GetDC(nullptr);
 			int32 DPI = GetDeviceCaps(Context, LOGPIXELSX);
-			Scale = (float)DPI / 96.0f;
+			Scale = (float)DPI / USER_DEFAULT_SCREEN_DPI;
 			ReleaseDC(nullptr, Context);
 		}
 	}
diff --git a/Engine/Source/Runtime/Core/Private/Windows/WindowsPlatformMisc.cpp b/Engine/Source/Runtime/Core/Private/Windows/WindowsPlatformMisc.cpp
--- a/Engine/Source/Runtime/Core/Private/Windows/WindowsPlatformMisc.cpp
+++ b/Engine/Source/Runtime/Core/Private/Windows/WindowsPlatformMisc.cpp
@@ -1454,7 +1454,7 @@ private:
 		LPGetDpiForWindow GetDpiForWindow = (LPGetDpiForWindow)(void*)GetProcAddress(User32Module, "GetDpiForWindow");
 		if (GetDpiForWindow == nullptr) { return 1.0f; }
 
-		return static_cast<float>(GetDpiForWindow(HandleWnd)) / 96.0f;
+		return static_cast<float>(GetDpiForWindow(HandleWnd)) / USER_DEFAULT_SCREEN_DPI;
 	}
 
 	static void SetWindowStyleFlags(HWND HandleWnd, LONG Flags, bool bEnabled)
@@ -1802,7 +1802,7 @@ private:
 			return true;
 			case WM_DPICHANGED:
 				/*{
-					float NewScale = LOWORD(WParam) / 96.0f;
+					float NewScale = LOWORD(WParam) / static_cast<float>(USER_DEFAULT_SCREEN_DPI);
 					RECT* SuggestedRect = (RECT*)LParam;
 					SetWindowPos(HandleWnd,
 						HWND_NOTOPMOST,
