Order CDN's by ping time on iOS

diff --git a/Engine/Config/IOS/BaseIOSEngine.ini b/Engine/Config/IOS/BaseIOSEngine.ini
--- a/Engine/Config/IOS/BaseIOSEngine.ini
+++ b/Engine/Config/IOS/BaseIOSEngine.ini
@@ -56,6 +56,8 @@ RetryResumeDataLimit=3
 ; How many milliseconds to wait for HEAD replies on first CDN request to decide if CDN is reachable or not (0 = disabled).
 ; Unreachable CDNs will be deprioritized in all download requests.
 CDNReorderingTimeout=400
+; Reoder by ping time CDNs that responded with-in CDNReorderingTimeout. No-op if CDNReorderingTimeout is disabled.
+bCDNReorderByPingTime=true
 ; When in foreground every N seconds check if downloads are not receiving bytes (regardless of bUseForegroundSession), retry them if so.
 CheckForForegroundStaleDownloadsWithInterval=1
 ; How long to wait in seconds before considering retrying download in CheckForForegroundStaleDownloadsWithInterval
diff --git a/Engine/Source/Runtime/ApplicationCore/Private/IOS/IOSBackgroundURLSessionHandler.cpp b/Engine/Source/Runtime/ApplicationCore/Private/IOS/IOSBackgroundURLSessionHandler.cpp
--- a/Engine/Source/Runtime/ApplicationCore/Private/IOS/IOSBackgroundURLSessionHandler.cpp
+++ b/Engine/Source/Runtime/ApplicationCore/Private/IOS/IOSBackgroundURLSessionHandler.cpp
@@ -214,6 +214,22 @@ NSString* const SerializationKeyRetryCountPerURL = @"r";
 
 // --------------------------------------------------------------------------------------------------------------------
 
+// NSURLSession CDN info
+@interface FBackgroundNSURLCDNInfo : NSObject
+
+@property (nonatomic, retain) NSString* URL;
+@property (nonatomic) BOOL IsReachable;
+@property (nonatomic) NSTimeInterval ResponseTime;
+@property (nonatomic) NSUInteger ProvidedOrder;
+
+@end
+
+@implementation FBackgroundNSURLCDNInfo
+
+@end
+
+// --------------------------------------------------------------------------------------------------------------------
+
 // NSURLSession wrapper focused on background downloading and CDN failover
 @interface FBackgroundNSURLSession : NSObject<NSURLSessionDelegate, NSURLSessionTaskDelegate, NSURLSessionDownloadDelegate>
 
@@ -282,12 +298,13 @@ NSString* const SerializationKeyRetryCountPerURL = @"r";
 	NSUInteger _NextDownloadId;
 	std::promise<void> _AllDownloadsPromise;
 	std::future<void> _AllDownloadsFuture;
-	NSSet<__kindof NSString*>* _UnreachableCDNs;
+	NSArray<__kindof FBackgroundNSURLCDNInfo*>* _CDNInfo;
 	NSTimer* _ForegroundStaleDownloadCheckTimer;
 	BackgroundHttpFileHashHelperPtr _HelperPtr;
 	int32 _MaximumConnectionsPerHost;
 	int32 _RetryResumeDataLimit;
 	int32 _CDNReorderingTimeout;
+	bool _bCDNReorderByPingTime;
 	double _CheckForForegroundStaleDownloadsWithInterval;
 	double _ForegroundStaleDownloadTimeout;
 	std::atomic<bool> _bAnyTaskDidCompleteWithError;
@@ -345,6 +362,7 @@ static constexpr NSInteger HTTPStatusCodeErrorServer = 500;
 	double TimeoutIntervalForResource = 60.0 * 60.0;
 	_RetryResumeDataLimit = 3;
 	_CDNReorderingTimeout = 400;
+	_bCDNReorderByPingTime = true;
 	_CheckForForegroundStaleDownloadsWithInterval = 1.0; // how often to check for stale downloads, <=0.0 to disable
 	_ForegroundStaleDownloadTimeout = 30.0; // If download hasn't received any bytes for this duration, cancel and retry if possible
 
@@ -357,6 +375,7 @@ static constexpr NSInteger HTTPStatusCodeErrorServer = 500;
 	GConfig->GetDouble(TEXT("BackgroundHttp.iOSSettings"), TEXT("BackgroundHttpResourceTimeout"), TimeoutIntervalForResource, GEngineIni);
 	GConfig->GetInt(TEXT("BackgroundHttp.iOSSettings"), TEXT("RetryResumeDataLimit"), _RetryResumeDataLimit, GEngineIni);
 	GConfig->GetInt(TEXT("BackgroundHttp.iOSSettings"), TEXT("CDNReorderingTimeout"), _CDNReorderingTimeout, GEngineIni);
+	GConfig->GetBool(TEXT("BackgroundHttp.iOSSettings"), TEXT("bCDNReorderByPingTime"), _bCDNReorderByPingTime, GEngineIni);
 	GConfig->GetDouble(TEXT("BackgroundHttp.iOSSettings"), TEXT("CheckForForegroundStaleDownloadsWithInterval"), _CheckForForegroundStaleDownloadsWithInterval, GEngineIni);
 	GConfig->GetDouble(TEXT("BackgroundHttp.iOSSettings"), TEXT("ForegroundStaleDownloadTimeout"), _ForegroundStaleDownloadTimeout, GEngineIni);
 #endif
@@ -369,13 +388,14 @@ static constexpr NSInteger HTTPStatusCodeErrorServer = 500;
 	UE_DNLD_LOG(@"TimeoutIntervalForResource=%f", TimeoutIntervalForResource);
 	UE_DNLD_LOG(@"RetryResumeDataLimit=%i", _RetryResumeDataLimit);
 	UE_DNLD_LOG(@"CDNReorderingTimeout=%i", _CDNReorderingTimeout);
+	UE_DNLD_LOG(@"CDNReorderingTimeout=%u", _bCDNReorderByPingTime ? 1 : 0);
 	UE_DNLD_LOG(@"CheckForForegroundStaleDownloadsWithInterval=%f", _CheckForForegroundStaleDownloadsWithInterval);
 	UE_DNLD_LOG(@"ForegroundStaleDownloadTimeout=%f", _ForegroundStaleDownloadTimeout);
 
 	_AllDownloads = [NSMutableDictionary new];
 	_AllDownloadsFuture = _AllDownloadsPromise.get_future();
 	_NextDownloadId = InvalidDownloadId + 1;
-	_UnreachableCDNs = nil;
+	_CDNInfo = nil;
 	_ForegroundStaleDownloadCheckTimer = nil;
 	_bAnyTaskDidCompleteWithError = false;
 
@@ -472,10 +492,10 @@ static constexpr NSInteger HTTPStatusCodeErrorServer = 500;
 	[_AllDownloads release];
 	_AllDownloads = nil;
 
-	if (_UnreachableCDNs != nil)
+	if (_CDNInfo != nil)
 	{
-		[_UnreachableCDNs release];
-		_UnreachableCDNs = nil;
+		[_CDNInfo release];
+		_CDNInfo = nil;
 	}
 
 	if (_ForegroundStaleDownloadCheckTimer != nil)
@@ -536,9 +556,9 @@ static constexpr NSInteger HTTPStatusCodeErrorServer = 500;
 		return URLs;
 	}
 
-	@synchronized (_UnreachableCDNs)
+	@synchronized (_CDNInfo)
 	{
-		if (_UnreachableCDNs == nil)
+		if (_CDNInfo == nil)
 		{
 			UE_DNLD_LOG(@"Starting to check for CDN reachability");
 
@@ -556,13 +576,15 @@ static constexpr NSInteger HTTPStatusCodeErrorServer = 500;
 			Configuration.HTTPMaximumConnectionsPerHost = _MaximumConnectionsPerHost;
 
 			NSURLSession* Session = [NSURLSession sessionWithConfiguration:Configuration];
-			NSMutableSet<__kindof NSString*>* ReachableHosts = [NSMutableSet set];
+			NSMutableArray<__kindof FBackgroundNSURLCDNInfo*>* CDNInfo = [NSMutableArray array];
 
 			std::shared_ptr<std::atomic<int32>> PendingTasks = std::make_shared<std::atomic<int32>>();
 			std::shared_ptr<std::promise<void>> PendingTasksFinished = std::make_shared<std::promise<void>>();
 
 			PendingTasks->fetch_add((int32)URLs.count);
 
+			NSDate* StartTime = [NSDate date];
+
 			for (NSURL* URL in URLs)
 			{
 				NSMutableURLRequest* Request = [NSMutableURLRequest requestWithURL:URL];
@@ -577,11 +599,17 @@ static constexpr NSInteger HTTPStatusCodeErrorServer = 500;
 					if (Response != nil && [Response isKindOfClass:[NSHTTPURLResponse class]])
 					{
 						NSHTTPURLResponse* HTTPResponse = (NSHTTPURLResponse*)Response;
-						UE_DNLD_LOG(@"Finished data task for '%@' (host '%@') with status code %li", Request.URL.absoluteString, Request.URL.host, HTTPResponse.statusCode);
+						const NSTimeInterval ResponseTime = -[StartTime timeIntervalSinceNow];
+						UE_DNLD_LOG(@"Finished data task for '%@' (host '%@') with status code %li and response time %f", Request.URL.absoluteString, Request.URL.host, HTTPResponse.statusCode, ResponseTime);
 
 						if (HTTPResponse.statusCode < HTTPStatusCodeErrorBadRequest)
 						{
-							[ReachableHosts addObject:Request.URL.host];
+							FBackgroundNSURLCDNInfo* Info = [[FBackgroundNSURLCDNInfo alloc] init];
+							[Info setURL:Request.URL.host];
+							[Info setResponseTime:ResponseTime];
+							[Info setIsReachable:YES];
+							[CDNInfo addObject:Info];
+							[Info release];
 						}
 					}
 					else
@@ -604,46 +632,107 @@ static constexpr NSInteger HTTPStatusCodeErrorServer = 500;
 
 			[Session invalidateAndCancel];
 
-			// We only store unreachable CDNs instead of reachable ones,
-			// to cover a case where a new unknown CDN comes in we would assume it's reachable rather than unreachable.
-			NSMutableSet<__kindof NSString*>* UnreachableCDNs = [NSMutableSet set];
-
-			for (NSURL* URL in URLs)
+			for (NSUInteger URLIndex = 0; URLIndex < [URLs count]; URLIndex++)
 			{
-				if ([ReachableHosts containsObject:URL.host])
+				NSURL* URL = [URLs objectAtIndex:URLIndex];
+				bool bFoundCDNInfo = false;
+
+				for (FBackgroundNSURLCDNInfo* Info in CDNInfo)
 				{
-					UE_DNLD_LOG(@"CDN '%@' is reachable", URL.absoluteString);
+					if ([Info.URL isEqualToString:URL.host])
+					{
+						[Info setProvidedOrder:URLIndex];
+						bFoundCDNInfo = true;
+						break;
+					}
 				}
-				else
+
+				if (!bFoundCDNInfo)
 				{
-					UE_DNLD_LOG(@"CDN '%@' is not reachable", URL.absoluteString);
-					[UnreachableCDNs addObject:URL.host];
+					FBackgroundNSURLCDNInfo* Info = [[FBackgroundNSURLCDNInfo alloc] init];
+					[Info setURL:URL.host];
+					[Info setIsReachable:NO];
+					[Info setProvidedOrder:URLIndex];
+					[CDNInfo addObject:Info];
+					[Info release];
 				}
 			}
+			
+			[CDNInfo sortUsingComparator:^NSComparisonResult(FBackgroundNSURLCDNInfo* _Nonnull A, FBackgroundNSURLCDNInfo* _Nonnull B)
+			{
+				if (A.IsReachable && B.IsReachable && _bCDNReorderByPingTime) // order by smaller ResponseTime if enabled
+				{
+					if (A.ResponseTime < B.ResponseTime)
+					{
+						return NSOrderedAscending;
+					}
+					else if (A.ResponseTime > B.ResponseTime)
+					{
+						return NSOrderedDescending;
+					}
+					else
+					{
+						return NSOrderedSame;
+					}
+				}
+				else if (A.IsReachable && !B.IsReachable)
+				{
+					// reachable CDNs should go first
+					return NSOrderedAscending;
+				}
+				else if (!A.IsReachable && B.IsReachable)
+				{
+					// reachable CDNs should go first
+					return NSOrderedDescending;
+				}
+				else if (A.ProvidedOrder < B.ProvidedOrder) // order by order in which CDN's were given
+				{
+					return NSOrderedAscending;
+				}
+				else if (A.ProvidedOrder > B.ProvidedOrder)
+				{
+					return NSOrderedDescending;
+				}
+
+				return NSOrderedSame;
+			}];
 
-			_UnreachableCDNs = [[NSSet setWithSet:UnreachableCDNs] retain];
+			for (NSUInteger i = 0; i < [CDNInfo count]; i++)
+			{
+				FBackgroundNSURLCDNInfo* Info = [CDNInfo objectAtIndex:i];
+				UE_DNLD_LOG(@"%lu CDN '%@' ResponseTime:%f IsReachable:%u ProvidedOrder:%lu", i, Info.URL, Info.ResponseTime, Info.IsReachable, (unsigned long)Info.ProvidedOrder);;
+			}
+			
+			_CDNInfo = [CDNInfo retain];
 		}
 	}
 
-	if (_UnreachableCDNs != nil && _UnreachableCDNs.count > 0)
+	if (_CDNInfo != nil && [_CDNInfo count] > 0 && [URLs count] > 0)
 	{
-		NSMutableArray<__kindof NSURL*>* Result = [NSMutableArray array];
-		NSMutableArray<__kindof NSURL*>* UnreachableURLs = [NSMutableArray array];
+		// Array sizes are assumed small enough that hashmap is not needed
+		NSMutableArray<__kindof NSURL*>* Result = [NSMutableArray arrayWithCapacity:[URLs count]];
 
-		for (NSURL* URL in URLs)
+		for (FBackgroundNSURLCDNInfo* Info in _CDNInfo)
 		{
-			// don't change the order if we don't know if host was unreachable
-			if (![_UnreachableCDNs containsObject:URL.host])
+			for (NSURL* URL in URLs)
 			{
-				[Result addObject:URL];
+				if ([Info.URL isEqualToString:URL.host])
+				{
+					[Result addObject:URL];
+					break;
+				}
 			}
-			else // otherwise put all hosts that were unreachable in the end without changing the relative order between them
+		}
+
+		// Add CDN's that weren't present at first lookup
+		for (NSURL* URL in URLs)
+		{
+			if (![Result containsObject:URL])
 			{
-				[UnreachableURLs addObject:URL];
+				[Result addObject:URL];
 			}
 		}
 
-		[Result addObjectsFromArray:UnreachableURLs];
 		return Result;
 	}
 	else
