Fix to CL 34911042 that is using the wrong array for CloseQueuedWindows

diff --git a/Engine/Source/Runtime/ApplicationCore/Private/Mac/MacApplication.cpp b/Engine/Source/Runtime/ApplicationCore/Private/Mac/MacApplication.cpp
--- a/Engine/Source/Runtime/ApplicationCore/Private/Mac/MacApplication.cpp
+++ b/Engine/Source/Runtime/ApplicationCore/Private/Mac/MacApplication.cpp
@@ -2178,7 +2178,7 @@ void FMacApplication::CloseQueuedWindows()
 		}
 	}
 
-	if (CocoaWindowsToClose.Num() > 0)
+	if (LocalCocoaWindowsToClose.Num() > 0)
 	{
 		MainThreadCall(^{
 			SCOPED_AUTORELEASE_POOL;
