Quality adjustments for iOS_High and iOS_Epic devices.
Added MemBuckets for 6GB and 8GB devices
Fixed iOS backbuffer size computations to match scene textures size computations
Enable grass for iOS_High and Epic devices
Bump texture pool size for iOS_High and Epic devices

diff --git a/Engine/Config/IOS/BaseIOSEngine.ini b/Engine/Config/IOS/BaseIOSEngine.ini
--- a/Engine/Config/IOS/BaseIOSEngine.ini
+++ b/Engine/Config/IOS/BaseIOSEngine.ini
@@ -97,6 +97,9 @@ PlatformRequiresDataCrypto=True
 PakSigningRequired=False
 
 [PlatformMemoryBuckets]
+;iOS specific: uses available memory on app start instead of total device memory
+LargestMemoryBucket_MinGB=6
+LargerMemoryBucket_MinGB=4
 DefaultMemoryBucket_MinGB=4
 SmallerMemoryBucket_MinGB=3
 SmallestMemoryBucket_MinGB=2
diff --git a/Engine/Source/Runtime/ApplicationCore/Private/IOS/IOSWindow.cpp b/Engine/Source/Runtime/ApplicationCore/Private/IOS/IOSWindow.cpp
--- a/Engine/Source/Runtime/ApplicationCore/Private/IOS/IOSWindow.cpp
+++ b/Engine/Source/Runtime/ApplicationCore/Private/IOS/IOSWindow.cpp
@@ -133,6 +133,14 @@ FPlatformRect FIOSWindow::GetScreenRect()
 		ScreenRect.Left = FMath::TruncToInt(Frame.origin.x * Scale);
 		ScreenRect.Bottom = FMath::TruncToInt((Frame.origin.y + Frame.size.height) * Scale);
 		ScreenRect.Right = FMath::TruncToInt((Frame.origin.x + Frame.size.width) * Scale);
+		
+		// see QuantizeSceneBufferSize
+		const uint32 DividableBy = 4;
+		const uint32 Mask = ~(DividableBy - 1);
+		ScreenRect.Top = (ScreenRect.Top + DividableBy - 1) & Mask;
+		ScreenRect.Left = (ScreenRect.Left + DividableBy - 1) & Mask;
+		ScreenRect.Bottom = (ScreenRect.Bottom + DividableBy - 1) & Mask;
+		ScreenRect.Right = (ScreenRect.Right + DividableBy - 1) & Mask;
 	}
 #if PLATFORM_VISIONOS
 	else if (AppDelegate.SwiftLayer != nullptr)
diff --git a/Engine/Source/Runtime/Core/Private/Apple/ApplePlatformMemory.cpp b/Engine/Source/Runtime/Core/Private/Apple/ApplePlatformMemory.cpp
--- a/Engine/Source/Runtime/Core/Private/Apple/ApplePlatformMemory.cpp
+++ b/Engine/Source/Runtime/Core/Private/Apple/ApplePlatformMemory.cpp
@@ -471,6 +471,7 @@ const FPlatformMemoryConstants& FApplePlatformMemory::GetConstants()
 		// actual physical memory. To work around this, we add 1Gb - 1b so it will be truncated 
 		// correctly and will not affect macOS
 		MemoryConstants.TotalPhysicalGB = ([NSProcessInfo processInfo].physicalMemory + (1024*1024*1024 - 1)) / 1024 / 1024 / 1024;
+		MemoryConstants.AddressLimit = FPlatformMath::RoundUpToPowerOfTwo64(MemoryConstants.TotalPhysical);
 
 		// Calculate total and available Virtual Memory
 		mach_port_t HostPort = mach_host_self();
