Enable ApplicationReceivedScreenOrientationChangedNotificationDelegate in engine pre-init on iOS

diff --git a/Engine/Source/Runtime/ApplicationCore/Private/IOS/IOSAppDelegate.cpp b/Engine/Source/Runtime/ApplicationCore/Private/IOS/IOSAppDelegate.cpp
--- a/Engine/Source/Runtime/ApplicationCore/Private/IOS/IOSAppDelegate.cpp
+++ b/Engine/Source/Runtime/ApplicationCore/Private/IOS/IOSAppDelegate.cpp
@@ -1347,6 +1347,19 @@ static FAutoConsoleVariableRef CVarGEnableThermalsReport(
 			FCoreDelegates::OnSafeFrameChangedEvent.Broadcast();
 		}, TStatId(), NULL, ENamedThreads::GameThread);
 	}
+	else
+	{
+		[FIOSAsyncTask CreateTaskWithBlock : ^ bool(void)
+		{
+			FIOSApplication* Application = [IOSAppDelegate GetDelegate].IOSApplication;
+			Application->OrientationChanged(Orientation);
+			FCoreDelegates::ApplicationReceivedScreenOrientationChangedNotificationDelegate.Broadcast((int32)[IOSAppDelegate ConvertFromUIInterfaceOrientation:Orientation]);
+
+			//we also want to fire off the safe frame event
+			FCoreDelegates::OnSafeFrameChangedEvent.Broadcast();
+			return true;
+		}];
+	}
 #endif
 }
 
