Disable reordering by CDN ping time

diff --git a/Engine/Config/IOS/BaseIOSEngine.ini b/Engine/Config/IOS/BaseIOSEngine.ini
--- a/Engine/Config/IOS/BaseIOSEngine.ini
+++ b/Engine/Config/IOS/BaseIOSEngine.ini
@@ -57,7 +57,7 @@ RetryResumeDataLimit=3
 ; Unreachable CDNs will be deprioritized in all download requests.
 CDNReorderingTimeout=400
 ; Reoder by ping time CDNs that responded with-in CDNReorderingTimeout. No-op if CDNReorderingTimeout is disabled.
-bCDNReorderByPingTime=true
+bCDNReorderByPingTime=false
 ; When in foreground every N seconds check if downloads are not receiving bytes (regardless of bUseForegroundSession), retry them if so.
 CheckForForegroundStaleDownloadsWithInterval=1
 ; How long to wait in seconds before considering retrying download in CheckForForegroundStaleDownloadsWithInterval
diff --git a/Engine/Source/Runtime/ApplicationCore/Private/IOS/IOSBackgroundURLSessionHandler.cpp b/Engine/Source/Runtime/ApplicationCore/Private/IOS/IOSBackgroundURLSessionHandler.cpp
--- a/Engine/Source/Runtime/ApplicationCore/Private/IOS/IOSBackgroundURLSessionHandler.cpp
+++ b/Engine/Source/Runtime/ApplicationCore/Private/IOS/IOSBackgroundURLSessionHandler.cpp
@@ -362,7 +362,7 @@ static constexpr NSInteger HTTPStatusCodeErrorServer = 500;
 	double TimeoutIntervalForResource = 60.0 * 60.0;
 	_RetryResumeDataLimit = 3;
 	_CDNReorderingTimeout = 400;
-	_bCDNReorderByPingTime = true;
+	_bCDNReorderByPingTime = false;
 	_CheckForForegroundStaleDownloadsWithInterval = 1.0; // how often to check for stale downloads, <=0.0 to disable
 	_ForegroundStaleDownloadTimeout = 30.0; // If download hasn't received any bytes for this duration, cancel and retry if possible
 
