iOS: Add some task metrics for the downloads

Output looks like:

diff --git a/Engine/Source/Runtime/ApplicationCore/Private/IOS/IOSBackgroundURLSessionHandler.cpp b/Engine/Source/Runtime/ApplicationCore/Private/IOS/IOSBackgroundURLSessionHandler.cpp
--- a/Engine/Source/Runtime/ApplicationCore/Private/IOS/IOSBackgroundURLSessionHandler.cpp
+++ b/Engine/Source/Runtime/ApplicationCore/Private/IOS/IOSBackgroundURLSessionHandler.cpp
@@ -903,6 +903,72 @@ static constexpr NSInteger HTTPStatusCodeErrorServer = 500;
 	}
 }
 
+#if !UE_BUILD_SHIPPING
+- (void)URLSession:(NSURLSession *)session task:(NSURLSessionTask *)task didFinishCollectingMetrics:(NSURLSessionTaskMetrics *)metrics {
+	NSDateFormatter *formatter = [[NSDateFormatter alloc] init];
+    [formatter setDateFormat:@"yyyy-MM-dd'T'HH:mm:ss.SSSZZZZZ"];
+    [formatter setLocale:[NSLocale localeWithLocaleIdentifier:@"en_US_POSIX"]];
+    
+    // Calculate task interval
+    NSURLSessionTaskTransactionMetrics *firstMetric = metrics.transactionMetrics.firstObject;
+    NSURLSessionTaskTransactionMetrics *lastMetric = metrics.transactionMetrics.lastObject;
+	UE_DNLD_LOG(@"Task %lu got metrics", task.taskIdentifier);
+    if (firstMetric && lastMetric) {
+        NSTimeInterval taskInterval = [lastMetric.requestEndDate timeIntervalSinceDate:firstMetric.requestStartDate];
+        UE_DNLD_LOG(@"Task Interval: %f", taskInterval);
+    }
+
+    // Get redirect count
+    NSInteger redirectCount = metrics.transactionMetrics.count - 1; // The first metric is the initial request, redirects are additional metrics
+    UE_DNLD_LOG(@"Redirect Count: %ld", (long)redirectCount);
+
+
+	for (NSURLSessionTaskTransactionMetrics *metric in metrics.transactionMetrics) {
+		UE_DNLD_LOG(@"Network Protocol Name: %@", metric.networkProtocolName);
+		UE_DNLD_LOG(@"Reused Connection: %@", metric.reusedConnection ? @"Yes" : @"No");
+		UE_DNLD_LOG(@"Proxy Connection: %@", metric.proxyConnection ? @"Yes": @"No");
+
+        UE_DNLD_LOG(@"Fetch Start Date: %@", [formatter stringFromDate:metric.fetchStartDate]);
+        UE_DNLD_LOG(@"Request Start Date: %@", [formatter stringFromDate:metric.requestStartDate]);
+        UE_DNLD_LOG(@"Response Start Date: %@", [formatter stringFromDate:metric.responseStartDate]);
+        UE_DNLD_LOG(@"Request End Date: %@", [formatter stringFromDate:metric.requestEndDate]);
+        UE_DNLD_LOG(@"Response End Date: %@", [formatter stringFromDate:metric.responseEndDate]);
+
+		// Calculate and log response duration
+        if (metric.responseStartDate && metric.responseEndDate) {
+            NSTimeInterval responseDuration = [metric.responseEndDate timeIntervalSinceDate:metric.responseStartDate];
+            UE_DNLD_LOG(@"Response Duration: %f seconds", responseDuration);
+
+			// Calculate and log download speed
+            if (responseDuration > 0 && metric.countOfResponseBodyBytesReceived > 0) {
+				// bytes per second
+                double downloadSpeed = (double)metric.countOfResponseBodyBytesReceived / responseDuration;
+                NSString *formattedSpeed = [self formattedSpeed:downloadSpeed];
+                UE_DNLD_LOG(@"Download Speed: %@", formattedSpeed);
+            } else {
+                UE_DNLD_LOG(@"Download Speed: Not Available");
+            }
+        } else {
+            UE_DNLD_LOG(@"Response Duration: Not Available");
+        }
+    }
+	UE_DNLD_LOG(@"-------------------------");
+}
+
+- (NSString *)formattedSpeed:(double)speedInBytesPerSecond {
+    NSArray *units = @[@"bytes/second", @"KB/s", @"MB/s", @"GB/s"];
+    double speed = speedInBytesPerSecond;
+    NSInteger unitIndex = 0;
+
+    while (speed >= 1024 && unitIndex < units.count - 1) {
+        speed /= 1024;
+        unitIndex++;
+    }
+    
+    return [NSString stringWithFormat:@"%.2f %@", speed, units[unitIndex]];
+}
+#endif
+
 - (void)SetCurrentDownloadedBytes:(uint64)DownloadedBytes ForTask:(NSURLSessionDownloadTask*)Task
 {
 	if (Task != nil)
