Fixes to iOS background downloader code

diff --git a/Engine/Source/Runtime/ApplicationCore/Private/IOS/IOSBackgroundURLSessionHandler.cpp b/Engine/Source/Runtime/ApplicationCore/Private/IOS/IOSBackgroundURLSessionHandler.cpp
--- a/Engine/Source/Runtime/ApplicationCore/Private/IOS/IOSBackgroundURLSessionHandler.cpp
+++ b/Engine/Source/Runtime/ApplicationCore/Private/IOS/IOSBackgroundURLSessionHandler.cpp
@@ -179,7 +179,7 @@ NSString* const SerializationKeyRetryCountPerURL = @"r";
 {
 	for (NSUInteger i = 0; i < self.RetryCountPerURL.count; ++i)
 	{
-		NSUInteger RetryValue = [self.RetryCountPerURL objectAtIndex:i].integerValue;
+		NSInteger RetryValue = [self.RetryCountPerURL objectAtIndex:i].integerValue;
 		if (RetryValue == 0)
 		{
 			continue;
@@ -995,6 +995,8 @@ static constexpr NSInteger HTTPStatusCodeErrorServer = 500;
 
 - (void)RecreateDownload:(NSUInteger)DownloadId ShouldResetRetryCount:(bool)ResetRetryCount
 {
+	UE_DNLD_LOG(@"RecreateDownload for DownloadId %lu", DownloadId);
+
 	NSURLSessionDownloadTask* OldTask = [self FindDownloadTaskFor:DownloadId];
 	const float OldTaskPriority = OldTask.priority;
 	const NSURLSessionTaskState OldTaskState = OldTask.state;
@@ -1293,8 +1295,15 @@ static constexpr NSInteger HTTPStatusCodeErrorServer = 500;
 
 - (void)URLSession:(NSURLSession*)Session task:(NSURLSessionTask*)GenericTask didCompleteWithError:(NSError*)Error
 {
-	if (GenericTask.state == NSURLSessionTaskStateCompleted || ![GenericTask isKindOfClass:[NSURLSessionDownloadTask class]])
+	if (Error == nil)
+	{
+		UE_DNLD_LOG(@"didCompleteWithError, task '%@' with taskIdentifier %lu is completed", GenericTask.taskDescription, GenericTask.taskIdentifier);
+		return;
+	}
+
+	if (![GenericTask isKindOfClass:[NSURLSessionDownloadTask class]])
 	{
+		UE_DNLD_LOG(@"didCompleteWithError, ignoring task '%@' with taskIdentifier %lu", GenericTask.taskDescription, GenericTask.taskIdentifier);
 		return;
 	}
 
@@ -1303,7 +1312,7 @@ static constexpr NSInteger HTTPStatusCodeErrorServer = 500;
 	_bAnyTaskDidCompleteWithError = true;
 
 	NSURLSessionDownloadTask* Task = (NSURLSessionDownloadTask*)GenericTask;
-	NSString* LocalizedDescription = Error != nil ? Error.localizedDescription : @"nil";
+	NSString* LocalizedDescription = Error.localizedDescription;
 
 	const NSUInteger DownloadId = [self FindDownloadIdForTask:Task];
 	const bool bIsTrackedTask = DownloadId != InvalidDownloadId;
