Fixing typo in CheckForStaleDownloads selector

diff --git a/Engine/Source/Runtime/ApplicationCore/Private/IOS/IOSBackgroundURLSessionHandler.cpp b/Engine/Source/Runtime/ApplicationCore/Private/IOS/IOSBackgroundURLSessionHandler.cpp
--- a/Engine/Source/Runtime/ApplicationCore/Private/IOS/IOSBackgroundURLSessionHandler.cpp
+++ b/Engine/Source/Runtime/ApplicationCore/Private/IOS/IOSBackgroundURLSessionHandler.cpp
@@ -263,7 +263,7 @@ NSString* const SerializationKeyRetryCountPerURL = @"r";
 
 - (void)StartCheckingForStaleDownloads;
 - (void)StopCheckingForStaleDownloads;
-- (void)CheckForStaleDownloads;
+- (void)CheckForStaleDownloads:(NSTimer*)Timer;
 
 - (NSURLSessionDownloadTask*)FindDownloadTaskFor:(NSUInteger)DownloadId;
 - (NSUInteger)FindDownloadIdForTask:(NSURLSessionDownloadTask*)Task;
@@ -1078,16 +1078,23 @@ static constexpr NSInteger HTTPStatusCodeErrorServer = 500;
 {
 	if (_ForegroundStaleDownloadCheckTimer == nil && _CheckForForegroundStaleDownloadsWithInterval > 0.0)
 	{
-		_ForegroundStaleDownloadCheckTimer = [[NSTimer
-											   scheduledTimerWithTimeInterval:_CheckForForegroundStaleDownloadsWithInterval
-											   target:self
-											   selector:@selector(CheckForStaleDownloads)
-											   userInfo:nil
-											   repeats:YES]
-											  retain];
-		_ForegroundStaleDownloadCheckTimer.tolerance = _CheckForForegroundStaleDownloadsWithInterval * 0.5;
-
-		UE_DNLD_LOG(@"Start checking for stale downloads");
+		dispatch_async(dispatch_get_main_queue(), ^
+		{
+			if (_ForegroundStaleDownloadCheckTimer == nil && _CheckForForegroundStaleDownloadsWithInterval > 0.0)
+			{
+				_ForegroundStaleDownloadCheckTimer = [[NSTimer
+													   scheduledTimerWithTimeInterval:_CheckForForegroundStaleDownloadsWithInterval
+													   target:self
+													   selector:@selector(CheckForStaleDownloads:)
+													   userInfo:nil
+													   repeats:YES]
+													  retain];
+				_ForegroundStaleDownloadCheckTimer.tolerance = _CheckForForegroundStaleDownloadsWithInterval * 0.5;
+				[_ForegroundStaleDownloadCheckTimer fire];
+				
+				UE_DNLD_LOG(@"Start checking for stale downloads");
+			}
+		});
 	}
 }
 
@@ -1095,15 +1102,21 @@ static constexpr NSInteger HTTPStatusCodeErrorServer = 500;
 {
 	if (_ForegroundStaleDownloadCheckTimer != nil)
 	{
-		[_ForegroundStaleDownloadCheckTimer invalidate];
-		[_ForegroundStaleDownloadCheckTimer release];
-		_ForegroundStaleDownloadCheckTimer = nil;
-
-		UE_DNLD_LOG(@"Stop checking for stale downloads");
+		dispatch_async(dispatch_get_main_queue(), ^
+		{
+			if (_ForegroundStaleDownloadCheckTimer != nil)
+			{
+				[_ForegroundStaleDownloadCheckTimer invalidate];
+				[_ForegroundStaleDownloadCheckTimer release];
+				_ForegroundStaleDownloadCheckTimer = nil;
+				
+				UE_DNLD_LOG(@"Stop checking for stale downloads");
+			}
+		});
 	}
 }
 
-- (void)CheckForStaleDownloads
+- (void)CheckForStaleDownloads:(NSTimer*)Timer
 {
 	// Copy keys to avoid deadlocking in case if cancel/resume/etc will call delegates in-place
 	NSArray<__kindof NSNumber*>* AllKeys = nil;
@@ -1118,6 +1131,10 @@ static constexpr NSInteger HTTPStatusCodeErrorServer = 500;
 	{
 		const NSUInteger DownloadId = IterKey.unsignedIntegerValue;
 		NSURLSessionDownloadTask* Task = [_AllDownloads objectForKey:IterKey];
+		if (Task == nil)
+		{
+			continue;
+		}
 
 		NSNumber* ResultStatusCode = [Task.progress.userInfo objectForKey:NSProgressDownloadResultStatusCode];
 		if (ResultStatusCode != nil)
