Save URL mapping manifest when invoked in BG

diff --git a/Engine/Source/Runtime/ApplicationCore/Private/IOS/IOSBackgroundURLSessionHandler.cpp b/Engine/Source/Runtime/ApplicationCore/Private/IOS/IOSBackgroundURLSessionHandler.cpp
--- a/Engine/Source/Runtime/ApplicationCore/Private/IOS/IOSBackgroundURLSessionHandler.cpp
+++ b/Engine/Source/Runtime/ApplicationCore/Private/IOS/IOSBackgroundURLSessionHandler.cpp
@@ -1305,6 +1305,8 @@ static constexpr NSInteger HTTPStatusCodeErrorServer = 500;
 		FBackgroundURLSessionHandler::OnDownloadsCompletedWhileAppWasNotRunning.Broadcast(!_bAnyTaskDidCompleteWithError);
 		_bAnyTaskDidCompleteWithError = false;
 
+		[self SaveFileHashHelperState];
+
 		CompletionHandler();
 		[CompletionHandler release];
 	}];
@@ -1418,6 +1420,8 @@ static constexpr NSInteger HTTPStatusCodeErrorServer = 500;
 
 	NSString* DestinationPath = [self GetTempPathForURL:Task.originalRequest.URL];
 
+	[self SaveFileHashHelperState];
+
 	// Try to remove existing file in case if we have a stale file.
 	if ([[NSFileManager defaultManager] fileExistsAtPath:DestinationPath])
 	{
