Don't check for stale downloads while in background

diff --git a/Engine/Source/Runtime/ApplicationCore/Private/IOS/IOSAppDelegate.cpp b/Engine/Source/Runtime/ApplicationCore/Private/IOS/IOSAppDelegate.cpp
--- a/Engine/Source/Runtime/ApplicationCore/Private/IOS/IOSAppDelegate.cpp
+++ b/Engine/Source/Runtime/ApplicationCore/Private/IOS/IOSAppDelegate.cpp
@@ -1511,6 +1511,7 @@ FCriticalSection RenderSuspend;
 
     self.bAudioActive = false;
     FAppEntry::Suspend(true);
+	FBackgroundURLSessionHandler::HandleDidEnterBackground();
 
 	FEmbeddedCommunication::KeepAwake(TEXT("Background"), false);
 
@@ -1528,6 +1529,7 @@ FCriticalSection RenderSuspend;
 
 - (void)applicationWillEnterForeground:(UIApplication *)application
 {
+	FBackgroundURLSessionHandler::HandleWillEnterForeground();
 	FEmbeddedCommunication::KeepAwake(TEXT("Background"), false);
 	/*
 	 Called as part of the transition from the background to the inactive state; here you can undo many of the changes made on entering the background.
diff --git a/Engine/Source/Runtime/ApplicationCore/Private/IOS/IOSBackgroundURLSessionHandler.cpp b/Engine/Source/Runtime/ApplicationCore/Private/IOS/IOSBackgroundURLSessionHandler.cpp
--- a/Engine/Source/Runtime/ApplicationCore/Private/IOS/IOSBackgroundURLSessionHandler.cpp
+++ b/Engine/Source/Runtime/ApplicationCore/Private/IOS/IOSBackgroundURLSessionHandler.cpp
@@ -265,6 +265,9 @@ NSString* const SerializationKeyRetryCountPerURL = @"r";
 - (void)StopCheckingForStaleDownloads;
 - (void)CheckForStaleDownloads:(NSTimer*)Timer;
 
+- (void)HandleDidEnterBackground;
+- (void)HandleWillEnterForeground;
+
 - (NSURLSessionDownloadTask*)FindDownloadTaskFor:(NSUInteger)DownloadId;
 - (NSUInteger)FindDownloadIdForTask:(NSURLSessionDownloadTask*)Task;
 - (NSUInteger)EnsureTaskIsTracked:(NSURLSessionDownloadTask*)Task;
@@ -1118,6 +1121,12 @@ static constexpr NSInteger HTTPStatusCodeErrorServer = 500;
 
 - (void)CheckForStaleDownloads:(NSTimer*)Timer
 {
+	// Only check for stale downloads in foreground
+	if ([UIApplication sharedApplication].applicationState != UIApplicationStateActive)
+	{
+		return;
+	}
+
 	// Copy keys to avoid deadlocking in case if cancel/resume/etc will call delegates in-place
 	NSArray<__kindof NSNumber*>* AllKeys = nil;
 	@synchronized(_AllDownloads)
@@ -1166,6 +1175,36 @@ static constexpr NSInteger HTTPStatusCodeErrorServer = 500;
 	}
 }
 
+- (void)HandleDidEnterBackground
+{
+}
+
+- (void)HandleWillEnterForeground
+{
+	@synchronized(_AllDownloads)
+	{
+		[_AllDownloads enumerateKeysAndObjectsUsingBlock:^(NSNumber* Key, NSURLSessionDownloadTask* Task, BOOL* IterStop)
+		{
+			NSNumber* ResultStatusCode = [Task.progress.userInfo objectForKey:NSProgressDownloadResultStatusCode];
+
+			// If task was not complete
+			if (ResultStatusCode == nil)
+			{
+				NSNumber* LastUpdateTimeNumber = [Task.progress.userInfo objectForKey:NSProgressDownloadLastUpdateTime];
+
+				// But has last update time
+				if (LastUpdateTimeNumber != nil)
+				{
+					// Refresh task update time so stale timer doesn't retry task for first N seconds after app goes to foreground
+					[Task.progress setUserInfoObject:[NSNumber numberWithDouble:[[NSDate date] timeIntervalSince1970]] forKey:NSProgressDownloadLastUpdateTime];
+
+					UE_DNLD_LOG(@"Refreshing last update time for task '%@' taskIdentifier %lu", Task.taskDescription, Task.taskIdentifier);
+				}
+			}
+		}];
+	}
+}
+
 - (NSURLSessionDownloadTask*)FindDownloadTaskFor:(NSUInteger)DownloadId
 {
 	@synchronized(_AllDownloads)
@@ -1583,6 +1622,22 @@ void FBackgroundURLSessionHandler::HandleEventsForBackgroundURLSession(const FSt
 	}
 }
 
+void FBackgroundURLSessionHandler::HandleDidEnterBackground()
+{
+	@autoreleasepool
+	{
+		[[FBackgroundNSURLSession Shared] HandleDidEnterBackground];
+	}
+}
+
+void FBackgroundURLSessionHandler::HandleWillEnterForeground()
+{
+	@autoreleasepool
+	{
+		[[FBackgroundNSURLSession Shared] HandleWillEnterForeground];
+	}
+}
+
 void FBackgroundURLSessionHandler::SaveBackgroundHttpFileHashHelperState()
 {
 	@autoreleasepool
diff --git a/Engine/Source/Runtime/ApplicationCore/Public/IOS/IOSBackgroundURLSessionHandler.h b/Engine/Source/Runtime/ApplicationCore/Public/IOS/IOSBackgroundURLSessionHandler.h
--- a/Engine/Source/Runtime/ApplicationCore/Public/IOS/IOSBackgroundURLSessionHandler.h
+++ b/Engine/Source/Runtime/ApplicationCore/Public/IOS/IOSBackgroundURLSessionHandler.h
@@ -94,6 +94,12 @@ public:
 	// To be used by app delegate, call it from handleEventsForBackgroundURLSession.
 	static void HandleEventsForBackgroundURLSession(const FString& SessionIdentifier);
 
+	// To be used by app delegate, call it from applicationDidEnterBackground
+	static void HandleDidEnterBackground();
+
+	// To be used by app delegate, call it from applicationWillEnterForeground
+	static void HandleWillEnterForeground();
+
 	// To be used by ApplePlatformBackgroundHttpManager
 	static void SaveBackgroundHttpFileHashHelperState();
 
