Down prio errored out CDN's over timed out CDN's

diff --git a/Engine/Source/Runtime/ApplicationCore/Private/IOS/IOSBackgroundURLSessionHandler.cpp b/Engine/Source/Runtime/ApplicationCore/Private/IOS/IOSBackgroundURLSessionHandler.cpp
--- a/Engine/Source/Runtime/ApplicationCore/Private/IOS/IOSBackgroundURLSessionHandler.cpp
+++ b/Engine/Source/Runtime/ApplicationCore/Private/IOS/IOSBackgroundURLSessionHandler.cpp
@@ -214,18 +214,50 @@ NSString* const SerializationKeyRetryCountPerURL = @"r";
 
 // --------------------------------------------------------------------------------------------------------------------
 
+enum class EBackgroundNSURLCDNInfoResponse : uint32 // Beware these constants define sorting order in SortingKeyWith
+{
+	// CDN responded with a valid HTTP response with a code smaller than HTTPStatusCodeErrorBadRequest (400)
+	Ok = 1,
+	
+	// CDN request timed outed or was cancelled
+	Timeout = 2,
+	
+	// CDN or networking responded with error, e.g. DNS resolution error, etc
+	Error = 3
+};
+
 // NSURLSession CDN info
 @interface FBackgroundNSURLCDNInfo : NSObject
 
 @property (nonatomic, retain) NSString* URL;
-@property (nonatomic) BOOL IsReachable;
+@property (nonatomic) EBackgroundNSURLCDNInfoResponse Response;
 @property (nonatomic) NSTimeInterval ResponseTime;
 @property (nonatomic) NSUInteger ProvidedOrder;
 
+- (double)SortingKeyWith:(BOOL)bSortByResponseTime;
+
 @end
 
 @implementation FBackgroundNSURLCDNInfo
 
+- (double)SortingKeyWith:(BOOL)bSortByResponseTime
+{
+	double Key = (double)self.Response * 100000.0;
+	
+	if (bSortByResponseTime && self.Response == EBackgroundNSURLCDNInfoResponse::Ok)
+	{
+		// sort by response time only if response was valid
+		Key += self.ResponseTime;
+	}
+	else
+	{
+		// all other cases should be sorted by provided order
+		Key += self.ProvidedOrder;
+	}
+	
+	return Key;
+}
+
 @end
 
 // --------------------------------------------------------------------------------------------------------------------
@@ -391,7 +423,7 @@ static constexpr NSInteger HTTPStatusCodeErrorServer = 500;
 	UE_DNLD_LOG(@"TimeoutIntervalForResource=%f", TimeoutIntervalForResource);
 	UE_DNLD_LOG(@"RetryResumeDataLimit=%i", _RetryResumeDataLimit);
 	UE_DNLD_LOG(@"CDNReorderingTimeout=%i", _CDNReorderingTimeout);
-	UE_DNLD_LOG(@"CDNReorderingTimeout=%u", _bCDNReorderByPingTime ? 1 : 0);
+	UE_DNLD_LOG(@"CDNReorderByPingTime=%u", _bCDNReorderByPingTime ? 1 : 0);
 	UE_DNLD_LOG(@"CheckForForegroundStaleDownloadsWithInterval=%f", _CheckForForegroundStaleDownloadsWithInterval);
 	UE_DNLD_LOG(@"ForegroundStaleDownloadTimeout=%f", _ForegroundStaleDownloadTimeout);
 
@@ -599,26 +631,43 @@ static constexpr NSInteger HTTPStatusCodeErrorServer = 500;
 				// Note, completion handler might be invoked after end of this method.
 				NSURLSessionDataTask* Task = [Session dataTaskWithRequest:Request completionHandler:^(NSData* _Nullable Data, NSURLResponse* _Nullable Response, NSError* _Nullable Error)
 				{
+					const NSTimeInterval ResponseTime = -[StartTime timeIntervalSinceNow];
+
+					bool bIsOk = false;
+					const bool bIsTimeout = (Error != nil) ? (Error.code == NSURLErrorTimedOut || Error.code == NSURLErrorCancelled) : false;
+
 					if (Response != nil && [Response isKindOfClass:[NSHTTPURLResponse class]])
 					{
 						NSHTTPURLResponse* HTTPResponse = (NSHTTPURLResponse*)Response;
-						const NSTimeInterval ResponseTime = -[StartTime timeIntervalSinceNow];
 						UE_DNLD_LOG(@"Finished data task for '%@' (host '%@') with status code %li and response time %f", Request.URL.absoluteString, Request.URL.host, HTTPResponse.statusCode, ResponseTime);
 
 						if (HTTPResponse.statusCode < HTTPStatusCodeErrorBadRequest)
 						{
-							FBackgroundNSURLCDNInfo* Info = [[FBackgroundNSURLCDNInfo alloc] init];
-							[Info setURL:Request.URL.host];
-							[Info setResponseTime:ResponseTime];
-							[Info setIsReachable:YES];
-							[CDNInfo addObject:Info];
-							[Info release];
+							bIsOk = true;
 						}
 					}
 					else
 					{
-						UE_DNLD_LOG(@"Finished data task for '%@' with error '%@'", Request.URL.absoluteString, (Error != nil ? Error.localizedDescription : @"nil"));
+						UE_DNLD_LOG(@"Finished data task for '%@' with error '%@' (%i, %i) and response time %f", Request.URL.absoluteString, (Error != nil ? Error.localizedDescription : @"nil"), (Error != nil ? (int32)Error.code : 0), bIsTimeout ? 1 : 0, ResponseTime);
+					}
+
+					FBackgroundNSURLCDNInfo* Info = [[FBackgroundNSURLCDNInfo alloc] init];
+					[Info setURL:Request.URL.host];
+					if (bIsOk)
+					{
+						[Info setResponse:EBackgroundNSURLCDNInfoResponse::Ok];
+					}
+					else if (bIsTimeout)
+					{
+						[Info setResponse:EBackgroundNSURLCDNInfoResponse::Timeout];
+					}
+					else
+					{
+						[Info setResponse:EBackgroundNSURLCDNInfoResponse::Error];
 					}
+					[Info setResponseTime:ResponseTime];
+					[CDNInfo addObject:Info];
+					[Info release];
 
 					if (PendingTasks->fetch_add(-1) <= 1)
 					{
@@ -654,7 +703,8 @@ static constexpr NSInteger HTTPStatusCodeErrorServer = 500;
 				{
 					FBackgroundNSURLCDNInfo* Info = [[FBackgroundNSURLCDNInfo alloc] init];
 					[Info setURL:URL.host];
-					[Info setIsReachable:NO];
+					// If cdn/networking hasn't provided us with any info, consider request timed out
+					[Info setResponse:EBackgroundNSURLCDNInfoResponse::Timeout];
 					[Info setProvidedOrder:URLIndex];
 					[CDNInfo addObject:Info];
 					[Info release];
@@ -663,47 +713,27 @@ static constexpr NSInteger HTTPStatusCodeErrorServer = 500;
 			
 			[CDNInfo sortUsingComparator:^NSComparisonResult(FBackgroundNSURLCDNInfo* _Nonnull A, FBackgroundNSURLCDNInfo* _Nonnull B)
 			{
-				if (A.IsReachable && B.IsReachable && _bCDNReorderByPingTime) // order by smaller ResponseTime if enabled
-				{
-					if (A.ResponseTime < B.ResponseTime)
-					{
-						return NSOrderedAscending;
-					}
-					else if (A.ResponseTime > B.ResponseTime)
-					{
-						return NSOrderedDescending;
-					}
-					else
-					{
-						return NSOrderedSame;
-					}
-				}
-				else if (A.IsReachable && !B.IsReachable)
+				const double KeyA = [A SortingKeyWith:_bCDNReorderByPingTime];
+				const double KeyB = [B SortingKeyWith:_bCDNReorderByPingTime];
+
+				if (KeyA < KeyB)
 				{
-					// reachable CDNs should go first
 					return NSOrderedAscending;
 				}
-				else if (!A.IsReachable && B.IsReachable)
+				else if (KeyA > KeyB)
 				{
-					// reachable CDNs should go first
 					return NSOrderedDescending;
 				}
-				else if (A.ProvidedOrder < B.ProvidedOrder) // order by order in which CDN's were given
-				{
-					return NSOrderedAscending;
-				}
-				else if (A.ProvidedOrder > B.ProvidedOrder)
+				else
 				{
-					return NSOrderedDescending;
+					return NSOrderedSame;
 				}
-
-				return NSOrderedSame;
 			}];
 
 			for (NSUInteger i = 0; i < [CDNInfo count]; i++)
 			{
 				FBackgroundNSURLCDNInfo* Info = [CDNInfo objectAtIndex:i];
-				UE_DNLD_LOG(@"%lu CDN '%@' ResponseTime:%f IsReachable:%u ProvidedOrder:%lu", i, Info.URL, Info.ResponseTime, Info.IsReachable, (unsigned long)Info.ProvidedOrder);;
+				UE_DNLD_LOG(@"%lu CDN '%@' Response:%u ResponseTime:%f ProvidedOrder:%lu SortingKey:%f", i, Info.URL, Info.Response, Info.ResponseTime, (unsigned long)Info.ProvidedOrder, [Info SortingKeyWith:_bCDNReorderByPingTime]);
 			}
 			
 			_CDNInfo = [CDNInfo retain];
