Fix out-of-bounds read in FWindowsConsoleOutputDevice2::FConsoleWindow::HitTestNCB.

diff --git a/Engine/Source/Runtime/ApplicationCore/Private/Windows/WindowsConsoleOutputDevice2.cpp b/Engine/Source/Runtime/ApplicationCore/Private/Windows/WindowsConsoleOutputDevice2.cpp
--- a/Engine/Source/Runtime/ApplicationCore/Private/Windows/WindowsConsoleOutputDevice2.cpp
+++ b/Engine/Source/Runtime/ApplicationCore/Private/Windows/WindowsConsoleOutputDevice2.cpp
@@ -1087,7 +1087,7 @@ public:
 		if (Y < WindowRect.top + 2 || Y > WindowRect.top + 27 || X > WindowRect.right - 2)
 			return HTNOWHERE;
 		int32 ButtonWidth = 27;
-		if (X < WindowRect.right - ButtonWidth*4)
+		if (X <= WindowRect.right - ButtonWidth*4)
 			return HTCAPTION;
 		int32 ButtonIndex = (WindowRect.right - X) / ButtonWidth;
 		int32 Buttons[] = { HTCLOSE, HTMAXBUTTON, HTMINBUTTON, HTEXPAND };
