Add an optional per-message scope to Windows message processing and custom message handlers.

diff --git a/Engine/Source/Runtime/ApplicationCore/Private/Windows/WindowsApplication.cpp b/Engine/Source/Runtime/ApplicationCore/Private/Windows/WindowsApplication.cpp
--- a/Engine/Source/Runtime/ApplicationCore/Private/Windows/WindowsApplication.cpp
+++ b/Engine/Source/Runtime/ApplicationCore/Private/Windows/WindowsApplication.cpp
@@ -50,6 +50,8 @@ THIRD_PARTY_INCLUDES_END
 
 DEFINE_LOG_CATEGORY(LogWindowsDesktop);
 
+DECLARE_CYCLE_STAT(TEXT("WindowsApplication Custom Message Handler"), STAT_CustomMessageHandler_ProcessMessage, STATGROUP_Quick);
+
 /**
  * Hack to get around multiple mouse events being triggered for touch events.
  * Enabling this will prevent pen tablets from working since until we switch to the windows 8 sdk (and can use WM_POINTER*) events we cannot detect the difference
@@ -82,6 +84,13 @@ static int32 ForceRawInputSimulation = false;
 static int32 EnableRawInputSimulationOverRDP = false;
 #endif
 
+static bool bEnableWindowMessageProfilerScopes = false;
+FAutoConsoleVariableRef CVarEnableWindowMessageProfilerScopes(
+	TEXT("WindowsApplication.EnableWindowMessageProfilerScopes"),
+	bEnableWindowMessageProfilerScopes,
+	TEXT("Enable profiler scopes for window message handling on a per-message basis and for custom message handlers.")
+);
+
 /* 
 * Enabling first touch event will prevent small pop on some touch input devices.
 */
@@ -1003,6 +1012,7 @@ int32 FWindowsApplication::ProcessMessage( HWND hwnd, uint32 msg, WPARAM wParam,
 		// give others a chance to handle messages
 		for (IWindowsMessageHandler* Handler : MessageHandlers)
 		{
+			CONDITIONAL_SCOPE_CYCLE_COUNTER(STAT_CustomMessageHandler_ProcessMessage, bEnableWindowMessageProfilerScopes);
 			int32 HandlerResult = 0;
 			if (Handler->ProcessMessage(hwnd, msg, wParam, lParam, HandlerResult))
 			{
@@ -1020,6 +1030,14 @@ int32 FWindowsApplication::ProcessMessage( HWND hwnd, uint32 msg, WPARAM wParam,
 			return 0;
 		}
 
+
+		FString MessageScopeName;
+		if (bEnableWindowMessageProfilerScopes)
+		{
+			MessageScopeName = FString::Printf(TEXT("Window Message %u"), msg);
+		}
+
+		TRACE_CPUPROFILER_EVENT_SCOPE_TEXT_CONDITIONAL(*MessageScopeName, !MessageScopeName.IsEmpty());
 		switch(msg)
 		{
 		case WM_INPUTLANGCHANGEREQUEST:
