Expose background download CDN order on iOS

diff --git a/Engine/Source/Runtime/ApplicationCore/Private/IOS/IOSBackgroundURLSessionHandler.cpp b/Engine/Source/Runtime/ApplicationCore/Private/IOS/IOSBackgroundURLSessionHandler.cpp
--- a/Engine/Source/Runtime/ApplicationCore/Private/IOS/IOSBackgroundURLSessionHandler.cpp
+++ b/Engine/Source/Runtime/ApplicationCore/Private/IOS/IOSBackgroundURLSessionHandler.cpp
@@ -36,7 +36,7 @@ static void LogIOSBackgroundDownloadMessage(NSString* Message)
 		const FString MessageWrapper(Message);
 		UE_LOG(LogIOSBackgroundDownload, Log, TEXT("%s"), *MessageWrapper);
 	}
-} 
+}
 
 #define UE_DNLD_LOG(...) LogIOSBackgroundDownloadMessage([NSString stringWithFormat:__VA_ARGS__])
 
@@ -218,10 +218,10 @@ enum class EBackgroundNSURLCDNInfoResponse : uint32 // Beware these constants de
 {
 	// CDN responded with a valid HTTP response with a code smaller than HTTPStatusCodeErrorBadRequest (400)
 	Ok = 1,
-	
+
 	// CDN request timed outed or was cancelled
 	Timeout = 2,
-	
+
 	// CDN or networking responded with error, e.g. DNS resolution error, etc
 	Error = 3
 };
@@ -229,21 +229,43 @@ enum class EBackgroundNSURLCDNInfoResponse : uint32 // Beware these constants de
 // NSURLSession CDN info
 @interface FBackgroundNSURLCDNInfo : NSObject
 
-@property (nonatomic, retain) NSString* URL;
+@property (nonatomic, retain) NSString* CDNHost;
+@property (nonatomic, retain) NSString* CDNAbsoluteURL;
 @property (nonatomic) EBackgroundNSURLCDNInfoResponse Response;
 @property (nonatomic) NSTimeInterval ResponseTime;
 @property (nonatomic) NSUInteger ProvidedOrder;
 
+- (void)SetFromURL:(NSURL*)URL;
 - (double)SortingKeyWith:(BOOL)bSortByResponseTime;
 
 @end
 
 @implementation FBackgroundNSURLCDNInfo
 
+- (void)SetFromURL:(NSURL*)URL
+{
+	[self setCDNHost:URL.host];
+
+	NSURLComponents* Components = [NSURLComponents componentsWithURL:URL resolvingAgainstBaseURL:NO];
+	Components.fragment = nil;
+	Components.path = nil;
+	Components.query = nil;
+
+	NSString* AbsoluteURL = Components.string;
+	if (![AbsoluteURL hasSuffix:@"/"])
+	{
+		// Trailing / is not part of RFC 3986, but is used in CDN configs,
+		// append it to keep CDN string formatting.
+		AbsoluteURL = [AbsoluteURL stringByAppendingString:@"/"];
+	}
+
+	[self setCDNAbsoluteURL:AbsoluteURL];
+}
+
 - (double)SortingKeyWith:(BOOL)bSortByResponseTime
 {
 	double Key = (double)self.Response * 100000.0;
-	
+
 	if (bSortByResponseTime && self.Response == EBackgroundNSURLCDNInfoResponse::Ok)
 	{
 		// sort by response time only if response was valid
@@ -254,7 +276,7 @@ enum class EBackgroundNSURLCDNInfoResponse : uint32 // Beware these constants de
 		// all other cases should be sorted by provided order
 		Key += self.ProvidedOrder;
 	}
-	
+
 	return Key;
 }
 
@@ -277,6 +299,7 @@ enum class EBackgroundNSURLCDNInfoResponse : uint32 // Beware these constants de
 - (void)SaveFileHashHelperState;
 - (NSString*)GetTempPathForURL:(NSURL* _Nonnull)URL;
 - (NSMutableArray<__kindof NSURL*>*)ReorderCDNsByReachability:(NSMutableArray<__kindof NSURL*>*)URLs;
+- (NSArray<__kindof FBackgroundNSURLCDNInfo*>*)GetCDNInfo;
 - (NSURLSessionDownloadTask*)CreateDownloadForURL:(NSURL* _Nonnull)URL WithPriority:(float)Priority WithTaskData:(FBackgroundNSURLSessionDownloadTaskData* _Nonnull)TaskData;
 - (NSURLSessionDownloadTask*)CreateDownloadForResumeData:(NSData* _Nonnull)ResumeData WithPriority:(float)Priority WithTaskData:(FBackgroundNSURLSessionDownloadTaskData* _Nonnull)TaskData;
 
@@ -652,7 +675,7 @@ static constexpr NSInteger HTTPStatusCodeErrorServer = 500;
 					}
 
 					FBackgroundNSURLCDNInfo* Info = [[FBackgroundNSURLCDNInfo alloc] init];
-					[Info setURL:Request.URL.host];
+					[Info SetFromURL:Request.URL];
 					if (bIsOk)
 					{
 						[Info setResponse:EBackgroundNSURLCDNInfoResponse::Ok];
@@ -691,7 +714,7 @@ static constexpr NSInteger HTTPStatusCodeErrorServer = 500;
 
 				for (FBackgroundNSURLCDNInfo* Info in CDNInfo)
 				{
-					if ([Info.URL isEqualToString:URL.host])
+					if ([Info.CDNHost isEqualToString:URL.host])
 					{
 						[Info setProvidedOrder:URLIndex];
 						bFoundCDNInfo = true;
@@ -702,7 +725,7 @@ static constexpr NSInteger HTTPStatusCodeErrorServer = 500;
 				if (!bFoundCDNInfo)
 				{
 					FBackgroundNSURLCDNInfo* Info = [[FBackgroundNSURLCDNInfo alloc] init];
-					[Info setURL:URL.host];
+					[Info SetFromURL:URL];
 					// If cdn/networking hasn't provided us with any info, consider request timed out
 					[Info setResponse:EBackgroundNSURLCDNInfoResponse::Timeout];
 					[Info setProvidedOrder:URLIndex];
@@ -733,7 +756,7 @@ static constexpr NSInteger HTTPStatusCodeErrorServer = 500;
 			for (NSUInteger i = 0; i < [CDNInfo count]; i++)
 			{
 				FBackgroundNSURLCDNInfo* Info = [CDNInfo objectAtIndex:i];
-				UE_DNLD_LOG(@"%lu CDN '%@' Response:%u ResponseTime:%f ProvidedOrder:%lu SortingKey:%f", i, Info.URL, Info.Response, Info.ResponseTime, (unsigned long)Info.ProvidedOrder, [Info SortingKeyWith:_bCDNReorderByPingTime]);
+				UE_DNLD_LOG(@"%lu CDN '%@' AbsoluteURL '%@' Response:%u ResponseTime:%f ProvidedOrder:%lu SortingKey:%f", i, Info.CDNHost, Info.CDNAbsoluteURL, Info.Response, Info.ResponseTime, (unsigned long)Info.ProvidedOrder, [Info SortingKeyWith:_bCDNReorderByPingTime]);
 			}
 			
 			_CDNInfo = [CDNInfo retain];
@@ -749,7 +772,7 @@ static constexpr NSInteger HTTPStatusCodeErrorServer = 500;
 		{
 			for (NSURL* URL in URLs)
 			{
-				if ([Info.URL isEqualToString:URL.host])
+				if ([Info.CDNHost isEqualToString:URL.host])
 				{
 					[Result addObject:URL];
 					break;
@@ -774,6 +797,14 @@ static constexpr NSInteger HTTPStatusCodeErrorServer = 500;
 	}
 }
 
+- (NSArray<__kindof FBackgroundNSURLCDNInfo*>*)GetCDNInfo
+{
+	@synchronized (_CDNInfo)
+	{
+		return _CDNInfo;
+	}
+}
+
 - (NSURLSessionDownloadTask*)CreateDownloadForURL:(NSURL* _Nonnull)URL WithPriority:(float)Priority WithTaskData:(FBackgroundNSURLSessionDownloadTaskData* _Nonnull)TaskData;
 {
 	NSMutableURLRequest* URLRequest = [NSMutableURLRequest requestWithURL:URL];
@@ -1692,6 +1723,23 @@ void FBackgroundURLSessionHandler::SaveBackgroundHttpFileHashHelperState()
 	}
 }
 
+TArray<FString> FBackgroundURLSessionHandler::GetCDNOrderArray()
+{
+	TArray<FString> Result;
+
+	@autoreleasepool
+	{
+		NSArray<__kindof FBackgroundNSURLCDNInfo*>* CDNInfo = [[FBackgroundNSURLSession Shared] GetCDNInfo];
+
+		for (FBackgroundNSURLCDNInfo* Info in CDNInfo)
+		{
+			Result.Add(FString(Info.CDNAbsoluteURL));
+		}
+	}
+
+	return Result;
+}
+
 #if !UE_BUILD_SHIPPING
 void FBackgroundURLSessionHandler::GetDownloadDebugText(const uint64 DownloadId, TArray<FString>& Output)
 {
diff --git a/Engine/Source/Runtime/ApplicationCore/Public/IOS/IOSBackgroundURLSessionHandler.h b/Engine/Source/Runtime/ApplicationCore/Public/IOS/IOSBackgroundURLSessionHandler.h
--- a/Engine/Source/Runtime/ApplicationCore/Public/IOS/IOSBackgroundURLSessionHandler.h
+++ b/Engine/Source/Runtime/ApplicationCore/Public/IOS/IOSBackgroundURLSessionHandler.h
@@ -107,6 +107,11 @@ public:
 	// To be used by ApplePlatformBackgroundHttpManager
 	static void SaveBackgroundHttpFileHashHelperState();
 
+	// Returns an ordered list of CDNs used to issue actual downloads.
+	// A list of URLS provided to CreateOrFindDownload might change order if CDNReorderingTimeout > 0 to ensure better success rate.
+	// List is empty before first CreateOrFindDownload call.
+	static TArray<FString> GetCDNOrderArray();
+
 #if !UE_BUILD_SHIPPING
 	static void GetDownloadDebugText(const uint64 DownloadId, TArray<FString>& Output);
 #endif
