Add support for DualSense and DualShock4 touchpads.
Fix for setting playerIndex to PlayerOne on controllers 2-4 when controller disconnects.

diff --git a/Engine/Source/Runtime/ApplicationCore/Private/Apple/AppleControllerInterface.cpp b/Engine/Source/Runtime/ApplicationCore/Private/Apple/AppleControllerInterface.cpp
--- a/Engine/Source/Runtime/ApplicationCore/Private/Apple/AppleControllerInterface.cpp
+++ b/Engine/Source/Runtime/ApplicationCore/Private/Apple/AppleControllerInterface.cpp
@@ -108,7 +108,16 @@ void FAppleControllerInterface::SetControllerType(uint32 ControllerIndex)
 
     if ([Controller.productCategory isEqualToString:@"DualShock 4"])
     {
-        Controllers[ControllerIndex].ControllerType = ControllerType::DualShockGamepad;
+    	// validate it's of class GCDualShockGamepad so we can cast it later
+	    if ([Controller.extendedGamepad isKindOfClass:[GCDualShockGamepad class]])
+	    {
+			Controllers[ControllerIndex].ControllerType = ControllerType::DualShockGamepad;
+	    }
+    	else
+    	{
+			UE_LOG(LogAppleController, Warning, TEXT("DualShock4 detected, but invalid class"));
+		    Controllers[ControllerIndex].ControllerType = ControllerType::ExtendedGamepad;
+    	}
     }
     else if ([Controller.productCategory isEqualToString:@"Xbox One"])
     {
@@ -116,9 +125,17 @@ void FAppleControllerInterface::SetControllerType(uint32 ControllerIndex)
     }
     else if ([Controller.productCategory isEqualToString:@"DualSense"])
     {
-        Controllers[ControllerIndex].ControllerType = ControllerType::DualSenseGamepad;
+    	// validate it is of class GCDualSenseGamepad so we can cast it later
+	    if ([Controller.extendedGamepad isKindOfClass:[GCDualSenseGamepad class]])
+	    {
+			Controllers[ControllerIndex].ControllerType = ControllerType::DualSenseGamepad;
+	    }
+    	else
+    	{
+			UE_LOG(LogAppleController, Warning, TEXT("DualSense detected, but invalid class"));
+		    Controllers[ControllerIndex].ControllerType = ControllerType::ExtendedGamepad;
+    	}
     }
-
     else if (Controller.extendedGamepad != nil)
     {
         Controllers[ControllerIndex].ControllerType = ControllerType::ExtendedGamepad;
@@ -139,6 +156,11 @@ void FAppleControllerInterface::SetCurrentController(GCController* Controller)
     int32 ControllerIndex = 0;
     PlayerIndex PreviousIndex = PlayerIndex::PlayerUnset;
 
+	if (Controller == nil)
+	{
+		return;
+	}
+	
     for (ControllerIndex = 0; ControllerIndex < UE_ARRAY_COUNT(Controllers); ControllerIndex++)
     {
         if (Controllers[ControllerIndex].Controller == Controller)
@@ -165,9 +187,7 @@ void FAppleControllerInterface::SetCurrentController(GCController* Controller)
 
 void FAppleControllerInterface::HandleConnection(GCController* Controller)
 {
-	static_assert(GCControllerPlayerIndex1 == 0 && GCControllerPlayerIndex4 == 3, "Apple changed the player index enums");
-
-	if (!bAllowControllers)
+	if (!bAllowControllers || Controller == nil)
 	{
 		return;
 	}
@@ -196,10 +216,8 @@ void FAppleControllerInterface::HandleConnection(GCController* Controller)
         DeviceMapper.Internal_MapInputDeviceToUser(DeviceId, UserId, EInputDeviceConnectionState::Connected);
         
         UE_LOG(LogAppleController, Log, TEXT("New %s controller inserted, assigned to playerIndex %d"),
-               Controllers[ControllerIndex].ControllerType == ControllerType::ExtendedGamepad ||
-               Controllers[ControllerIndex].ControllerType == ControllerType::XboxGamepad ||
-               Controllers[ControllerIndex].ControllerType == ControllerType::DualShockGamepad
-               ? TEXT("Gamepad") : TEXT("Remote"), Controllers[ControllerIndex].PlayerIndex);
+               Controllers[ControllerIndex].ControllerType == ControllerType::SiriRemote
+               ? TEXT("SiriRemote") : *FString(Controller.productCategory), Controllers[ControllerIndex].PlayerIndex);
         break;
 	}
 	checkf(bFoundSlot, TEXT("Used a fifth controller somehow!"));
@@ -207,8 +225,8 @@ void FAppleControllerInterface::HandleConnection(GCController* Controller)
 
 void FAppleControllerInterface::HandleDisconnect(GCController* Controller)
 {
-		// if we don't allow controllers, there could be unset player index here
-	if (!bAllowControllers)
+	// if we don't allow controllers, there could be unset player index here
+	if (!bAllowControllers || Controller == nil)
 	{
         return;
 	}
@@ -247,19 +265,17 @@ void FAppleControllerInterface::SendControllerEvents()
 		FUserController& Controller = Controllers[i];
 		
 		// make sure the connection handler has run on this
-		if (Controller.PlayerIndex == PlayerIndex::PlayerUnset)
+		if (Controller.PlayerIndex == PlayerIndex::PlayerUnset || Controller.Controller == nil)
 		{
             continue;
 		}
 		
-		GCController* ControllerImpl = Controller.Controller;
-		
 		IPlatformInputDeviceMapper& DeviceMapper = IPlatformInputDeviceMapper::Get();
 		FPlatformUserId UserId = FGenericPlatformMisc::GetPlatformUserForUserIndex(Controller.PlayerIndex);
 		FInputDeviceId DeviceId = DeviceMapper.GetPrimaryInputDeviceForUser(UserId);
 		
-        GCExtendedGamepad* ExtendedGamepad = [ControllerImpl capture].extendedGamepad;
-		GCMotion* Motion = ControllerImpl.motion;
+		GCController* ControllerImpl = Controller.Controller;
+		GCExtendedGamepad* ExtendedGamepad = [ControllerImpl capture].extendedGamepad;
         
 		if (ExtendedGamepad != nil)
 		{
@@ -483,9 +499,25 @@ bIsPressed = Gamepad.GCButton.pressed; \
             else if (UEButton == FGamepadKeyNames::DPadRight){SET_PRESSED(ExtendedGamepad, ExtendedPreviousGamepad, dpad.right, UEButton);}
             else if (UEButton == FGamepadKeyNames::DPadLeft){SET_PRESSED(ExtendedGamepad, ExtendedPreviousGamepad, dpad.left, UEButton);}
             else if (UEButton == FGamepadKeyNames::SpecialRight){SET_PRESSED(ExtendedGamepad, ExtendedPreviousGamepad, buttonMenu, UEButton);}
-            else if (UEButton == FGamepadKeyNames::SpecialLeft){SET_PRESSED(ExtendedGamepad, ExtendedPreviousGamepad, buttonOptions, UEButton);}
             else if (UEButton == FGamepadKeyNames::LeftThumb){SET_PRESSED(ExtendedGamepad, ExtendedPreviousGamepad, leftThumbstickButton, UEButton);}
             else if (UEButton == FGamepadKeyNames::RightThumb){SET_PRESSED(ExtendedGamepad, ExtendedPreviousGamepad, rightThumbstickButton, UEButton);}
+			else if (UEButton == FGamepadKeyNames::SpecialLeft)
+			{
+				// DualSense and DualShock4 use the touchpad button for the "SpecialLeft" button
+				if (Controllers[ControllerIndex].ControllerType == ControllerType::DualSenseGamepad)
+				{
+					// SetControllerType() verified that ExtendedGamepad is of type GCDualSenseGamepad
+					SET_PRESSED(static_cast<GCDualSenseGamepad*>(ExtendedGamepad), static_cast<GCDualSenseGamepad*>(ExtendedPreviousGamepad), touchpadButton, UEButton);
+				}
+				else if (Controllers[ControllerIndex].ControllerType == ControllerType::DualShockGamepad)
+				{
+					SET_PRESSED(static_cast<GCDualShockGamepad*>(ExtendedGamepad), static_cast<GCDualShockGamepad*>(ExtendedPreviousGamepad), touchpadButton, UEButton);
+				}
+				else
+				{
+					SET_PRESSED(ExtendedGamepad, ExtendedPreviousGamepad, buttonOptions, UEButton);
+				}
+			}
             break;
         case ControllerType::SiriRemote:
             
