Fix crash due to possible nil task data

diff --git a/Engine/Source/Runtime/ApplicationCore/Private/IOS/IOSBackgroundURLSessionHandler.cpp b/Engine/Source/Runtime/ApplicationCore/Private/IOS/IOSBackgroundURLSessionHandler.cpp
--- a/Engine/Source/Runtime/ApplicationCore/Private/IOS/IOSBackgroundURLSessionHandler.cpp
+++ b/Engine/Source/Runtime/ApplicationCore/Private/IOS/IOSBackgroundURLSessionHandler.cpp
@@ -95,10 +95,21 @@ NSString* const SerializationKeyRetryCountPerURL = @"r";
 
 + (FBackgroundNSURLSessionDownloadTaskData* _Nullable)TaskDataFromSerializedString:(NSString*)SerializedData
 {
+	if (SerializedData == nil)
+	{
+		return nil;
+	}
+
+	NSData* SerializedDataEncoded = [SerializedData dataUsingEncoding:NSUTF8StringEncoding];
+	if (SerializedDataEncoded == nil)
+	{
+		return nil;
+	}
+
 	FBackgroundNSURLSessionDownloadTaskData* Data = [[FBackgroundNSURLSessionDownloadTaskData alloc] init];
 
 	NSError* Error = nil;
-	NSDictionary* Dict = [NSJSONSerialization JSONObjectWithData:[SerializedData dataUsingEncoding:NSUTF8StringEncoding] options:0 error:&Error];
+	NSDictionary* Dict = [NSJSONSerialization JSONObjectWithData:SerializedDataEncoded options:0 error:&Error];
 
 	NSNumber* Version = [Dict valueForKey:SerializationKeyProtocolVersion];
 	NSMutableArray<__kindof NSString*>* CDNs = [Dict valueForKey:SerializationKeyCDNs];
