Fixes for compiling without exec commands in shipping (bUseExecCommandsInShipping = false). Swapped Exec() for Exec_Dev() as these console commands are meant for development purposes and not needed in shipping.

diff --git a/Engine/Plugins/Experimental/Text3D/Source/Text3D/Private/Text3DEngineSubsystem.cpp b/Engine/Plugins/Experimental/Text3D/Source/Text3D/Private/Text3DEngineSubsystem.cpp
--- a/Engine/Plugins/Experimental/Text3D/Source/Text3D/Private/Text3DEngineSubsystem.cpp
+++ b/Engine/Plugins/Experimental/Text3D/Source/Text3D/Private/Text3DEngineSubsystem.cpp
@@ -63,7 +63,7 @@ void UText3DEngineSubsystem::Deinitialize()
 	Super::Deinitialize();
 }
 
-bool UText3DEngineSubsystem::Exec(UWorld* InWorld, const TCHAR* InCmd, FOutputDevice& InAr)
+bool UText3DEngineSubsystem::Exec_Dev(UWorld* InWorld, const TCHAR* InCmd, FOutputDevice& InAr)
 {
 	if (HasAllFlags(RF_ClassDefaultObject))
 	{
diff --git a/Engine/Plugins/Experimental/Text3D/Source/Text3D/Private/Text3DEngineSubsystem.h b/Engine/Plugins/Experimental/Text3D/Source/Text3D/Private/Text3DEngineSubsystem.h
--- a/Engine/Plugins/Experimental/Text3D/Source/Text3D/Private/Text3DEngineSubsystem.h
+++ b/Engine/Plugins/Experimental/Text3D/Source/Text3D/Private/Text3DEngineSubsystem.h
@@ -193,7 +193,7 @@ public:
 	// ~End UEngineSubsystem
 
 	// ~Begin FSelfRegisteringExec
-	virtual bool Exec(class UWorld* InWorld, const TCHAR* InCmd, FOutputDevice& InAr) override;
+	virtual bool Exec_Dev(class UWorld* InWorld, const TCHAR* InCmd, FOutputDevice& InAr) override;
 	// ~End FSelfRegisteringExec
 
 	void PrintCache() const;
diff --git a/Engine/Source/Runtime/ApplicationCore/Private/Null/NullApplication.cpp b/Engine/Source/Runtime/ApplicationCore/Private/Null/NullApplication.cpp
--- a/Engine/Source/Runtime/ApplicationCore/Private/Null/NullApplication.cpp
+++ b/Engine/Source/Runtime/ApplicationCore/Private/Null/NullApplication.cpp
@@ -277,7 +277,7 @@ void FNullApplication::SetWorkArea(const FPlatformRect& NewWorkArea)
 	WorkArea = NewWorkArea;
 }
 
-bool FNullApplication::Exec(UWorld* InWorld, const TCHAR* Cmd, FOutputDevice& Ar)
+bool FNullApplication::Exec_Dev(UWorld* InWorld, const TCHAR* Cmd, FOutputDevice& Ar)
 {
 	// Ignore any execs that doesn't start with NullApp
 	if (!FParse::Command(&Cmd, TEXT("NullApp")))
diff --git a/Engine/Source/Runtime/ApplicationCore/Public/Null/NullApplication.h b/Engine/Source/Runtime/ApplicationCore/Public/Null/NullApplication.h
--- a/Engine/Source/Runtime/ApplicationCore/Public/Null/NullApplication.h
+++ b/Engine/Source/Runtime/ApplicationCore/Public/Null/NullApplication.h
@@ -53,7 +53,7 @@ public:
 	APPLICATIONCORE_API virtual void DestroyApplication() override;
 
 	// FSelfRegisteringExec
-	APPLICATIONCORE_API virtual bool Exec(UWorld* InWorld, const TCHAR* Cmd, FOutputDevice& Ar) override;
+	APPLICATIONCORE_API virtual bool Exec_Dev(UWorld* InWorld, const TCHAR* Cmd, FOutputDevice& Ar) override;
 
 	APPLICATIONCORE_API TSharedPtr<FNullWindow> FindWindowByPtr(FGenericWindow* WindowToFind);
 
diff --git a/Engine/Source/Runtime/Core/Private/Misc/ConfigCacheIni.cpp b/Engine/Source/Runtime/Core/Private/Misc/ConfigCacheIni.cpp
--- a/Engine/Source/Runtime/Core/Private/Misc/ConfigCacheIni.cpp
+++ b/Engine/Source/Runtime/Core/Private/Misc/ConfigCacheIni.cpp
@@ -6508,7 +6508,7 @@ double GConfigShrinkTime = 0;
 
 class FIniExec : public FSelfRegisteringExec
 {
-	virtual bool Exec(class UWorld* InWorld, const TCHAR* Cmd, FOutputDevice& Ar) override
+	virtual bool Exec_Dev(class UWorld* InWorld, const TCHAR* Cmd, FOutputDevice& Ar) override
 	{
 		if (!FParse::Command(&Cmd, TEXT("CONFIG")))
 		{
diff --git a/Engine/Source/Runtime/Online/ICMP/Private/IcmpModule.cpp b/Engine/Source/Runtime/Online/ICMP/Private/IcmpModule.cpp
--- a/Engine/Source/Runtime/Online/ICMP/Private/IcmpModule.cpp
+++ b/Engine/Source/Runtime/Online/ICMP/Private/IcmpModule.cpp
@@ -17,7 +17,7 @@ void FIcmpModule::ShutdownModule()
 }
 
 
-bool FIcmpModule::Exec(UWorld* InWorld, const TCHAR* Cmd, FOutputDevice& Ar)
+bool FIcmpModule::Exec_Dev(UWorld* InWorld, const TCHAR* Cmd, FOutputDevice& Ar)
 {
 	// Ignore any execs that don't start with Icmp
 	if (FParse::Command(&Cmd, TEXT("Icmp")))
diff --git a/Engine/Source/Runtime/Online/ICMP/Public/IcmpModule.h b/Engine/Source/Runtime/Online/ICMP/Public/IcmpModule.h
--- a/Engine/Source/Runtime/Online/ICMP/Public/IcmpModule.h
+++ b/Engine/Source/Runtime/Online/ICMP/Public/IcmpModule.h
@@ -56,7 +56,7 @@ private:
 	 *
 	 * @return true if the handler consumed the input, false to continue searching handlers
 	 */
-	virtual bool Exec(UWorld* InWorld, const TCHAR* Cmd, FOutputDevice& Ar) override;
+	virtual bool Exec_Dev(UWorld* InWorld, const TCHAR* Cmd, FOutputDevice& Ar) override;
 
 	// IModuleInterface
 
