Keep track and limit the number of allocated FPlatformUserId's in the application's life cycle.

If the application allocates more then the MaxPlatformUserCount, then an ensure will be triggered.

The max number of users is defined in the project's input settings under "MaxPlatformUserCount", or via the config system.

diff --git a/Engine/Config/Mac/BaseMacInput.ini b/Engine/Config/Mac/BaseMacInput.ini
--- a/Engine/Config/Mac/BaseMacInput.ini
+++ b/Engine/Config/Mac/BaseMacInput.ini
@@ -5,4 +5,7 @@
 ; If this is false, then EKeys::Platform_Delete will be set to EKeys::Delete.
 ; 
 ; See FMacPlatformInput::GetPlatformDeleteKey for implementation
-bPlatformDeleteIsBackspace=true
\ No newline at end of file
+bPlatformDeleteIsBackspace=true
+
+[InputPlatformSettings_Mac InputPlatformSettings]
+MaxPlatformUserCount=8
\ No newline at end of file
diff --git a/Engine/Config/Windows/BaseWindowsInput.ini b/Engine/Config/Windows/BaseWindowsInput.ini
new file mode 100644
--- /dev/null
+++ b/Engine/Config/Windows/BaseWindowsInput.ini
@@ -0,0 +1,2 @@
+[InputPlatformSettings_Windows InputPlatformSettings]
+MaxPlatformUserCount=8
\ No newline at end of file
diff --git a/Engine/Source/Runtime/ApplicationCore/Private/GenericPlatform/GenericPlatformInputDeviceMapper.cpp b/Engine/Source/Runtime/ApplicationCore/Private/GenericPlatform/GenericPlatformInputDeviceMapper.cpp
--- a/Engine/Source/Runtime/ApplicationCore/Private/GenericPlatform/GenericPlatformInputDeviceMapper.cpp
+++ b/Engine/Source/Runtime/ApplicationCore/Private/GenericPlatform/GenericPlatformInputDeviceMapper.cpp
@@ -3,6 +3,7 @@
 #include "GenericPlatform/GenericPlatformInputDeviceMapper.h"
 #include "HAL/PlatformApplicationMisc.h"
 #include "Misc/CoreDelegates.h"
+#include "Misc/ConfigCacheIni.h"
 
 ///////////////////////////////////////////////////////////////////////////
 // IPlatformInputDeviceMapper
@@ -12,6 +13,38 @@ DEFINE_LOG_CATEGORY(LogInputDeviceMapper);
 IPlatformInputDeviceMapper::FOnUserInputDeviceConnectionChange IPlatformInputDeviceMapper::OnInputDeviceConnectionChange;
 IPlatformInputDeviceMapper::FOnUserInputDevicePairingChange IPlatformInputDeviceMapper::OnInputDevicePairingChange;
 
+namespace UE::Input
+{
+	// The cached value of MaxPlatformUserCount from the input config.
+	static int32 CachedMaxUserCount = -1;
+
+	static int32 GetCachedMaxUserCount()
+	{
+		// If we are uninitialized, read from the platform config
+		if (CachedMaxUserCount == -1)
+		{
+			int32 ConfigInt = -1;
+			const FString InputSettingsSection = FString::Printf(TEXT("InputPlatformSettings_%s InputPlatformSettings"), ANSI_TO_TCHAR(FPlatformProperties::IniPlatformName()));
+			static const TCHAR* DeviceMappingPolicyName = TEXT("MaxPlatformUserCount");
+
+			if (GConfig->GetInt(*InputSettingsSection, DeviceMappingPolicyName, OUT ConfigInt, GInputIni))
+			{
+				CachedMaxUserCount = ConfigInt;
+			}
+			else
+			{
+				// Require a MaxPlatformUserCount setting to exist in some Input.ini.
+				// Default to 8 as a reasonably safe fallback.
+				ensureAlwaysMsgf(false, TEXT("Unable to find MaxPlatformUserCount from config, a max of 8 will be used."));
+				CachedMaxUserCount = 8;
+			}
+		}
+		
+		return CachedMaxUserCount;
+	}
+};
+
+
 IPlatformInputDeviceMapper& IPlatformInputDeviceMapper::Get()
 {
 	static IPlatformInputDeviceMapper* StaticManager = nullptr;
@@ -82,6 +115,30 @@ int32 IPlatformInputDeviceMapper::GetAllActiveUsers(TArray<FPlatformUserId>& Out
 	return OutUsers.Num();
 }
 
+FPlatformUserId IPlatformInputDeviceMapper::GetFirstPlatformUserWithNoInputDevice() const
+{
+	const FPlatformUserId UnpairedUser = GetUserForUnpairedInputDevices();
+	
+	for (const FPlatformUserId ExistingUser : AllocatedPlatformUserIds)
+	{
+		// Skip the unpaired user, they can't have devices mapped to them.
+		if (ExistingUser == UnpairedUser)
+		{
+			continue;
+		}
+		
+		// If this use has no input device's mapped them, return that
+		const FInputDeviceId ExistingDevice = GetPrimaryInputDeviceForUser(ExistingUser);
+		if (!ExistingDevice.IsValid())
+		{
+			return ExistingUser;
+		}
+	}
+
+	// Nothing found
+	return PLATFORMUSERID_NONE;
+}
+
 bool IPlatformInputDeviceMapper::IsUnpairedUserId(const FPlatformUserId PlatformId) const
 {
 	return PlatformId == GetUserForUnpairedInputDevices();
@@ -217,6 +274,12 @@ bool IPlatformInputDeviceMapper::Internal_ChangeInputDeviceUserMapping(FInputDev
 		return false;
 	}
 
+	if (NewUserId == OldUserId)
+	{
+		UE_LOG(LogInputDeviceMapper, Log, TEXT("[%hs] DeviceId of '%d' is already mapped to platform user '%d'."), __func__, DeviceId.GetId(), OldUserId.GetInternalId());
+		return false;
+	}
+
 	// Update the existing device state to be the new owning platform user
 	if (FPlatformInputDeviceState* ExistingDeviceState = MappedInputDevices.Find(DeviceId))
 	{
@@ -247,6 +310,12 @@ bool IPlatformInputDeviceMapper::Internal_ChangeInputDeviceUserMapping(FInputDev
 	return true;
 }
 
+int32 IPlatformInputDeviceMapper::GetMaxPlatformUserCount() const
+{
+	// By default, return the cached value from the input config.
+	return UE::Input::GetCachedMaxUserCount();
+}
+
 void IPlatformInputDeviceMapper::BindCoreDelegates()
 {
 	FCoreDelegates::OnUserLoginChangedEvent.AddRaw(this, &IPlatformInputDeviceMapper::OnUserLoginChangedEvent);
@@ -280,6 +349,9 @@ FGenericPlatformInputDeviceMapper::FGenericPlatformInputDeviceMapper(const bool
 	// By default map the Default Input device to the Primary platform user in a connected state. This ensures that the SlateApplication has a 
 	// "Default" user to deal with representing the keyboard and mouse
 	Internal_MapInputDeviceToUser(GetDefaultInputDevice(), GetPrimaryPlatformUser(), EInputDeviceConnectionState::Connected);
+	
+	// Keep track of any allocated platform users so that we can utilize them when remapping input devices
+	AllocatedPlatformUserIds.AddUnique(GetPrimaryPlatformUser());
 }
 
 FPlatformUserId FGenericPlatformInputDeviceMapper::GetUserForUnpairedInputDevices() const
@@ -345,6 +417,8 @@ bool FGenericPlatformInputDeviceMapper::RemapControllerIdToPlatformUserAndDevice
 			{
 				// Otherwise just have a 1:1 mapping of input device to user id's
 				InOutUserId = FPlatformUserId::CreateFromInternalId(ControllerId);
+				
+				AllocatedPlatformUserIds.AddUnique(InOutUserId);
 			}
 			return true;
 		}
@@ -435,6 +509,14 @@ FPlatformUserId FGenericPlatformInputDeviceMapper::AllocateNewUserId()
 {
 	// Create a new platform user ID that is 1 higher than the last one
 	LastPlatformUserId = FPlatformUserId::CreateFromInternalId(LastPlatformUserId.GetInternalId() + 1);
+
+	AllocatedPlatformUserIds.AddUnique(LastPlatformUserId);
+
+	// We want to warn about this state happening, but not crash the application or anything. 
+	// Each platform has a specific number of platform users which can be signed on at any given time
+	// so going over that amount would cause undefined behavior such as your input no being routed correctly
+	// to the active local player.
+	ensureAlwaysMsgf(AllocatedPlatformUserIds.Num() <= GetMaxPlatformUserCount(), TEXT("Requested more then the max number of supported platform users! Undefined behavior may occur."));
 	
 	return LastPlatformUserId;
 }
diff --git a/Engine/Source/Runtime/ApplicationCore/Public/GenericPlatform/GenericPlatformInputDeviceMapper.h b/Engine/Source/Runtime/ApplicationCore/Public/GenericPlatform/GenericPlatformInputDeviceMapper.h
--- a/Engine/Source/Runtime/ApplicationCore/Public/GenericPlatform/GenericPlatformInputDeviceMapper.h
+++ b/Engine/Source/Runtime/ApplicationCore/Public/GenericPlatform/GenericPlatformInputDeviceMapper.h
@@ -67,6 +67,15 @@ public:
 	 */
 	APPLICATIONCORE_API virtual int32 GetAllActiveUsers(TArray<FPlatformUserId>& OutUsers) const;
 
+	/**
+	 * Finds the first FPlatformUserId which was already allocated with no input devices.
+	 * 
+	 * Returns PLATFORMUSERID_NONE if every current platform user has a valid input device mapped to it.
+	 * 
+	 * Will not return the Unpaired platform user.
+	 */
+	APPLICATIONCORE_API virtual FPlatformUserId GetFirstPlatformUserWithNoInputDevice() const;
+
 	/**
 	 * Returns the platform user id that is being used for unmapped input devices.
 	 * Will be PLATFORMUSERID_NONE if platform does not support this (this is the default behavior)
@@ -139,6 +148,13 @@ public:
 	 */
 	APPLICATIONCORE_API virtual bool Internal_ChangeInputDeviceUserMapping(FInputDeviceId DeviceId, FPlatformUserId NewUserId, FPlatformUserId OldUserId);
 
+	/**
+	* Returns the max allowed number of platform users this platform can have.
+	* 
+	* By default, this will return the value specified in the input settings ini file
+	*/
+	APPLICATIONCORE_API virtual int32 GetMaxPlatformUserCount() const;
+
 	//////////////////////////////////////////////////////////////////////////////
 	// Delegates for listening to input device changes
 	
@@ -269,6 +285,9 @@ protected:
 
 	/** Highest used input device id. Incremented in AllocateNewInputDeviceId and Internal_MapInputDeviceToUser by default. */
 	FInputDeviceId LastInputDeviceId = INPUTDEVICEID_NONE;
+
+	/** Keeps track of any allocated platform user ID's by this device mapper. */
+	TArray<FPlatformUserId> AllocatedPlatformUserIds;
 };
 
 /**
diff --git a/Engine/Source/Runtime/Engine/Classes/GameFramework/InputSettings.h b/Engine/Source/Runtime/Engine/Classes/GameFramework/InputSettings.h
--- a/Engine/Source/Runtime/Engine/Classes/GameFramework/InputSettings.h
+++ b/Engine/Source/Runtime/Engine/Classes/GameFramework/InputSettings.h
@@ -515,6 +515,17 @@ public:
 	/** Returns an array of all Hardware Device Identifiers known to this platform */
 	ENGINE_API const TArray<FHardwareDeviceIdentifier>& GetHardwareDevices() const;
 
+	////////////////////////////////////////////////////
+	// Input Device Mapping
+
+	/**
+	* The max number of FPlatformUserId's which can be allocated on this platform.
+	* 
+	* This is checked in FGenericPlatformInputDeviceMapper::AllocateNewUserId
+	*/
+	UPROPERTY(config, EditAnywhere, Category = "Device Mapping")
+	int32 MaxPlatformUserCount;
+	
 	////////////////////////////////////////////////////
 	// Trigger Feedback
 	
diff --git a/Engine/Source/Runtime/Engine/Private/UserInterface/InputSettings.cpp b/Engine/Source/Runtime/Engine/Private/UserInterface/InputSettings.cpp
--- a/Engine/Source/Runtime/Engine/Private/UserInterface/InputSettings.cpp
+++ b/Engine/Source/Runtime/Engine/Private/UserInterface/InputSettings.cpp
@@ -640,7 +640,8 @@ bool FHardwareDeviceIdentifier::operator!=(const FHardwareDeviceIdentifier& Othe
 // UInputPlatformSettings
 
 UInputPlatformSettings::UInputPlatformSettings()
-	: MaxTriggerFeedbackPosition(8)
+	: MaxPlatformUserCount(8)
+	, MaxTriggerFeedbackPosition(8)
 	, MaxTriggerFeedbackStrength(8)
 	, MaxTriggerVibrationTriggerPosition(9)
 	, MaxTriggerVibrationFrequency(255)
