Allow UE_WINDOWS_SPLASH_USE_TEXT_OUTLINE to be configurable

#rnx

diff --git a/Engine/Source/Runtime/ApplicationCore/Private/Windows/WindowsPlatformSplash.cpp b/Engine/Source/Runtime/ApplicationCore/Private/Windows/WindowsPlatformSplash.cpp
--- a/Engine/Source/Runtime/ApplicationCore/Private/Windows/WindowsPlatformSplash.cpp
+++ b/Engine/Source/Runtime/ApplicationCore/Private/Windows/WindowsPlatformSplash.cpp
@@ -34,7 +34,9 @@ THIRD_PARTY_INCLUDES_END
 
 #pragma comment( lib, "windowscodecs.lib" )
 
+#ifndef UE_WINDOWS_SPLASH_USE_TEXT_OUTLINE
 #define UE_WINDOWS_SPLASH_USE_TEXT_OUTLINE (1)
+#endif//UE_WINDOWS_SPLASH_USE_TEXT_OUTLINE 
 
 #if !defined(UE_WINDOWS_SPLASH_ENABLE_DRAG)
 #define UE_WINDOWS_SPLASH_ENABLE_DRAG WITH_EDITOR
