Fixing missing synchronization

diff --git a/Engine/Source/Runtime/ApplicationCore/Private/IOS/IOSBackgroundURLSessionHandler.cpp b/Engine/Source/Runtime/ApplicationCore/Private/IOS/IOSBackgroundURLSessionHandler.cpp
--- a/Engine/Source/Runtime/ApplicationCore/Private/IOS/IOSBackgroundURLSessionHandler.cpp
+++ b/Engine/Source/Runtime/ApplicationCore/Private/IOS/IOSBackgroundURLSessionHandler.cpp
@@ -367,7 +367,7 @@ enum class EBackgroundNSURLCDNInfoResponse : uint32 // Beware these constants de
 	NSUInteger _NextDownloadId;
 	std::promise<void> _AllDownloadsPromise;
 	std::future<void> _AllDownloadsFuture;
-	NSArray<__kindof FBackgroundNSURLCDNInfo*>* _CDNInfo;
+	NSMutableArray<__kindof FBackgroundNSURLCDNInfo*>* _CDNInfo;
 	NSTimer* _ForegroundStaleDownloadCheckTimer;
 	BackgroundHttpFileHashHelperPtr _HelperPtr;
 	int32 _MaximumConnectionsPerHost;
@@ -464,7 +464,7 @@ static constexpr NSInteger HTTPStatusCodeErrorServer = 500;
 	_AllDownloads = [NSMutableDictionary new];
 	_AllDownloadsFuture = _AllDownloadsPromise.get_future();
 	_NextDownloadId = InvalidDownloadId + 1;
-	_CDNInfo = nil;
+	_CDNInfo = [NSMutableArray new];
 	_ForegroundStaleDownloadCheckTimer = nil;
 	_bAnyTaskDidCompleteWithError = false;
 
@@ -620,14 +620,14 @@ static constexpr NSInteger HTTPStatusCodeErrorServer = 500;
 
 - (NSMutableArray<__kindof NSURL*>*)ReorderCDNsByReachability:(NSMutableArray<__kindof NSURL*>*)URLs
 {
-	if (_CDNReorderingTimeout == 0)
+	if (_CDNReorderingTimeout == 0 || [URLs count] == 0)
 	{
 		return URLs;
 	}
 
 	@synchronized (_CDNInfo)
 	{
-		if (_CDNInfo == nil)
+		if ([_CDNInfo count] == 0)
 		{
 			UE_DNLD_LOG(@"Starting to check for CDN reachability");
 
@@ -645,10 +645,10 @@ static constexpr NSInteger HTTPStatusCodeErrorServer = 500;
 			Configuration.HTTPMaximumConnectionsPerHost = _MaximumConnectionsPerHost;
 
 			NSURLSession* Session = [NSURLSession sessionWithConfiguration:Configuration];
-			NSMutableArray<__kindof FBackgroundNSURLCDNInfo*>* CDNInfo = [NSMutableArray array];
 
 			std::shared_ptr<std::atomic<int32>> PendingTasks = std::make_shared<std::atomic<int32>>();
 			std::shared_ptr<std::promise<void>> PendingTasksFinished = std::make_shared<std::promise<void>>();
+			std::shared_ptr<std::atomic<bool>> WaitingForTasksCompletionHandlers = std::make_shared<std::atomic<bool>>(true);
 
 			PendingTasks->fetch_add((int32)URLs.count);
 
@@ -665,6 +665,13 @@ static constexpr NSInteger HTTPStatusCodeErrorServer = 500;
 				// Note, completion handler might be invoked after end of this method.
 				NSURLSessionDataTask* Task = [Session dataTaskWithRequest:Request completionHandler:^(NSData* _Nullable Data, NSURLResponse* _Nullable Response, NSError* _Nullable Error)
 				{
+					// The delegate can be invoked from other thread way past invalidateAndCancel,
+					// we cannot modify _CDNInfo without synchronization.
+					if (!WaitingForTasksCompletionHandlers->load())
+					{
+						return;
+					}
+
 					const NSTimeInterval ResponseTime = -[StartTime timeIntervalSinceNow];
 
 					bool bIsOk = false;
@@ -700,7 +707,7 @@ static constexpr NSInteger HTTPStatusCodeErrorServer = 500;
 						[Info setResponse:EBackgroundNSURLCDNInfoResponse::Error];
 					}
 					[Info setResponseTime:ResponseTime];
-					[CDNInfo addObject:Info];
+					[_CDNInfo addObject:Info];
 					[Info release];
 
 					if (PendingTasks->fetch_add(-1) <= 1)
@@ -716,6 +723,7 @@ static constexpr NSInteger HTTPStatusCodeErrorServer = 500;
 			PendingTasksFinished->get_future().wait_for(std::chrono::milliseconds(_CDNReorderingTimeout));
 			UE_DNLD_LOG(@"Finished waiting for CDN reachability");
 
+			WaitingForTasksCompletionHandlers->store(false);
 			[Session invalidateAndCancel];
 
 			for (NSUInteger URLIndex = 0; URLIndex < [URLs count]; URLIndex++)
@@ -723,7 +731,7 @@ static constexpr NSInteger HTTPStatusCodeErrorServer = 500;
 				NSURL* URL = [URLs objectAtIndex:URLIndex];
 				bool bFoundCDNInfo = false;
 
-				for (FBackgroundNSURLCDNInfo* Info in CDNInfo)
+				for (FBackgroundNSURLCDNInfo* Info in _CDNInfo)
 				{
 					if ([Info.CDNHost isEqualToString:URL.host])
 					{
@@ -740,12 +748,12 @@ static constexpr NSInteger HTTPStatusCodeErrorServer = 500;
 					// If cdn/networking hasn't provided us with any info, consider request timed out
 					[Info setResponse:EBackgroundNSURLCDNInfoResponse::Timeout];
 					[Info setProvidedOrder:URLIndex];
-					[CDNInfo addObject:Info];
+					[_CDNInfo addObject:Info];
 					[Info release];
 				}
 			}
 			
-			[CDNInfo sortUsingComparator:^NSComparisonResult(FBackgroundNSURLCDNInfo* _Nonnull A, FBackgroundNSURLCDNInfo* _Nonnull B)
+			[_CDNInfo sortUsingComparator:^NSComparisonResult(FBackgroundNSURLCDNInfo* _Nonnull A, FBackgroundNSURLCDNInfo* _Nonnull B)
 			{
 				const double KeyA = [A SortingKeyWith:_bCDNReorderByPingTime];
 				const double KeyB = [B SortingKeyWith:_bCDNReorderByPingTime];
@@ -764,18 +772,13 @@ static constexpr NSInteger HTTPStatusCodeErrorServer = 500;
 				}
 			}];
 
-			for (NSUInteger i = 0; i < [CDNInfo count]; i++)
+			for (NSUInteger i = 0; i < [_CDNInfo count]; i++)
 			{
-				FBackgroundNSURLCDNInfo* Info = [CDNInfo objectAtIndex:i];
+				FBackgroundNSURLCDNInfo* Info = [_CDNInfo objectAtIndex:i];
 				UE_DNLD_LOG(@"%lu CDN '%@' AbsoluteURL '%@' Response:%u ResponseTime:%f ProvidedOrder:%lu SortingKey:%f", i, Info.CDNHost, Info.CDNAbsoluteURL, Info.Response, Info.ResponseTime, (unsigned long)Info.ProvidedOrder, [Info SortingKeyWith:_bCDNReorderByPingTime]);
 			}
-			
-			_CDNInfo = [CDNInfo retain];
 		}
-	}
 
-	if (_CDNInfo != nil && [_CDNInfo count] > 0 && [URLs count] > 0)
-	{
 		// Array sizes are assumed small enough that hashmap is not needed
 		NSMutableArray<__kindof NSURL*>* Result = [NSMutableArray arrayWithCapacity:[URLs count]];
 
@@ -802,10 +805,6 @@ static constexpr NSInteger HTTPStatusCodeErrorServer = 500;
 
 		return Result;
 	}
-	else
-	{
-		return URLs;
-	}
 }
 
 - (NSArray<__kindof FBackgroundNSURLCDNInfo*>*)GetCDNInfo
@@ -1226,7 +1225,11 @@ static constexpr NSInteger HTTPStatusCodeErrorServer = 500;
 	for (NSNumber* IterKey in AllKeys)
 	{
 		const NSUInteger DownloadId = IterKey.unsignedIntegerValue;
-		NSURLSessionDownloadTask* Task = [_AllDownloads objectForKey:IterKey];
+		NSURLSessionDownloadTask* Task = nil;
+		@synchronized(_AllDownloads)
+		{
+			Task = [_AllDownloads objectForKey:IterKey];
+		}
 		if (Task == nil)
 		{
 			continue;
