Fixing cursor offset on Macbooks with notch

#rnx

diff --git a/Engine/Source/Runtime/Apple/MetalRHI/Private/MetalRHI.cpp b/Engine/Source/Runtime/Apple/MetalRHI/Private/MetalRHI.cpp
--- a/Engine/Source/Runtime/Apple/MetalRHI/Private/MetalRHI.cpp
+++ b/Engine/Source/Runtime/Apple/MetalRHI/Private/MetalRHI.cpp
@@ -1380,13 +1380,25 @@ bool FMetalDynamicRHI::RHIGetAvailableResolutions(FScreenResolutionArray& Resolu
 					for (int32 CheckIndex = 0; CheckIndex < Resolutions.Num(); CheckIndex++)
 					{
 						FScreenResolutionRHI& CheckResolution = Resolutions[CheckIndex];
-						if ((CheckResolution.Width == Width) &&
-							(CheckResolution.Height == Height))
-						{
-							// Already in the list...
-							bAddIt = false;
-							break;
-						}
+                        if ((CheckResolution.Width == Width) &&
+                            (CheckResolution.Height == Height))
+                        {
+                            // Already in the list...
+                            bAddIt = false;
+                            break;
+                        }
+                        // Filter out unusable resolutions on notched Macs
+                        else if ((CheckResolution.Width == Width) &&
+                                 (CheckResolution.Height != Height))
+                        {
+                            bAddIt = false;
+                            if (Height < CheckResolution.Height)
+                            {
+                                // Only use the shorter (below notch and padding) version
+                                CheckResolution.Height = Height;
+                            }
+                            break;
+                        }
 					}
 				}
 				
diff --git a/Engine/Source/Runtime/ApplicationCore/Private/Mac/MacApplication.cpp b/Engine/Source/Runtime/ApplicationCore/Private/Mac/MacApplication.cpp
--- a/Engine/Source/Runtime/ApplicationCore/Private/Mac/MacApplication.cpp
+++ b/Engine/Source/Runtime/ApplicationCore/Private/Mac/MacApplication.cpp
@@ -115,6 +115,8 @@ FMacApplication::FMacApplication()
 ,   bHasLoadedInputPlugins(false)
 ,	DraggedWindow(nullptr)
 ,	WindowUnderCursor(nullptr)
+,   SystemResolutionX(0)
+,   SystemResolutionY(0)
 ,	bSystemModalMode(false)
 ,	ModifierKeysFlags(0)
 ,	CurrentModifierFlags(0)
@@ -187,6 +189,12 @@ FMacApplication::FMacApplication()
 #endif
 
 	FMacApplication::OnDisplayReconfiguration(kCGNullDirectDisplay, kCGDisplayDesktopShapeChangedFlag, this);
+    
+    // Listen for res changes to update targeted fullscreen res
+    FCoreDelegates::OnSystemResolutionChanged.AddLambda([this](uint32 ResolutionX, uint32 ResolutionY) {
+        SystemResolutionX = ResolutionX;
+        SystemResolutionY = ResolutionY;
+    });
 }
 
 FMacApplication::~FMacApplication()
@@ -1187,18 +1195,28 @@ void FMacApplication::ProcessKeyUpEvent(const FDeferredMacEvent& Event)
 void FMacApplication::OnWindowDidMove(TSharedRef<FMacWindow> Window)
 {
 	SCOPED_AUTORELEASE_POOL;
-
-	NSRect WindowFrame = [Window->GetWindowHandle() frame];
-	NSRect OpenGLFrame = [Window->GetWindowHandle() openGLFrame];
-
-	const double X = WindowFrame.origin.x;
-	const double Y = WindowFrame.origin.y + ([Window->GetWindowHandle() windowMode] == EWindowMode::Fullscreen ? WindowFrame.size.height : OpenGLFrame.size.height);
-
-	FVector2D SlatePosition = ConvertCocoaPositionToSlate(X, Y);
-
-	MessageHandler->OnMovedWindow(Window, FMath::TruncToInt(SlatePosition.X), FMath::TruncToInt(SlatePosition.Y));
-	Window->PositionX = FMath::TruncToInt(SlatePosition.X);
-	Window->PositionY = FMath::TruncToInt(SlatePosition.Y);
+    if ([Window->GetWindowHandle() windowMode] == EWindowMode::Fullscreen)
+    {
+        // Fullscreen mode always moves to 0,0
+        MessageHandler->OnMovedWindow(Window, 0, 0);
+        Window->PositionX = 0;
+        Window->PositionY = 0;
+    }
+    else
+    {
+		// Make sure to take notch and padding into account
+        NSRect WindowFrame = [Window->GetWindowHandle() frame];
+        NSRect OpenGLFrame = [Window->GetWindowHandle() openGLFrame];
+
+        const double X = WindowFrame.origin.x;
+        const double Y = WindowFrame.origin.y + OpenGLFrame.size.height;
+
+        FVector2D SlatePosition = ConvertCocoaPositionToSlate(X, Y);
+        
+        MessageHandler->OnMovedWindow(Window, FMath::TruncToInt(SlatePosition.X), FMath::TruncToInt(SlatePosition.Y));
+        Window->PositionX = FMath::TruncToInt(SlatePosition.X);
+        Window->PositionY = FMath::TruncToInt(SlatePosition.Y);
+    }
 }
 
 void FMacApplication::OnWindowWillResize(TSharedRef<FMacWindow> Window)
@@ -1217,16 +1235,10 @@ void FMacApplication::OnWindowDidResize(TSharedRef<FMacWindow> Window, bool bRes
 	OnWindowDidMove(Window);
 
 	const FCocoaWindow* CocoaWindow = Window->GetWindowHandle();
-	const NSScreen* Screen = [CocoaWindow screen];
 
-	// default is no override
-	const uint32 ScreenWidth  = FMath::TruncToInt([CocoaWindow openGLFrame].size.width * Window->GetDPIScaleFactor());
-	const uint32 ScreenHeight = FMath::TruncToInt([CocoaWindow openGLFrame].size.height * Window->GetDPIScaleFactor());
+	const uint32 WindowWidth  = FMath::TruncToInt([CocoaWindow openGLFrame].size.width * Window->GetDPIScaleFactor());
+	const uint32 WindowHeight = FMath::TruncToInt([CocoaWindow openGLFrame].size.height * Window->GetDPIScaleFactor());
 
-	// Grab current monitor data for sizing
-	const uint32 VisibleWidth = FMath::TruncToInt([Screen visibleFrame].size.width * Window->GetDPIScaleFactor());
-	const uint32 VisibleHeight = FMath::TruncToInt([Screen visibleFrame].size.height * Window->GetDPIScaleFactor());
-	
 	if (bRestoreMouseCursorLocking)
 	{
 		FMacCursor* MacCursor = (FMacCursor*)MacApplication->Cursor.Get();
@@ -1235,18 +1247,22 @@ void FMacApplication::OnWindowDidResize(TSharedRef<FMacWindow> Window, bool bRes
 			MacCursor->SetShouldIgnoreLocking(false);
 		}
 	}
-	// Depending on how the window is resized, it may result in actually moving the window slightly,
-	// e.g. going from fullscreenwindowed to fullscreen there's a few pixels of extra padding below
-	// camera housing on Apple screens for the menu bar, fullscreen doesn't have menu bars so that extra
-	// padding is removed, in effect making the window shift up.
-	if (Window->GetWindowHandle().TargetWindowMode == EWindowMode::WindowedFullscreen || Window->GetWindowHandle().TargetWindowMode == EWindowMode::Fullscreen  )
-	{
-		MessageHandler->OnMovedWindow(Window, ScreenWidth - VisibleWidth, ScreenHeight - VisibleHeight);
-		MessageHandler->OnSizeChanged(Window, VisibleWidth, VisibleHeight);
+    
+	if (Window->GetWindowHandle().TargetWindowMode == EWindowMode::Fullscreen)
+	{
+        if (SystemResolutionX != 0 && SystemResolutionY != 0)
+        {
+            // If resolution setting has changed recently, use the target resolution, window size is no longer accurate
+            MessageHandler->OnSizeChanged(Window, SystemResolutionX, SystemResolutionY);
+        }
+        else
+        {
+            MessageHandler->OnSizeChanged(Window, WindowWidth, WindowHeight);
+        }
 	}
 	else
-	{
-		MessageHandler->OnSizeChanged(Window, ScreenWidth, ScreenHeight);
+    {
+        MessageHandler->OnSizeChanged(Window, WindowWidth, WindowHeight);
 	}
 	MessageHandler->OnResizingWindow(Window);
 }
diff --git a/Engine/Source/Runtime/ApplicationCore/Public/Mac/MacApplication.h b/Engine/Source/Runtime/ApplicationCore/Public/Mac/MacApplication.h
--- a/Engine/Source/Runtime/ApplicationCore/Public/Mac/MacApplication.h
+++ b/Engine/Source/Runtime/ApplicationCore/Public/Mac/MacApplication.h
@@ -362,6 +362,8 @@ private:
 	FCocoaWindow* WindowUnderCursor;
 
 	TSharedPtr<FMacWindow> ActiveWindow;
+    uint32 SystemResolutionX;
+    uint32 SystemResolutionY;
 
 	bool bSystemModalMode;
 
diff --git a/Engine/Source/Runtime/SlateCore/Private/Widgets/SWindow.cpp b/Engine/Source/Runtime/SlateCore/Private/Widgets/SWindow.cpp
--- a/Engine/Source/Runtime/SlateCore/Private/Widgets/SWindow.cpp
+++ b/Engine/Source/Runtime/SlateCore/Private/Widgets/SWindow.cpp
@@ -964,12 +964,16 @@ void SWindow::ReshapeWindow( UE::Slate::FDeprecateVector2DParameter NewPosition,
 	{
 		if ( NativeWindow.IsValid() )
 		{
+
+// The following hack with speculative position is messing up with Mac window position with notched display
+// The accurate position is calculated in FMacApplication::OnWindowDidMove
+#if !PLATFORM_MAC
 			// Slate code often expects cached screen position to be accurate immediately after the move.
 			// This expectation is generally invalid (see UE-1308) as there may be a delay before the OS reports it back.
 			// This hack sets the position speculatively, keeping Slate happy while also giving the OS chance to report it
 			// correctly after or even during the actual call.
 			SetCachedScreenPosition(NewPositionTruncated);
-
+#endif
 			NativeWindow->ReshapeWindow(NewPositionTruncated.X, NewPositionTruncated.Y, NewSizeRounded.X, NewSizeRounded.Y);
 		}
 		else
@@ -2321,4 +2325,4 @@ TOptional<FText> SWindow::GetDefaultAccessibleText(EAccessibleType AccessibleTyp
 {
 	return GetTitle();
 }
-#endif
\ No newline at end of file
+#endif
