Create an Input Device Mapping Policy for PC

diff --git a/Engine/Config/Android/BaseAndroidInput.ini b/Engine/Config/Android/BaseAndroidInput.ini
new file mode 100644
--- /dev/null
+++ b/Engine/Config/Android/BaseAndroidInput.ini
@@ -0,0 +1,3 @@
+[InputPlatformSettings_Android InputPlatformSettings]
+MaxPlatformUserCount=4
+input.DeviceMappingPolicy=0
\ No newline at end of file
diff --git a/Engine/Config/IOS/BaseIOSInput.ini b/Engine/Config/IOS/BaseIOSInput.ini
new file mode 100644
--- /dev/null
+++ b/Engine/Config/IOS/BaseIOSInput.ini
@@ -0,0 +1,3 @@
+[InputPlatformSettings_IOS InputPlatformSettings]
+MaxPlatformUserCount=4
+input.DeviceMappingPolicy=0
\ No newline at end of file
diff --git a/Engine/Config/Mac/BaseMacInput.ini b/Engine/Config/Mac/BaseMacInput.ini
--- a/Engine/Config/Mac/BaseMacInput.ini
+++ b/Engine/Config/Mac/BaseMacInput.ini
@@ -8,4 +8,5 @@
 bPlatformDeleteIsBackspace=true
 
 [InputPlatformSettings_Mac InputPlatformSettings]
-MaxPlatformUserCount=8
\ No newline at end of file
+MaxPlatformUserCount=8
+input.DeviceMappingPolicy=1
\ No newline at end of file
diff --git a/Engine/Config/Windows/BaseWindowsInput.ini b/Engine/Config/Windows/BaseWindowsInput.ini
--- a/Engine/Config/Windows/BaseWindowsInput.ini
+++ b/Engine/Config/Windows/BaseWindowsInput.ini
@@ -1,2 +1,3 @@
 [InputPlatformSettings_Windows InputPlatformSettings]
-MaxPlatformUserCount=8
\ No newline at end of file
+MaxPlatformUserCount=8
+input.DeviceMappingPolicy=1
\ No newline at end of file
diff --git a/Engine/Plugins/Experimental/GameInputWindows/Source/GameInputWindows/Private/GameInputWindowsDevice.cpp b/Engine/Plugins/Experimental/GameInputWindows/Source/GameInputWindows/Private/GameInputWindowsDevice.cpp
--- a/Engine/Plugins/Experimental/GameInputWindows/Source/GameInputWindows/Private/GameInputWindowsDevice.cpp
+++ b/Engine/Plugins/Experimental/GameInputWindows/Source/GameInputWindows/Private/GameInputWindowsDevice.cpp
@@ -70,11 +70,14 @@ void FGameInputWindowsInputDevice::HandleDeviceDisconnected(IGameInputDevice* De
 			const FPlatformUserId NewUserToAssign = IPlatformInputDeviceMapper::Get().GetUserForUnpairedInputDevices();
 			const FInputDeviceId DeviceId = Data->GetDeviceId();
 
-			const bool bSuccess = IPlatformInputDeviceMapper::Get().Internal_MapInputDeviceToUser(DeviceId, NewUserToAssign, EInputDeviceConnectionState::Disconnected);
-			if (bSuccess)
+			if (DeviceId.IsValid())
 			{
-				Data->SetPlatformUserId(NewUserToAssign);
-			}
+				const bool bSuccess = IPlatformInputDeviceMapper::Get().Internal_MapInputDeviceToUser(DeviceId, NewUserToAssign, EInputDeviceConnectionState::Disconnected);
+				if (bSuccess)
+				{
+					Data->SetPlatformUserId(NewUserToAssign);
+				}
+			}			
 		}
 		else
 		{
@@ -100,20 +103,33 @@ void FGameInputWindowsInputDevice::HandleDeviceConnected(IGameInputDevice* Devic
 
 	// Map this input device to its user	
 	const FInputDeviceId DeviceId = Data->GetDeviceId();
-
-	// TODO: This will map every game input device to the "primary" user. This is NOT what we want in the long term.
-	// Long term, we want you to be able to configure how input devices are mapped to local players in UE with some kind of
-	// schema or config variable. We should add this as a setting in the Developer Settings object here.
-	const FPlatformUserId UserToAssign = IPlatformInputDeviceMapper::Get().GetPrimaryPlatformUser();
 	
-	const bool bSuccess = IPlatformInputDeviceMapper::Get().Internal_MapInputDeviceToUser(DeviceId, UserToAssign, EInputDeviceConnectionState::Connected);
-	if (ensure(bSuccess))
+	// We only want to map this input device to a user if it has a valid Input Device Id.
+	// Everything has an input device id, so if we got one that is invalid
+	// then that means this device has no processors in it because they were explicitly disabled.
+	if (DeviceId.IsValid())
 	{
-		Data->SetPlatformUserId(UserToAssign);
-	}
-
-	UE_LOG(LogGameInput, Log, TEXT("Using PlatformUserId %d and InputDeviceId %d for device %s"), UserToAssign.GetInternalId(), DeviceId.GetId(), *UE::GameInput::LexToString(Device));
+		IPlatformInputDeviceMapper& DeviceMapper = IPlatformInputDeviceMapper::Get();
+		
+		const FPlatformUserId UserToAssign = DeviceMapper.GetPlatformUserForNewlyConnectedDevice();
+	
+		const bool bSuccess = DeviceMapper.Internal_MapInputDeviceToUser(DeviceId, UserToAssign, EInputDeviceConnectionState::Connected);
+		if (ensure(bSuccess))
+		{
+			Data->SetPlatformUserId(UserToAssign);
+		}
 
+		UE_LOG(LogGameInput, Log, TEXT("Using PlatformUserId %d and InputDeviceId %d for device %s"),
+			UserToAssign.GetInternalId(),
+			DeviceId.GetId(),
+			*UE::GameInput::LexToString(Device));
+	}
+	else
+	{
+		UE_LOG(LogGameInput, Log, TEXT("Game Input Device %s had no processors, so it will not be assiged to a platform user. You may need to configure it in the project settings."),
+			*UE::GameInput::LexToString(Device));
+	}
+	
 	EnumerateCurrentlyConnectedDeviceTypes();
 }
 
@@ -122,8 +138,16 @@ FGameInputDeviceContainer* FGameInputWindowsInputDevice::CreateDeviceData(IGameI
 	TSharedPtr<FGameInputDeviceContainer> Container = MakeShared<FGameInputDeviceContainer>(MessageHandler, InDevice, GetCurrentGameInputKindSupport());
 	Container->InitalizeDeviceProcessors();
 
+	// If this device has any processors assigned to it (meaning that it can send input events) then we want to assign it a new 
+	// input device id. If there are no processors assigned to it, then it can't possibly send input events, so don't bother
+	// mapping it to an input device ID. This can happen if the device is explicitly disallowed for the application 
+	// via the project settings. In this case, we don't really want the rest of the engine to care about it at all
+	const FInputDeviceId AssignedInputDeviceId = 
+		Container->GetNumberOfProcessors() > 0 ?
+		IPlatformInputDeviceMapper::Get().AllocateNewInputDeviceId() :
+		INPUTDEVICEID_NONE;
+	
 	// This is a new device, we need to assign a new input device ID from the platform user
-	const FInputDeviceId AssignedInputDeviceId = IPlatformInputDeviceMapper::Get().AllocateNewInputDeviceId();
 	Container->SetInputDeviceId(AssignedInputDeviceId);
 
 	DeviceData.Emplace(Container);
diff --git a/Engine/Plugins/Runtime/Windows/WinDualShock/Source/WinDualShock/Private/WinDualShockControllers.h b/Engine/Plugins/Runtime/Windows/WinDualShock/Source/WinDualShock/Private/WinDualShockControllers.h
--- a/Engine/Plugins/Runtime/Windows/WinDualShock/Source/WinDualShock/Private/WinDualShockControllers.h
+++ b/Engine/Plugins/Runtime/Windows/WinDualShock/Source/WinDualShock/Private/WinDualShockControllers.h
@@ -12,6 +12,8 @@ THIRD_PARTY_INCLUDES_START
 	#include <pad_audio.h>
 THIRD_PARTY_INCLUDES_END
 
+#include "GameFramework/InputDeviceSubsystem.h"
+
 class FWinDualShockControllers : public FPlatformControllers
 {
 public:
@@ -51,6 +53,59 @@ public:
 	{
 		ControllerTypeIdentifiers[UserIndex] = GetControllerType(UserIndex);
 	}
+	
+	/**
+	 * Returns true if the most recently used FInputDeviceId for the given ControllerId
+	 * is "owned" by this input interface.  
+	 */
+	bool OwnsMostRecentlyUsedDevice(int32 ControllerId) const
+	{
+		// Find the most recently used input device by the given platform user
+		const FPlatformUserId UserId = IPlatformInputDeviceMapper::Get().GetPlatformUserForUserIndex(ControllerId);
+		UInputDeviceSubsystem* DeviceSubsystem = UInputDeviceSubsystem::Get();
+		check(DeviceSubsystem);
+		const FInputDeviceId MostRecentDevice = DeviceSubsystem->GetMostRecentlyUsedInputDeviceId(UserId);
+
+		// If that is one of our controllers, then process our force feedback!
+		for (int32 UserIndex=0; UserIndex < SCE_USER_SERVICE_MAX_LOGIN_USERS; ++UserIndex)
+		{
+			const FControllerState& ControllerState = ControllerStates[UserIndex];
+
+			// Note: This will only work for "normal" pad devices, not special devices. That is ok
+			// because windows does not support special device handles at all
+			const FInputDeviceId OurDeviceId = InternalDeviceIdMappings.FindDeviceId(ControllerState.Handle);
+			if (OurDeviceId == MostRecentDevice)
+			{
+				return true;	
+			}
+		}
+		
+		return false;
+	}
+
+	virtual void SetForceFeedbackChannelValue(int32 ControllerId, FForceFeedbackChannelType ChannelType, float Value) override
+	{
+		// If that is one of our controllers, then process our force feedback! 	
+		if (!OwnsMostRecentlyUsedDevice(ControllerId))
+		{
+			return;
+		}
+		
+		FPlatformControllers::SetForceFeedbackChannelValue(ControllerId, ChannelType, Value);
+	}
+	
+	virtual void SetForceFeedbackChannelValues(int32 ControllerId, const FForceFeedbackValues &Values) override
+	{
+		// If that is one of our controllers, then process our force feedback! 	
+		if (!OwnsMostRecentlyUsedDevice(ControllerId))
+		{
+			return;
+		}
+		
+		FPlatformControllers::SetForceFeedbackChannelValues(ControllerId, Values);
+	}
+	
+	
 
 private:
 	float OutputGain = 1.0f;
diff --git a/Engine/Plugins/Runtime/Windows/XInputDevice/Source/XInputDevice/Private/XInputInterface.cpp b/Engine/Plugins/Runtime/Windows/XInputDevice/Source/XInputDevice/Private/XInputInterface.cpp
--- a/Engine/Plugins/Runtime/Windows/XInputDevice/Source/XInputDevice/Private/XInputInterface.cpp
+++ b/Engine/Plugins/Runtime/Windows/XInputDevice/Source/XInputDevice/Private/XInputInterface.cpp
@@ -15,6 +15,8 @@
 #include "Windows/HideWindowsPlatformTypes.h"
 #pragma pack (pop)
 
+#include "GameFramework/InputDeviceSubsystem.h"
+
 static int32 ForceControllerStateUpdate = 0;
 FAutoConsoleVariableRef CVarForceControllerStateUpdate(
 	TEXT("XInput.ForceControllerStateUpdate"),
@@ -23,6 +25,8 @@ FAutoConsoleVariableRef CVarForceControllerStateUpdate(
 	TEXT("0: Not Enabled, 1: Enabled"),
 	ECVF_Default);
 
+DEFINE_LOG_CATEGORY_STATIC(LogXInput, Log, All);
+
 TSharedRef< XInputInterface > XInputInterface::Create(  const TSharedRef< FGenericApplicationMessageHandler >& InMessageHandler, bool bShouldBePrimaryDevice )
 {
 	return MakeShareable( new XInputInterface( InMessageHandler, bShouldBePrimaryDevice ) );
@@ -114,11 +118,27 @@ void XInputInterface::GetPlatformUserAndDevice(int32 InControllerId, EInputDevic
 {
 	if (bIsPrimaryDevice)
 	{
+		OutDeviceId = InternalDeviceIdMappings.GetOrCreateDeviceId(InControllerId);
+
 		IPlatformInputDeviceMapper& DeviceMapper = IPlatformInputDeviceMapper::Get();
-		DeviceMapper.RemapControllerIdToPlatformUserAndDevice(InControllerId, OUT OutPlatformUserId, OUT OutDeviceId);
+
+		// If we have just been connected, then get the new platform user for a new device connection
+		if (InDeviceState == EInputDeviceConnectionState::Connected)
+		{
+			OutPlatformUserId = DeviceMapper.GetPlatformUserForNewlyConnectedDevice();
+		}
+		// If we have been disconnected, remap this device to the unpaired device user
+		else if (InDeviceState == EInputDeviceConnectionState::Disconnected)
+		{
+			OutPlatformUserId = DeviceMapper.GetUserForUnpairedInputDevices();
+		}
+		else
+		{
+			OutPlatformUserId = DeviceMapper.GetUserForInputDevice(OutDeviceId);
+		}
 
 		// If the controller is connected now but was not before, refresh the information
-		if (InDeviceState == EInputDeviceConnectionState::Connected || InDeviceState == EInputDeviceConnectionState::Disconnected)
+		if (InDeviceState == EInputDeviceConnectionState::Connected)
 		{
 			DeviceMapper.Internal_MapInputDeviceToUser(OutDeviceId, OutPlatformUserId, InDeviceState);
 		}
@@ -189,9 +209,23 @@ void XInputInterface::SendControllerEvents()
 
 			FPlatformUserId PlatformUser = PLATFORMUSERID_NONE;
 			FInputDeviceId InputDevice = INPUTDEVICEID_NONE;
-			EInputDeviceConnectionState State = UE::XInputInterface::Private::GetInputDeviceConnectionState(bWasConnected, ControllerState.bIsConnected);
-			GetPlatformUserAndDevice(ControllerState.ControllerId, State, OUT PlatformUser, OUT InputDevice);
+			const EInputDeviceConnectionState ConnectionState = UE::XInputInterface::Private::GetInputDeviceConnectionState(bWasConnected, ControllerState.bIsConnected);
+			GetPlatformUserAndDevice(ControllerState.ControllerId, ConnectionState, OUT PlatformUser, OUT InputDevice);
 
+			// If the device has been disconnected, it needs to be remapped to the "Unpaired" input device
+			// at the end of it's input processing. Use the last valid platform user to ensure they receive
+			// "release" events for every active input.
+			if (!PlatformUser.IsValid() && ConnectionState == EInputDeviceConnectionState::Disconnected)
+			{
+				PlatformUser = ControllerState.LastUsedValidPlatformUserId;
+			}
+
+			// If we never got a valid platform user for some reason, we can't do anything
+			if (!PlatformUser.IsValid())
+			{
+				return;
+			}
+			
 			bool CurrentStates[MAX_NUM_CONTROLLER_BUTTONS] = {0};
 		
 			// Get the current state of all buttons
@@ -225,6 +259,9 @@ void XInputInterface::SendControllerEvents()
 				if (OldAxisValue != NewAxisValue || FMath::Abs((int32)NewAxisValue) > DeadZone)
 				{
 					MessageHandler->OnControllerAnalog(GamePadKey, PlatformUser, InputDevice, NewAxisValueNormalized);
+
+					UE_LOG(LogXInput, VeryVerbose, TEXT("[%hs] PlatUser: %d DeviceId %d Key: '%s'     Value: %.3f"),
+						__func__, PlatformUser.GetInternalId(), InputDevice.GetId(), *GamePadKey.ToString(), NewAxisValueNormalized);
 				}
 				OldAxisValue = NewAxisValue;
 			};
@@ -250,10 +287,16 @@ void XInputInterface::SendControllerEvents()
 					if( CurrentStates[ButtonIndex] )
 					{
 						MessageHandler->OnControllerButtonPressed( Buttons[ButtonIndex], PlatformUser, InputDevice, false );
+						
+						UE_LOG(LogXInput, Verbose, TEXT("[%hs] OnControllerButtonPressed PlatUser: %d DeviceId %d Key: '%s' bIsRepeat: false"),
+							__func__, PlatformUser.GetInternalId(), InputDevice.GetId(), *Buttons[ButtonIndex].ToString());
 					}
 					else
 					{
 						MessageHandler->OnControllerButtonReleased( Buttons[ButtonIndex], PlatformUser, InputDevice, false );
+						
+						UE_LOG(LogXInput, Verbose, TEXT("[%hs] OnControllerButtonReleased PlatUser: %d DeviceId %d Key: '%s' bIsRepeat: false"),
+							__func__, PlatformUser.GetInternalId(), InputDevice.GetId(), *Buttons[ButtonIndex].ToString());
 					}
 
 					if ( CurrentStates[ButtonIndex] != 0 )
@@ -265,6 +308,9 @@ void XInputInterface::SendControllerEvents()
 				else if ( CurrentStates[ButtonIndex] != 0 && ControllerState.NextRepeatTime[ButtonIndex] <= CurrentTime )
 				{
 					MessageHandler->OnControllerButtonPressed( Buttons[ButtonIndex], PlatformUser, InputDevice, true );
+					
+					UE_LOG(LogXInput, Verbose, TEXT("[%hs] OnControllerButtonPressed PlatUser: %d DeviceId %d Key: '%s' bIsRepeat: true"),
+						__func__, PlatformUser.GetInternalId(), InputDevice.GetId(), *Buttons[ButtonIndex].ToString());
 
 					// set the button's NextRepeatTime to the ButtonRepeatDelay
 					ControllerState.NextRepeatTime[ButtonIndex] = CurrentTime + ButtonRepeatDelay;
@@ -290,6 +336,20 @@ void XInputInterface::SendControllerEvents()
 				ControllerState.LastLargeValue = LargeValue;
 				ControllerState.LastSmallValue = SmallValue;
 			}
+
+			// Keep track of the last valid platform user id
+			if (PlatformUser.IsValid())
+			{
+				ControllerState.LastUsedValidPlatformUserId = PlatformUser;	
+			}
+
+			// Remap the input device to the now invalid platform user at the end of the frame
+			// This way we can ensure that slate gets reported the "0.0" analog values to stop any
+			// active inputs upon disconnect
+			if (ConnectionState == EInputDeviceConnectionState::Disconnected)
+			{
+				IPlatformInputDeviceMapper::Get().Internal_MapInputDeviceToUser(InputDevice, IPlatformInputDeviceMapper::Get().GetUserForUnpairedInputDevices(), ConnectionState);
+			}
 		}
 	}
 
@@ -314,14 +374,26 @@ void XInputInterface::SetDeviceProperty(int32 ControllerId, const FInputDevicePr
 
 void XInputInterface::SetChannelValue( int32 ControllerId, const FForceFeedbackChannelType ChannelType, const float Value )
 {
+	const FPlatformUserId UserId = IPlatformInputDeviceMapper::Get().GetPlatformUserForUserIndex(ControllerId);
+
+	UInputDeviceSubsystem* DeviceSubsystem = UInputDeviceSubsystem::Get();
+	check(DeviceSubsystem);
+	const FInputDeviceId MostRecentDevice = DeviceSubsystem->GetMostRecentlyUsedInputDeviceId(UserId);
+	
 	if (ControllerId >= 0 && ControllerId < MAX_NUM_XINPUT_CONTROLLERS)
 	{
-		FControllerState& ControllerState = ControllerStates[ ControllerId ];
-
-		if( ControllerState.bIsConnected )
+		for (int32 i = 0; i < MAX_NUM_XINPUT_CONTROLLERS; ++i)
 		{
-			switch( ChannelType )
+			FControllerState& ControllerState = ControllerStates[ i ];
+
+			const FInputDeviceId CurrentDeviceId = InternalDeviceIdMappings.FindDeviceId(ControllerState.ControllerId);
+
+			if (ControllerState.bIsConnected &&
+				ControllerState.LastUsedValidPlatformUserId == UserId &&
+				CurrentDeviceId.IsValid() && MostRecentDevice == CurrentDeviceId)
 			{
+				switch( ChannelType )
+				{
 				case FForceFeedbackChannelType::LEFT_LARGE:
 					ControllerState.ForceFeedback.LeftLarge = Value;
 					break;
@@ -337,20 +409,34 @@ void XInputInterface::SetChannelValue( int32 ControllerId, const FForceFeedbackC
 				case FForceFeedbackChannelType::RIGHT_SMALL:
 					ControllerState.ForceFeedback.RightSmall = Value;
 					break;
+				}
 			}
-		}
+		}	
 	}
 }
 
 void XInputInterface::SetChannelValues( int32 ControllerId, const FForceFeedbackValues &Values )
 {
+	const FPlatformUserId UserId = IPlatformInputDeviceMapper::Get().GetPlatformUserForUserIndex(ControllerId);
+
+	UInputDeviceSubsystem* DeviceSubsystem = UInputDeviceSubsystem::Get();
+	check(DeviceSubsystem);
+	const FInputDeviceId MostRecentDevice = DeviceSubsystem->GetMostRecentlyUsedInputDeviceId(UserId);
+	
 	if (ControllerId >= 0 && ControllerId < MAX_NUM_XINPUT_CONTROLLERS)
 	{
-		FControllerState& ControllerState = ControllerStates[ ControllerId ];
-
-		if( ControllerState.bIsConnected )
+		for (int32 i = 0; i < MAX_NUM_XINPUT_CONTROLLERS; ++i)
 		{
-			ControllerState.ForceFeedback = Values;
-		}
+			FControllerState& ControllerState = ControllerStates[ i ];
+			
+			const FInputDeviceId CurrentDeviceId = InternalDeviceIdMappings.FindDeviceId(ControllerState.ControllerId);
+
+			if (ControllerState.bIsConnected &&
+				ControllerState.LastUsedValidPlatformUserId == UserId &&
+				CurrentDeviceId.IsValid() && CurrentDeviceId == MostRecentDevice)
+			{
+				ControllerState.ForceFeedback = Values;
+			}	
+		}		
 	}
 }
diff --git a/Engine/Plugins/Runtime/Windows/XInputDevice/Source/XInputDevice/Private/XInputInterface.h b/Engine/Plugins/Runtime/Windows/XInputDevice/Source/XInputDevice/Private/XInputInterface.h
--- a/Engine/Plugins/Runtime/Windows/XInputDevice/Source/XInputDevice/Private/XInputInterface.h
+++ b/Engine/Plugins/Runtime/Windows/XInputDevice/Source/XInputDevice/Private/XInputInterface.h
@@ -6,6 +6,7 @@
 #include "GenericPlatform/IInputInterface.h"
 #include "IInputDevice.h"
 #include "GenericPlatform/GenericApplicationMessageHandler.h"
+#include "GenericPlatform/GenericInputDeviceMap.h"
 #include "Misc/CoreMiscDefines.h"
 
 /** Max number of controllers. */
@@ -114,6 +115,11 @@ private:
 
 		float LastLargeValue;
 		float LastSmallValue;
+
+		/**
+		 * The last valid FPlatformUserId that can be used on the frame where the input device is disconnected
+		 */
+		FPlatformUserId LastUsedValidPlatformUserId;
 	};
 
 	/** If we've been notified by the system that the controller state may have changed */
@@ -136,9 +142,15 @@ private:
 	/** Delay before sending a repeat message after a button has been pressed for a while */
 	float ButtonRepeatDelay;
 
-	/**  */
+	/** Array of gamepad button names where the index in the array is the XInput button */
 	FGamepadKeyNames::Type Buttons[MAX_NUM_CONTROLLER_BUTTONS];
 
-	/**  */
+	/** Reference to the message handler, used to send the input state to the application. */
 	TSharedRef<FGenericApplicationMessageHandler> MessageHandler;
+
+	/**
+	 * A map of XInput device identifiers (int32's) to their uniquely assigned
+	 * FInputDeviceId from the engine.
+	 */
+	TInputDeviceMap<int32> InternalDeviceIdMappings;
 };
diff --git a/Engine/Source/Runtime/ApplicationCore/Private/GenericPlatform/GenericPlatformInputDeviceMapper.cpp b/Engine/Source/Runtime/ApplicationCore/Private/GenericPlatform/GenericPlatformInputDeviceMapper.cpp
--- a/Engine/Source/Runtime/ApplicationCore/Private/GenericPlatform/GenericPlatformInputDeviceMapper.cpp
+++ b/Engine/Source/Runtime/ApplicationCore/Private/GenericPlatform/GenericPlatformInputDeviceMapper.cpp
@@ -4,6 +4,7 @@
 #include "HAL/PlatformApplicationMisc.h"
 #include "Misc/CoreDelegates.h"
 #include "Misc/ConfigCacheIni.h"
+#include "GenericPlatform/InputDeviceMappingPolicy.h"
 
 ///////////////////////////////////////////////////////////////////////////
 // IPlatformInputDeviceMapper
@@ -42,8 +43,56 @@ namespace UE::Input
 		
 		return CachedMaxUserCount;
 	}
+
+	/** Cache the input device policy from the config so we don't need to read from the config cache as often */
+	static EInputDeviceMappingPolicy CachedDevicePolicy = EInputDeviceMappingPolicy::Invalid;
+	
+	/**
+	* Returns the currently set mapping policy for the current platform.
+	* This is set by UInputSettings::DeviceMappingPolicy.
+	*/
+	static EInputDeviceMappingPolicy GetDeviceMappingPolicyFromConfig()
+	{
+		if (CachedDevicePolicy == EInputDeviceMappingPolicy::Invalid)
+		{
+			int32 PolicyInt = -1;
+			const FString InputSettingsSection = FString::Printf(TEXT("InputPlatformSettings_%s InputPlatformSettings"), ANSI_TO_TCHAR(FPlatformProperties::IniPlatformName()));
+			static const TCHAR* DeviceMappingPolicyName = TEXT("input.DeviceMappingPolicy");
+			
+			if (GConfig->GetInt(*InputSettingsSection, DeviceMappingPolicyName, OUT PolicyInt, GInputIni))
+			{
+				CachedDevicePolicy = static_cast<EInputDeviceMappingPolicy>(PolicyInt);
+			}
+			else
+			{
+				// Require a input.DeviceMappingPolicy setting to exist in some Input.ini
+#if PLATFORM_DESKTOP
+				ensureMsgf(false, TEXT("Unable to find an input device mapping policy, PrimaryUserSharesKeyboardAndFirstGamepad will be used."));
+				CachedDevicePolicy = EInputDeviceMappingPolicy::PrimaryUserSharesKeyboardAndFirstGamepad;
+#else
+				ensureMsgf(false, TEXT("Unable to find an input device mapping policy, UseManagedPlatformLogin will be used."));
+				CachedDevicePolicy = EInputDeviceMappingPolicy::UseManagedPlatformLogin;
+#endif
+			}	
+		}
+		
+		return CachedDevicePolicy;
+	}
 };
 
+FString LexToString(const EInputDeviceMappingPolicy Policy)
+{
+	switch (Policy)
+	{
+	case EInputDeviceMappingPolicy::Invalid: return TEXT("Invalid"); break;
+	case EInputDeviceMappingPolicy::UseManagedPlatformLogin: return TEXT("UseManagedPlatformLogin"); break;
+	case EInputDeviceMappingPolicy::PrimaryUserSharesKeyboardAndFirstGamepad: return TEXT("PrimaryUserSharesKeyboardAndFirstGamepad"); break;
+	case EInputDeviceMappingPolicy::CreateUniquePlatformUserForEachDevice: return TEXT("CreateUniquePlatformUserForEachDevice"); break;
+	case EInputDeviceMappingPolicy::MapAllDevicesToPrimaryUser: return TEXT("MapAllDevicesToPrimaryUser"); break;
+	}
+
+	return TEXT("Unknown EInputDeviceMappingPolicy");
+}
 
 IPlatformInputDeviceMapper& IPlatformInputDeviceMapper::Get()
 {
@@ -153,6 +202,46 @@ bool IPlatformInputDeviceMapper::IsInputDeviceMappedToUnpairedUser(const FInputD
 	return false;
 }
 
+FPlatformUserId IPlatformInputDeviceMapper::GetPlatformUserForNewlyConnectedDevice(const int32 InUserId /*= -1*/)
+{
+	const EInputDeviceMappingPolicy Policy = UE::Input::GetDeviceMappingPolicyFromConfig();
+	const FPlatformUserId PrimaryUser = GetPrimaryPlatformUser();
+	
+	// If the policy is to always map to the primary user, then do so.
+	if (Policy == EInputDeviceMappingPolicy::MapAllDevicesToPrimaryUser)
+	{
+		return PrimaryUser;
+	}
+	// If the primary user is supposed to share the first gamepad and keyboard,
+	// then we will map this device to it if it only has one device connected to it. 
+	else if (Policy == EInputDeviceMappingPolicy::PrimaryUserSharesKeyboardAndFirstGamepad)
+	{
+		TArray<FInputDeviceId> PrimaryUserDevices;
+		GetAllInputDevicesForUser(PrimaryUser, OUT PrimaryUserDevices);
+
+		if (PrimaryUserDevices.Num() <= 1)
+		{
+			return PrimaryUser;
+		}
+	}
+	// If you have given an optional UserID and this platform is using controller ID as it's user id, then
+	// we can return the platform user associated with that user.
+	else if (InUserId != -1 && IsUsingControllerIdAsUserId())
+	{
+		return GetPlatformUserForUserIndex(InUserId);
+	}
+
+	// Return the first valid platform user which does not have any input devices mapped to it
+	const FPlatformUserId ValidUserWithNoDevices = GetFirstPlatformUserWithNoInputDevice();
+	if (ValidUserWithNoDevices.IsValid())
+	{
+		return ValidUserWithNoDevices;
+	}
+
+	// If all current platform users have a valid input device mapped to them, then we want to create a new one!
+	return AllocateNewUserId();
+}
+
 FPlatformUserId IPlatformInputDeviceMapper::GetUserForInputDevice(FInputDeviceId DeviceId) const
 {
 	if (const FPlatformInputDeviceState* FoundState = MappedInputDevices.Find(DeviceId))
@@ -310,6 +399,132 @@ bool IPlatformInputDeviceMapper::Internal_ChangeInputDeviceUserMapping(FInputDev
 	return true;
 }
 
+
+#if WITH_EDITOR
+void IPlatformInputDeviceMapper::HandleInputDevicePolicyChanged(const EInputDeviceMappingPolicy NewPolicy)
+{
+	if (NewPolicy == UE::Input::CachedDevicePolicy)
+	{
+		UE_LOG(LogInputDeviceMapper, Log, TEXT("[%hs] Device Mapping Policy is already set to '%s'..."),
+			__func__,
+			*LexToString(UE::Input::CachedDevicePolicy));
+		
+		return;
+	}
+
+	UE_LOG(LogInputDeviceMapper, Log, TEXT("[%hs] Changing device mapping policy from '%s' to '%s'..."),
+		__func__,
+		*LexToString(UE::Input::CachedDevicePolicy),
+		*LexToString(NewPolicy));
+	
+	const FPlatformUserId PrimaryUser = GetPrimaryPlatformUser();
+	const FPlatformUserId UnpairedUser = GetUserForUnpairedInputDevices();
+
+	TArray<FInputDeviceId> AllConnectedDevices;
+	GetAllConnectedInputDevices(OUT AllConnectedDevices);
+		
+	switch(NewPolicy)
+	{
+	case EInputDeviceMappingPolicy::UseManagedPlatformLogin:
+		// Since this would be handled by the platform... no need to do anything
+		break;
+	case EInputDeviceMappingPolicy::PrimaryUserSharesKeyboardAndFirstGamepad:
+		{
+			// Start off by remapping all input devices to the invalid platform user id
+			// except for the primary device
+			for (FInputDeviceId Device : AllConnectedDevices)
+			{
+				if (Device == GetDefaultInputDevice())
+				{
+					continue;
+				}
+				
+				const FPlatformUserId OldUserId = GetUserForInputDevice(Device);
+				Internal_ChangeInputDeviceUserMapping(Device, UnpairedUser, OldUserId);
+			}
+		
+			TArray<FInputDeviceId> PrimaryUserDevices;
+			
+			// Ensure that the primary user has the primary input devices (keyboard)
+			// and a gamepad assigned to them
+			while (!AllConnectedDevices.IsEmpty())
+			{
+				const FInputDeviceId Device = AllConnectedDevices.Pop();
+				
+				const FPlatformUserId OldUserId = GetUserForInputDevice(Device);
+				
+				GetAllInputDevicesForUser(PrimaryUser, OUT PrimaryUserDevices);
+
+				// If the primary user only has 1 device then map this gamepad to it
+				// We also want to ensure that the default input device is always mapped to the primary user
+				if (PrimaryUserDevices.Num() <= 1 || Device == GetDefaultInputDevice())
+				{
+					Internal_ChangeInputDeviceUserMapping(Device, PrimaryUser, OldUserId);
+					continue;
+				}
+
+				// Otherwise, map to the next available platform user or a new one if necessary
+				FPlatformUserId NextUser = GetFirstPlatformUserWithNoInputDevice();
+				if (!NextUser.IsValid())
+				{
+					NextUser = AllocateNewUserId();
+				}
+
+				Internal_ChangeInputDeviceUserMapping(Device, NextUser, OldUserId);
+			}
+		}
+		break;
+	case EInputDeviceMappingPolicy::CreateUniquePlatformUserForEachDevice:
+		{
+			// Start off by remapping all input devices to the invalid platform user id
+			// except for the primary device
+			for (FInputDeviceId Device : AllConnectedDevices)
+			{
+				if (Device == GetDefaultInputDevice())
+				{
+					continue;
+				}
+				
+				const FPlatformUserId OldUserId = GetUserForInputDevice(Device);
+				Internal_ChangeInputDeviceUserMapping(Device, UnpairedUser, OldUserId);
+			}
+
+			while (!AllConnectedDevices.IsEmpty())
+			{
+				const FInputDeviceId Device = AllConnectedDevices.Pop();
+
+				// Map to the next available platform user ID, or create a new one if necessary
+				FPlatformUserId NextUser = GetFirstPlatformUserWithNoInputDevice();
+				if (!NextUser.IsValid())
+				{
+					NextUser = AllocateNewUserId();
+				}
+
+				Internal_ChangeInputDeviceUserMapping(Device, NextUser, PLATFORMUSERID_NONE);
+			}
+		}
+		break;
+	case EInputDeviceMappingPolicy::MapAllDevicesToPrimaryUser:
+	{
+		// Map all connected input devices to the primary platform user
+		for (const FInputDeviceId Device : AllConnectedDevices)
+		{
+			const FPlatformUserId OldUserId = GetUserForInputDevice(Device);
+			Internal_ChangeInputDeviceUserMapping(Device, PrimaryUser, OldUserId);
+		}
+	}
+	break;
+	case EInputDeviceMappingPolicy::Invalid:
+	default:
+		checkNoEntry();
+		break;
+	}
+
+	// Update the cached value to use this new one
+	UE::Input::CachedDevicePolicy = NewPolicy;
+}
+#endif	// #if WITH_EDITOR
+
 int32 IPlatformInputDeviceMapper::GetMaxPlatformUserCount() const
 {
 	// By default, return the cached value from the input config.
@@ -328,8 +543,11 @@ void IPlatformInputDeviceMapper::UnbindCoreDelegates()
 
 bool IPlatformInputDeviceMapper::ShouldCreateUniqueUserForEachDevice() const
 {
-	// TODO: Get this from some settings you can toggle per platform in the editor
-	return true;
+	const EInputDeviceMappingPolicy Policy = UE::Input::GetDeviceMappingPolicyFromConfig();
+
+	return 
+		Policy == EInputDeviceMappingPolicy::CreateUniquePlatformUserForEachDevice || 
+		Policy == EInputDeviceMappingPolicy::UseManagedPlatformLogin;
 }
 
 ///////////////////////////////////////////////////////////////////////////
@@ -387,8 +605,12 @@ bool FGenericPlatformInputDeviceMapper::RemapControllerIdToPlatformUserAndDevice
 		}
 		else if (ControllerId >= 0)
 		{
-			// Just use controller id, and copy over device
-			OutInputDeviceId = FInputDeviceId::CreateFromInternalId(ControllerId);
+			//if (!OutInputDeviceId.IsValid())
+			{
+				// Just use controller id, and copy over device
+				OutInputDeviceId = FInputDeviceId::CreateFromInternalId(ControllerId);	
+			}
+			
 			// If we were already given a valid platform user then we can just stop now.
 			// This will be the case on platforms with an existing concept of "User Logins" from platform itself
 			if (InOutUserId.IsValid())
diff --git a/Engine/Source/Runtime/ApplicationCore/Public/GenericPlatform/GenericInputDeviceMap.h b/Engine/Source/Runtime/ApplicationCore/Public/GenericPlatform/GenericInputDeviceMap.h
new file mode 100644
--- /dev/null
+++ b/Engine/Source/Runtime/ApplicationCore/Public/GenericPlatform/GenericInputDeviceMap.h
@@ -0,0 +1,110 @@
+// Copyright Epic Games, Inc. All Rights Reserved.
+
+#pragma once
+
+#include "Containers/Map.h"
+#include "GenericPlatform/GenericPlatformInputDeviceMapper.h"
+#include "Misc/CoreMiscDefines.h"	// For FInputDeviceId/FPlatformUserId
+
+/**
+* An interface which can be added to any IInputDevice (or any other type 
+* which is creating Human Interface Devices) to store it's unique identifier
+* for those physical devices and associate them with a FInputDeviceId so that
+* the rest of the Unreal Engine can interact with them
+* 
+* Some input interfaces use simple int32's to identify input devices, while
+* others may have more specific GUID types types which are needed.
+*/
+template <class TDeviceKeyType>
+class TInputDeviceMap
+{
+public:	
+
+	TInputDeviceMap() = default;
+	~TInputDeviceMap() = default;
+
+	/**
+	* Given the DeviceKey, find it's associated FInputDeviceId.
+	* 
+	* If one does not exist yet, such as for a newly connected device,
+	* then create one from the IPlatformInputDeviceMapper.
+	*/
+	[[nodiscard]] FInputDeviceId GetOrCreateDeviceId(const TDeviceKeyType DeviceKey)
+	{
+		FInputDeviceId OutDeviceId = INPUTDEVICEID_NONE;
+		
+		// If we already know about this device, then we can just use that info
+		if (FInputDeviceId* ExistingDeviceId = MappedDeviceIds.Find(DeviceKey))
+		{
+			OutDeviceId = *ExistingDeviceId;
+		}
+		// Otherwise, we have not seen this input device before (its a new connection)
+		// So we need a new FInputDeviceId and a FPlatformUserId to map it to. 
+		else
+		{
+			OutDeviceId = IPlatformInputDeviceMapper::Get().AllocateNewInputDeviceId();
+
+			// Keep track of both the FInputDeviceId -> DeviceKey 
+			// and DeviceKey -> FInputDeviceID for quick lookup
+			MappedDeviceIds.Add(DeviceKey, OutDeviceId);
+			MappedIdToKey.Add(OutDeviceId, DeviceKey);
+		}
+
+		return OutDeviceId;
+	}
+
+	/**
+	* Returns the FInputDeviceId for the given device key.
+	* 
+	* Returns FInputDeviceId::Invalid if it is not yet mapped.
+	*/
+	[[nodiscard]] FInputDeviceId FindDeviceId(const TDeviceKeyType DeviceKey) const
+	{
+		FInputDeviceId OutDeviceId = INPUTDEVICEID_NONE;
+
+		// If we already know about this device, then we can just use that info
+		if (const FInputDeviceId* ExistingDeviceId = MappedDeviceIds.Find(DeviceKey))
+		{
+			OutDeviceId = *ExistingDeviceId;
+		}
+
+		return OutDeviceId;
+	}
+
+	/**
+	* Returns the FInputDeviceId associated with the given device key. Asserts if not found.
+	*/
+	[[nodiscard]] FInputDeviceId FindDeviceIdChecked(const TDeviceKeyType DeviceKey) const
+	{
+		return MappedDeviceIds.FindChecked(DeviceKey);
+	}
+
+	/**
+	* Returns the device key for the given FInputDeviceId. Nullptr if not found.
+	*/
+	[[nodiscard]] const TDeviceKeyType& GetDeviceKeyChecked(const FInputDeviceId DeviceId) const
+	{
+		return MappedIdToKey.FindChecked(DeviceId);
+	}
+
+	/**
+	* Returns the device key for the given FInputDeviceId. Nullptr if not found.
+	*/
+	[[nodiscard]] const TDeviceKeyType* FindDeviceKey(const FInputDeviceId DeviceId) const
+	{
+		return MappedIdToKey.Find(DeviceId);
+	}
+
+protected:
+
+	/**
+	* Map of the assigned DeviceKeyType to their assigned FInputDeviceId's from the engine.
+	*/
+	TMap<TDeviceKeyType, FInputDeviceId> MappedDeviceIds;
+
+	/**
+	* A map of the assigned FInputDeviceId to their associated Device Key type for 
+	* fast lookup.
+	*/
+	TMap<FInputDeviceId, TDeviceKeyType> MappedIdToKey;
+};
\ No newline at end of file
diff --git a/Engine/Source/Runtime/ApplicationCore/Public/GenericPlatform/GenericPlatformInputDeviceMapper.h b/Engine/Source/Runtime/ApplicationCore/Public/GenericPlatform/GenericPlatformInputDeviceMapper.h
--- a/Engine/Source/Runtime/ApplicationCore/Public/GenericPlatform/GenericPlatformInputDeviceMapper.h
+++ b/Engine/Source/Runtime/ApplicationCore/Public/GenericPlatform/GenericPlatformInputDeviceMapper.h
@@ -10,6 +10,8 @@
 
 DECLARE_LOG_CATEGORY_EXTERN(LogInputDeviceMapper, Log, All);
 
+enum class EInputDeviceMappingPolicy : int32;
+
 /**
  * Base class to private a mapping of Platform Users (FPlatformUserID)
  * to their associated available input devices (FInputDeviceID).
@@ -99,6 +101,17 @@ public:
 	/** Returns the default device id used for things like keyboard/mouse input */
 	virtual FInputDeviceId GetDefaultInputDevice() const = 0;
 
+	/**
+	 * Returns the platform user id that a newly connected device should be mapped to.
+	 *
+	 * If "ShouldCreateUniqueUserForEachDevice" is true, then a new platform user id will be allocated for this device.
+	 * Otherwise, this will return the primary platform user.
+	 *
+	 * @param InUserId		An optional user index which can be passed in. If this is not -1, and IsUsingControllerIdAsUserId is true,
+	 *						then the platform user id for that user will be returned. 	
+	 */
+	APPLICATIONCORE_API virtual FPlatformUserId GetPlatformUserForNewlyConnectedDevice(const int32 InUserId = -1); 
+
 	/** Returns the platform user attached to this input device, or PLATFORMUSERID_NONE if invalid */
 	APPLICATIONCORE_API virtual FPlatformUserId GetUserForInputDevice(FInputDeviceId DeviceId) const;
 
@@ -178,6 +191,14 @@ public:
 
 	FOnUserInputDeviceConnectionChange& GetOnInputDeviceConnectionChange() const { return OnInputDeviceConnectionChange; }
 	FOnUserInputDevicePairingChange& GetOnInputDevicePairingChange() const { return OnInputDevicePairingChange; }
+
+#if WITH_EDITOR
+	/**
+	 * Remaps the currently connected input devices based on the new given policy. This property can only be changed
+	 * at editor time and should not be used at runtime.
+	 */
+	APPLICATIONCORE_API void HandleInputDevicePolicyChanged(const EInputDeviceMappingPolicy NewPolicy);
+#endif	// #if WITH_EDITOR
 	
 	//////////////////////////////////////////////////////////////////////////////
 	// Functions to provide compatibility between the old "int32 ControllerId"
diff --git a/Engine/Source/Runtime/ApplicationCore/Public/GenericPlatform/InputDeviceMappingPolicy.h b/Engine/Source/Runtime/ApplicationCore/Public/GenericPlatform/InputDeviceMappingPolicy.h
new file mode 100644
--- /dev/null
+++ b/Engine/Source/Runtime/ApplicationCore/Public/GenericPlatform/InputDeviceMappingPolicy.h
@@ -0,0 +1,91 @@
+// Copyright Epic Games, Inc. All Rights Reserved.
+
+#pragma once
+
+#include "HAL/Platform.h"	// for int32
+
+/**
+ * The input device mapping policy controls how Human Interface Devices (HID's)
+ * are "mapped" or "assigned" to local players in your game. Depending on this
+ * policy, the creation of FPlatformUserId's and FInputDeviceId's can change.
+ *
+ * This policy is only in affect on platforms which do NOT have a managed user login
+ * or some kind of dedicated input mapping OS functionality (consoles). 
+ */
+enum class EInputDeviceMappingPolicy : int32
+{
+	/**
+	 * Invalid mapping policy. Used to validate that the one in the settings is correct.
+	 */
+	Invalid = -1,
+	
+	/**
+	 * Use the current build target's built in platform login functionality.
+	 * This is often times a prompt which is displayed by the target system
+	 * to assign an input device to a specific player.
+	 *
+	 * Note: The behavior of this setting on non-user managed platforms (such as PC targets) is undefined.
+	 */
+	UseManagedPlatformLogin = 0,
+	
+	/**
+	 * Maps the keyboard/mouse and the first connected controller to the primary platform user.
+	 * Any subsequently connected gamepads will be mapped to the next platform user without a
+	 * valid input device already assigned to them.
+	 *
+	 * An example of this policy when you have 3 gamepads and a keyboard/mouse connected would be:
+	 *
+	 * -- Platform User 0 (Primary)
+	 * ---> Keyboard and mouse
+	 * ---> Gamepad 0
+	 * |
+	 * -- Platform User 1
+	 * ---> Gamepad 1
+	 * |
+	 * -- PlatformUser 2
+	 * ---> Gamepad 2
+	 *
+	 * This is the default device mapping policy for PC. 
+	 */
+	PrimaryUserSharesKeyboardAndFirstGamepad = 1,
+
+	/**
+	 * If this policy is set, then every input device that is connected will be mapped
+	 * to a unique platform user, meaning that each device is treated as a separate local player.
+	 * This is what you want if you want plugging in a new controller to mean that you have a
+	 * new local player in your game.
+	 *
+	 * This only defines the default behavior upon connection of a new device, you can still remap
+	 * these input devices to new platform users later in the application lifecycle.
+	 *
+	 * An example of this policy when you have 3 gamepads and a keyboard/mouse connected would be:
+	 *
+	 * -- Platform User 0 (Primary)
+	 * ---> Keyboard and mouse
+	 * |
+	 * -- Platform User 1
+	 * ---> Gamepad 0
+	 * |
+	 * -- PlatformUser 2
+	 * ---> Gamepad 1
+	 * |
+	 * -- PlatformUser 3
+	 * ---> Gamepad 2
+	 */
+	CreateUniquePlatformUserForEachDevice = 2,
+
+	/**
+	 * If this policy is set, then every connected input device will be mapped to the
+	 * primary platform user and new platform user Id's will not be created from input
+	 * devices.
+	 *
+	 * An example of this policy when you have 3 gamepads and a keyboard/mouse connected would be:
+	 * -- Platform User 0 (Primary)
+	 * ---> Keyboard and mouse
+	 * ---> Gamepad 0
+	 * ---> Gamepad 1
+	 * ---> Gamepad 2
+	 */
+	MapAllDevicesToPrimaryUser = 3
+};
+
diff --git a/Engine/Source/Runtime/CoreUObject/Public/UObject/NoExportTypes.h b/Engine/Source/Runtime/CoreUObject/Public/UObject/NoExportTypes.h
--- a/Engine/Source/Runtime/CoreUObject/Public/UObject/NoExportTypes.h
+++ b/Engine/Source/Runtime/CoreUObject/Public/UObject/NoExportTypes.h
@@ -63,6 +63,7 @@
 
 #include "../../../ApplicationCore/Public/GenericPlatform/ICursor.h"
 #include "../../../ApplicationCore/Public/GenericPlatform/IInputInterface.h"
+#include "../../../ApplicationCore/Public/GenericPlatform/InputDeviceMappingPolicy.h"
 
 #endif
 
@@ -2666,6 +2667,67 @@ enum class EInputDeviceConnectionState : uint8
 	Connected
 };
 
+/**
+ * The input device mapping policy controls how Human Interface Devices (HID's)
+ * are "mapped" or "assigned" to local players in your game. Depending on this
+ * policy, the creation of FPlatformUserId's and FInputDeviceId's can change.
+ *
+ * This policy is only in affect on platforms which do NOT have a managed user login
+ * or some kind of dedicated input mapping OS functionality (consoles). 
+ */
+UENUM()
+enum class EInputDeviceMappingPolicy : int32
+{
+	/**
+	 * Invalid mapping policy. Used to validate that the one in the settings is correct.
+	 */
+	Invalid = -1 UMETA(Hidden),
+	
+	/**
+	 * Use the current build target's built in platform login functionality.
+	 * This is often times a prompt which is displayed by the target system
+	 * to assign an input device to a specific player.
+	 */
+	UseManagedPlatformLogin = 0,
+	
+	/**
+	 * Maps the keyboard/mouse and the first connected controller to the primary platform user.
+	 * Any subsequently connected gamepads will be mapped to the next platform user without a
+	 * valid input device already assigned to them.
+	 *
+	 * An example of this policy when you have 3 gamepads and a keyboard/mouse connected would be something like this;
+	 *
+	 * -- Platform User 0 (Primary)
+	 * ---> Keyboard and mouse
+	 * ---> Gamepad 0
+	 * |
+	 * -- Platform User 1
+	 * ---> Gamepad 1
+	 * |
+	 * -- PlatformUser 2
+	 * ---> Gamepad 2
+	 */
+	PrimaryUserSharesKeyboardAndFirstGamepad = 1,
+
+	/**
+	 * If this policy is set, then every input device that is connected will be mapped
+	 * to a unique platform user, meaning that it is treated as a separate local player.
+	 * This is what you want if you want plugging in a new controller to mean that you have a
+	 * new local player in your game.
+	 *
+	 * This only defines the default behavior upon connection of a new device, you can still remap
+	 * these input devices to new platform users later in the application lifecycle.
+	 */
+	CreateUniquePlatformUserForEachDevice = 2,
+
+	/**
+	 * If this policy is set, then every connected input device will be mapped to the
+	 * primary platform user and new platform user Id's will not be created from input
+	 * devices. 
+	 */
+	MapAllDevicesToPrimaryUser = 3
+};
+
 /**
  * Represents input device triggers that are available
  *
diff --git a/Engine/Source/Runtime/Engine/Classes/GameFramework/InputSettings.h b/Engine/Source/Runtime/Engine/Classes/GameFramework/InputSettings.h
--- a/Engine/Source/Runtime/Engine/Classes/GameFramework/InputSettings.h
+++ b/Engine/Source/Runtime/Engine/Classes/GameFramework/InputSettings.h
@@ -8,6 +8,7 @@
 #include "InputCoreTypes.h"
 #include "UObject/SoftObjectPath.h"
 #include "GameFramework/PlayerInput.h"
+#include "GenericPlatform/InputDeviceMappingPolicy.h"
 #include "Components/InputComponent.h"
 #include "Engine/PlatformSettings.h"
 
@@ -525,6 +526,17 @@ public:
 	*/
 	UPROPERTY(config, EditAnywhere, Category = "Device Mapping")
 	int32 MaxPlatformUserCount;
+
+	/**
+	 * The input device mapping policy controls how Human Interface Devices (HID's)
+	 * are "mapped" or "assigned" to local players in your game. Depending on this
+	 * policy, the creation of FPlatformUserId's and FInputDeviceId's can change.
+	 *
+	 * This policy is only in affect on platforms which do NOT have a managed user login
+	 * or some kind of dedicated input mapping OS functionality (consoles have this).
+	 */
+	UPROPERTY(config, EditAnywhere, Category = "Device Mapping", meta=(ConsoleVariable="input.DeviceMappingPolicy"))
+	EInputDeviceMappingPolicy DeviceMappingPolicy;
 	
 	////////////////////////////////////////////////////
 	// Trigger Feedback
@@ -574,6 +586,10 @@ public:
 	
 protected:
 
+#if WITH_EDITOR
+	ENGINE_API virtual void PostEditChangeChainProperty(struct FPropertyChangedChainEvent& PropertyChangedEvent) override;
+#endif
+
 	/** A list of identifiable hardware devices available on this platform */
 	UPROPERTY(config, EditAnywhere, Category = "Hardware")
 	TArray<FHardwareDeviceIdentifier> HardwareDevices;
diff --git a/Engine/Source/Runtime/Engine/Private/UserInterface/InputSettings.cpp b/Engine/Source/Runtime/Engine/Private/UserInterface/InputSettings.cpp
--- a/Engine/Source/Runtime/Engine/Private/UserInterface/InputSettings.cpp
+++ b/Engine/Source/Runtime/Engine/Private/UserInterface/InputSettings.cpp
@@ -9,6 +9,7 @@
 #include "Misc/CommandLine.h"
 #include "UObject/UnrealType.h"
 #include "UObject/UObjectIterator.h"
+#include "HAL/IConsoleManager.h"
 #include "HAL/PlatformApplicationMisc.h"
 
 #include UE_INLINE_GENERATED_CPP_BY_NAME(InputSettings)
@@ -21,6 +22,19 @@
 #include "Windows/WindowsHWrapper.h"
 #endif
 
+namespace UE::Input
+{
+	static int32 DeviceMappingPolicy = 1;
+	static FAutoConsoleVariableRef CVarDeviceMappingPolicy(TEXT("input.DeviceMappingPolicy"),
+	DeviceMappingPolicy,
+	TEXT("The policy defining the behavior of a newly connected input device. See EInputDeviceMappingPolicy\n")
+	TEXT("0 - UseManagedPlatformLogin\n")
+	TEXT("1 - PrimaryUserSharesKeyboardAndFirstGamepad (PC Default)\n")
+	TEXT("2 - CreateUniquePlatformUserForEachDevice\n")
+	TEXT("3 - MapAllDevicesToPrimaryUser\n"),
+	ECVF_Default);
+};
+
 UInputSettings::UInputSettings(const FObjectInitializer& ObjectInitializer)
 	: Super(ObjectInitializer)
 	, bCaptureMouseOnLaunch(true)
@@ -641,6 +655,7 @@ bool FHardwareDeviceIdentifier::operator!=(const FHardwareDeviceIdentifier& Othe
 
 UInputPlatformSettings::UInputPlatformSettings()
 	: MaxPlatformUserCount(8)
+	, DeviceMappingPolicy(EInputDeviceMappingPolicy::PrimaryUserSharesKeyboardAndFirstGamepad)
 	, MaxTriggerFeedbackPosition(8)
 	, MaxTriggerFeedbackStrength(8)
 	, MaxTriggerVibrationTriggerPosition(9)
@@ -683,6 +698,18 @@ const TArray<FHardwareDeviceIdentifier>& UInputPlatformSettings::GetHardwareDevi
 }
 
 #if WITH_EDITOR
+void UInputPlatformSettings::PostEditChangeChainProperty(FPropertyChangedChainEvent& PropertyChangedEvent)
+{
+	Super::PostEditChangeChainProperty(PropertyChangedEvent);
+	
+	const FName MemberPropertyName = PropertyChangedEvent.PropertyChain.GetActiveMemberNode()->GetValue()->GetFName();
+	
+	if (MemberPropertyName == GET_MEMBER_NAME_CHECKED(UInputPlatformSettings, DeviceMappingPolicy))
+	{
+		IPlatformInputDeviceMapper::Get().HandleInputDevicePolicyChanged(DeviceMappingPolicy);
+	}
+}
+
 const TArray<FName>& UInputPlatformSettings::GetAllHardwareDeviceNames()
 {
 	static TArray<FName> HardwareDevices;
