Can now tell the player's controller about analog stick deadzones through PlayerController::SetControllerDeadZones, or directly as a device property using FInputDeviceAnalogStickDeadZoneProperty

While deadzone handling should always be done on the game's side, some device libraries also do some input filtering themselves so that noisy unintended input doesn't prevent power-saving modes from activating.

Since a game may allow players to set deadzones lower than the device defaults, we've added a way to inform the device of the player's deadzone settings.

#tests PIE, affected controllers, minimum, maximum, and default deadzone settings, among others.

diff --git a/Engine/Source/Runtime/ApplicationCore/Public/GenericPlatform/IInputInterface.h b/Engine/Source/Runtime/ApplicationCore/Public/GenericPlatform/IInputInterface.h
--- a/Engine/Source/Runtime/ApplicationCore/Public/GenericPlatform/IInputInterface.h
+++ b/Engine/Source/Runtime/ApplicationCore/Public/GenericPlatform/IInputInterface.h
@@ -103,6 +103,19 @@ enum class EInputDeviceTriggerMask : uint8
 };
 ENUM_CLASS_FLAGS(EInputDeviceTriggerMask)
 
+/**
+ * Represents input device analog sticks that are available
+ *
+ * NOTE: Make sure to keep this type in sync with the reflected version in NoExportTypes.h!
+ */
+enum class EInputDeviceAnalogStickMask : uint8
+{
+	None = 0x00,
+	Left = 0x01,
+	Right = 0x02
+};
+ENUM_CLASS_FLAGS(EInputDeviceAnalogStickMask)
+
 struct FInputDeviceProperty
 {
 	FInputDeviceProperty(FName InName)
@@ -209,6 +222,44 @@ struct FInputDeviceTriggerVibrationProperty : public FInputDeviceTriggerProperty
 	int32 VibrationAmplitude = 0;
 };
 
+/** Base class for device properties that affect Analog Sticks */
+struct FInputDeviceAnalogStickProperty : public FInputDeviceProperty
+{
+	FInputDeviceAnalogStickProperty(FName InName)
+		: FInputDeviceProperty(InName)
+	{}
+
+	FInputDeviceAnalogStickProperty(FName InName, EInputDeviceAnalogStickMask InAffectedStick)
+		: FInputDeviceProperty(InName)
+		, AffectedStick(InAffectedStick)
+	{}
+
+	/** Which stick this property should effect */
+	EInputDeviceAnalogStickMask AffectedStick = EInputDeviceAnalogStickMask::None;
+};
+
+/**
+ * A generic input device property that sets Deadzone for Analog Sticks.
+ * 
+ * NOTE: Not all input device libraries do their own deadzone filtering, so the game must still do it.
+ * This is just to tell those device libraries that do some filtering not to ignore inputs of this size (circular) or larger.
+ */
+struct FInputDeviceAnalogStickDeadZoneProperty : public FInputDeviceAnalogStickProperty
+{
+	FInputDeviceAnalogStickDeadZoneProperty()
+		: FInputDeviceAnalogStickProperty(PropertyName())
+	{}
+
+	FInputDeviceAnalogStickDeadZoneProperty(EInputDeviceAnalogStickMask InAffectedStick, float InDeadZone)
+		: FInputDeviceAnalogStickProperty(PropertyName(), InAffectedStick)
+		, DeadZone(InDeadZone)
+	{}
+
+	static FName PropertyName() { return FName("InputDeviceAnalogStickDeadZone"); }
+
+	float DeadZone = 0;
+};
+
 /**
  * Interface for the input interface.
  */
diff --git a/Engine/Source/Runtime/CoreUObject/Public/UObject/NoExportTypes.h b/Engine/Source/Runtime/CoreUObject/Public/UObject/NoExportTypes.h
--- a/Engine/Source/Runtime/CoreUObject/Public/UObject/NoExportTypes.h
+++ b/Engine/Source/Runtime/CoreUObject/Public/UObject/NoExportTypes.h
@@ -2742,6 +2742,19 @@ enum class EInputDeviceTriggerMask : uint8
 	All			= Left | Right
 };
 
+/**
+ * Represents input device analog sticks that are available
+ *
+ * NOTE: Make sure to keep this type in sync with the reflected version in IInputInterface.h!
+ */
+UENUM(BlueprintType)
+enum class EInputDeviceAnalogStickMask : uint8
+{
+	None = 0x00,
+	Left = 0x01,
+	Right = 0x02
+};
+
 /**
  * Data about an input device's current state
  */
diff --git a/Engine/Source/Runtime/Engine/Classes/GameFramework/PlayerController.h b/Engine/Source/Runtime/Engine/Classes/GameFramework/PlayerController.h
--- a/Engine/Source/Runtime/Engine/Classes/GameFramework/PlayerController.h
+++ b/Engine/Source/Runtime/Engine/Classes/GameFramework/PlayerController.h
@@ -1364,6 +1364,20 @@ public:
 	UFUNCTION(BlueprintCallable, Category="Game|Feedback")
 	ENGINE_API void ResetControllerLightColor();
 
+	/**
+	 * Sets the deadzones of the player's controller
+	 * @param	LeftDeadZone 	Inner DeadZone for the left analog stick
+	 * @param	RightDeadZone	Inner DeadZone for the right analog stick
+	 */
+	UFUNCTION(BlueprintCallable, Category = "Input")
+	ENGINE_API void SetControllerDeadZones(const float LeftDeadZone, const float RightDeadZone);
+
+	/**
+	 * Resets the player's controller deadzones to default
+	 */
+	UFUNCTION(BlueprintCallable, Category = "Input")
+	ENGINE_API void ResetControllerDeadZones();
+
 	/**
 	 * Travel to a different map or IP address. Calls the PreClientTravel event before doing anything.
 	 * NOTE: This is implemented as a locally executed wrapper for ClientTravelInternal, to avoid API compatability breakage
diff --git a/Engine/Source/Runtime/Engine/Private/PlayerController.cpp b/Engine/Source/Runtime/Engine/Private/PlayerController.cpp
--- a/Engine/Source/Runtime/Engine/Private/PlayerController.cpp
+++ b/Engine/Source/Runtime/Engine/Private/PlayerController.cpp
@@ -4605,6 +4605,48 @@ void APlayerController::ResetControllerLightColor()
 	}
 }
 
+void APlayerController::SetControllerDeadZones(const float LeftDeadZone, const float RightDeadZone)
+{
+	if (Player == nullptr)
+	{
+		return;
+	}
+
+	if (FSlateApplication::IsInitialized())
+	{
+		IInputInterface* InputInterface = FSlateApplication::Get().GetInputInterface();
+		if (InputInterface)
+		{
+			const int32 ControllerId = CastChecked<ULocalPlayer>(Player)->GetControllerId();
+			FInputDeviceAnalogStickDeadZoneProperty DeadZoneLeft(EInputDeviceAnalogStickMask::Left, LeftDeadZone);
+			InputInterface->SetDeviceProperty(ControllerId, &DeadZoneLeft);
+			FInputDeviceAnalogStickDeadZoneProperty DeadZoneRight(EInputDeviceAnalogStickMask::Right, RightDeadZone);
+			InputInterface->SetDeviceProperty(ControllerId, &DeadZoneRight);
+		}
+	}
+}
+
+void APlayerController::ResetControllerDeadZones()
+{
+	if (Player == nullptr)
+	{
+		return;
+	}
+
+	if (FSlateApplication::IsInitialized())
+	{
+		IInputInterface* InputInterface = FSlateApplication::Get().GetInputInterface();
+		if (InputInterface)
+		{
+			const int32 ControllerId = CastChecked<ULocalPlayer>(Player)->GetControllerId();
+			FInputDeviceAnalogStickDeadZoneProperty DeadZoneLeft(EInputDeviceAnalogStickMask::Left, 0.f);
+			InputInterface->SetDeviceProperty(ControllerId, &DeadZoneLeft);
+			FInputDeviceAnalogStickDeadZoneProperty DeadZoneRight(EInputDeviceAnalogStickMask::Right, 0.f);
+			InputInterface->SetDeviceProperty(ControllerId, &DeadZoneRight);
+		}
+	}
+}
+
 void APlayerController::ProcessForceFeedbackAndHaptics(const float DeltaTime, const bool bGamePaused)
 {
 	if (Player == nullptr)
