[IOS] Protecting usages of IOS main thread functions
The game was deadlocking due to UI elements being accessed in non-main threads.

The main issue here is that running the game on an iPhone 13 with iOS 18.0 the game will almost certainly hang at some point.
With this fix, the game runs smoothly.
During gameplay, it seems to have improved performance.
But I was not able to measure yet due to the inability to run Insights Captures, it will OOM.

Context here: https://epicgames.slack.com/archives/C06S7N8LAD7/p1727302915935169

#rnx

diff --git a/Engine/Source/Runtime/Apple/MetalRHI/Private/MetalViewport.cpp b/Engine/Source/Runtime/Apple/MetalRHI/Private/MetalViewport.cpp
--- a/Engine/Source/Runtime/Apple/MetalRHI/Private/MetalViewport.cpp
+++ b/Engine/Source/Runtime/Apple/MetalRHI/Private/MetalViewport.cpp
@@ -245,16 +245,19 @@ void FMetalViewport::Resize(uint32 InSizeX, uint32 InSizeY, bool bInIsFullscreen
 	// (even though Apple's HDR displays are P3)
 	// and its compositor will do the conversion.
 	{
-		IOSAppDelegate* AppDelegate = [IOSAppDelegate GetDelegate];
-		FIOSView* IOSView = AppDelegate.IOSView;
-		CAMetalLayer* MetalLayer = (CAMetalLayer*) IOSView.layer;
-		
-		if (MetalFormat != (MTL::PixelFormat) MetalLayer.pixelFormat)
-		{
-			MetalLayer.pixelFormat = (MTLPixelFormat) MetalFormat;
-		}
-		
-		[IOSView UpdateRenderWidth:InSizeX andHeight:InSizeY];
+		dispatch_sync(dispatch_get_main_queue(), ^{
+			IOSAppDelegate* AppDelegate = [IOSAppDelegate GetDelegate];
+			FIOSView* IOSView = AppDelegate.IOSView;
+			
+			CAMetalLayer* MetalLayer = (CAMetalLayer*) IOSView.layer;
+			
+			if (MetalFormat != (MTL::PixelFormat) MetalLayer.pixelFormat)
+			{
+				MetalLayer.pixelFormat = (MTLPixelFormat) MetalFormat;
+			}
+			
+			[IOSView UpdateRenderWidth:InSizeX andHeight:InSizeY];
+		});
 	}
 #endif
 
diff --git a/Engine/Source/Runtime/ApplicationCore/Private/IOS/IOSAppDelegate.cpp b/Engine/Source/Runtime/ApplicationCore/Private/IOS/IOSAppDelegate.cpp
--- a/Engine/Source/Runtime/ApplicationCore/Private/IOS/IOSAppDelegate.cpp
+++ b/Engine/Source/Runtime/ApplicationCore/Private/IOS/IOSAppDelegate.cpp
@@ -501,7 +501,11 @@ static IOSAppDelegate* CachedDelegate = nil;
 
 -(bool)IsIdleTimerEnabled
 {
-	return ([UIApplication sharedApplication].idleTimerDisabled == NO);
+	__block BOOL Result;
+	dispatch_sync(dispatch_get_main_queue(), ^{
+		Result = ([UIApplication sharedApplication].idleTimerDisabled == NO);
+	});
+	return Result;
 }
 
 -(void)DeferredEnableIdleTimer
diff --git a/Engine/Source/Runtime/ApplicationCore/Private/IOS/IOSApplication.cpp b/Engine/Source/Runtime/ApplicationCore/Private/IOS/IOSApplication.cpp
--- a/Engine/Source/Runtime/ApplicationCore/Private/IOS/IOSApplication.cpp
+++ b/Engine/Source/Runtime/ApplicationCore/Private/IOS/IOSApplication.cpp
@@ -53,7 +53,10 @@ void FIOSApplication::InitializeWindow( const TSharedRef< FGenericWindow >& InWi
 	const TSharedPtr< FIOSWindow > ParentWindow = StaticCastSharedPtr< FIOSWindow >( InParent );
 
 	Windows.Add( Window );
-	Window->Initialize( this, InDefinition, ParentWindow, bShowImmediately );
+	
+	dispatch_sync(dispatch_get_main_queue(), ^{
+		Window->Initialize( this, InDefinition, ParentWindow, bShowImmediately );
+	});
 }
 
 void FIOSApplication::SetMessageHandler( const TSharedRef< FGenericApplicationMessageHandler >& InMessageHandler )
@@ -130,7 +133,11 @@ void FIOSApplication::PollGameDeviceState( const float TimeDelta )
 
 FPlatformRect FIOSApplication::GetWorkArea( const FPlatformRect& CurrentWindow ) const
 {
-	return FIOSWindow::GetScreenRect();
+	__block FPlatformRect Rect;
+	dispatch_sync(dispatch_get_main_queue(), ^{
+		Rect = FIOSWindow::GetScreenRect();
+	});
+	return Rect;
 }
 
 static TAutoConsoleVariable<float> CVarSafeZone_Landscape_Left(TEXT("SafeZone.Landscape.Left"), -1.0f, TEXT("Safe Zone - Landscape - Left"));
@@ -145,12 +152,17 @@ UIEdgeInsets CachedInsets;
 
 void FDisplayMetrics::RebuildDisplayMetrics(FDisplayMetrics& OutDisplayMetrics)
 {
-	const FPlatformRect& Rect = FIOSWindow::GetScreenRect();
-	const FIOSView *View = [[IOSAppDelegate GetDelegate] IOSView];
-	[View CalculateContentScaleFactor:Rect.Right ScreenHeight:Rect.Bottom];
+	__block FPlatformRect Rect;
+	__block CGFloat ContentScaleFactor;
+	dispatch_sync(dispatch_get_main_queue(), ^{
+		Rect = FIOSWindow::GetScreenRect();
+		const FIOSView *View = [[IOSAppDelegate GetDelegate] IOSView];
+		[View CalculateContentScaleFactor:Rect.Right ScreenHeight:Rect.Bottom];
+		ContentScaleFactor = View.contentScaleFactor;
+	});
 	
 	// Get screen rect
-	OutDisplayMetrics.PrimaryDisplayWorkAreaRect = FIOSWindow::GetScreenRect();
+	OutDisplayMetrics.PrimaryDisplayWorkAreaRect = Rect;
 	OutDisplayMetrics.VirtualDisplayRect = OutDisplayMetrics.PrimaryDisplayWorkAreaRect;
 
     // Total screen size of the primary monitor
@@ -193,7 +205,7 @@ void FDisplayMetrics::RebuildDisplayMetrics(FDisplayMetrics& OutDisplayMetrics)
     OutDisplayMetrics.TitleSafePaddingSize.W = Inset_Bottom;
     
     //scale the thing
-    OutDisplayMetrics.TitleSafePaddingSize *= View.contentScaleFactor;
+	OutDisplayMetrics.TitleSafePaddingSize *= ContentScaleFactor;
     
     OutDisplayMetrics.ActionSafePaddingSize = OutDisplayMetrics.TitleSafePaddingSize;
 #endif
@@ -203,8 +215,10 @@ void FDisplayMetrics::RebuildDisplayMetrics(FDisplayMetrics& OutDisplayMetrics)
 void FIOSApplication::CacheDisplayMetrics()
 {
 #if !PLATFORM_TVOS
-    CachedInsets = [[[[UIApplication sharedApplication] delegate] window] safeAreaInsets];
-    CachedOrientation = [[[[[UIApplication sharedApplication] delegate] window] windowScene] interfaceOrientation];
+	dispatch_sync(dispatch_get_main_queue(), ^{
+		CachedInsets = [[[[UIApplication sharedApplication] delegate] window] safeAreaInsets];
+		CachedOrientation = [[[[[UIApplication sharedApplication] delegate] window] windowScene] interfaceOrientation];
+	});
 #endif
 }
 
@@ -216,17 +230,16 @@ TSharedRef< FGenericWindow > FIOSApplication::MakeWindow()
 #if !PLATFORM_TVOS
 void FIOSApplication::OrientationChanged(UIInterfaceOrientation orientation)
 {
-	// this is called on the IOS thread. It turns out that it's possible for the resolution to change AGAIN by the time the game
-	// thread processes the resize, so we queue up the current size from the ios thread, send that to the RHI to resize to that
-	// (and no longer checking if it matches current frame size, because it may not). If another resize happens, that new size
-	// will get queued up here and sent to RHI, so eventually the size will be correct
-	FPlatformRect WindowRect = FIOSWindow::GetScreenRect();
-	int32 WindowWidth = WindowRect.Right - WindowRect.Left;
-	int32 WindowHeight = WindowRect.Bottom - WindowRect.Top;
-
 	// queue up the size as we see it now all the way to the RHI 
 	[FIOSAsyncTask CreateTaskWithBlock : ^ bool(void)
 	 {
+		__block FPlatformRect WindowRect;
+		dispatch_sync(dispatch_get_main_queue(), ^{
+			WindowRect = FIOSWindow::GetScreenRect();
+		});
+		int32 WindowWidth = WindowRect.Right - WindowRect.Left;
+		int32 WindowHeight = WindowRect.Bottom - WindowRect.Top;
+		
 	 	FIOSApplication* App = [IOSAppDelegate GetDelegate].IOSApplication;
 	 
 		App->GetMessageHandler()->OnSizeChanged(App->Windows[0],WindowWidth,WindowHeight, false);
diff --git a/Engine/Source/Runtime/ApplicationCore/Private/IOS/IOSView.cpp b/Engine/Source/Runtime/ApplicationCore/Private/IOS/IOSView.cpp
--- a/Engine/Source/Runtime/ApplicationCore/Private/IOS/IOSView.cpp
+++ b/Engine/Source/Runtime/ApplicationCore/Private/IOS/IOSView.cpp
@@ -338,7 +338,13 @@ MTL::Device* GMetalDevice = nullptr;
 
 - (id<CAMetalDrawable>)MakeDrawable
 {
-    return [(CAMetalLayer*)self.layer nextDrawable];
+	__block CAMetalLayer* MetalLayer = nil;
+	dispatch_sync(dispatch_get_main_queue(), ^{
+		MetalLayer = (CAMetalLayer*)[self layer];
+	});
+	// this call cannot be made on the MainThread
+	// thus requiring the code above to MainThreadCall here
+	return [MetalLayer nextDrawable];
 }
 
 - (void)DestroyFramebuffer
diff --git a/Engine/Source/Runtime/Core/Private/IOS/IOSPlatformMisc.cpp b/Engine/Source/Runtime/Core/Private/IOS/IOSPlatformMisc.cpp
--- a/Engine/Source/Runtime/Core/Private/IOS/IOSPlatformMisc.cpp
+++ b/Engine/Source/Runtime/Core/Private/IOS/IOSPlatformMisc.cpp
@@ -1247,9 +1247,10 @@ void FIOSPlatformMisc::RegisterForRemoteNotifications()
 									  int32 types = (int32)granted;
                                       if (granted)
                                       {
-                                          UIApplication* application = [UIApplication sharedApplication];
-                                          [application registerForRemoteNotifications];
-                                          
+										  dispatch_sync(dispatch_get_main_queue(), ^{
+											  UIApplication* application = [UIApplication sharedApplication];
+											  [application registerForRemoteNotifications];  
+										  });
                                       }
 									  FFunctionGraphTask::CreateAndDispatchWhenReady([types]()
 																					 {
