PR #12347: Fix for iOS Integrated Keyboard dismissing itself every time space is pressed. Tweaked a bit from the original PR to handle both integrated and UITextFielf keyboards.
Tested with TextField and TextArea with both types of keyboard.

diff --git a/Engine/Source/Runtime/ApplicationCore/Private/IOS/IOSView.cpp b/Engine/Source/Runtime/ApplicationCore/Private/IOS/IOSView.cpp
--- a/Engine/Source/Runtime/ApplicationCore/Private/IOS/IOSView.cpp
+++ b/Engine/Source/Runtime/ApplicationCore/Private/IOS/IOSView.cpp
@@ -630,7 +630,10 @@ self.accessibilityElements = @[Window.accessibilityContainer];
 			FIOSInputInterface::QueueKeyInput(KEYCODE_ENTER, Char);
 			
 			// hide the keyboard
-			[self resignFirstResponder];
+			if (!bIsUsingIntegratedKeyboard)
+			{
+				[self resignFirstResponder];
+			}
 		}
 		else
 		{
@@ -662,6 +665,11 @@ self.accessibilityElements = @[Window.accessibilityContainer];
 		volatile int32 ShowCount = KeyboardShowCount;
 		if (ShowCount == 1)
 		{
+			bool bKeyboardSettingsChanged = self.keyboardType != KeyboardConfig.KeyboardType ||
+						self.autocorrectionType != KeyboardConfig.AutocorrectionType ||
+						self.autocapitalizationType != KeyboardConfig.AutocapitalizationType ||
+						self.secureTextEntry != KeyboardConfig.bSecureTextEntry;
+			
 			self.keyboardType = KeyboardConfig.KeyboardType;
 			self.autocorrectionType = KeyboardConfig.AutocorrectionType;
 			self.autocapitalizationType = KeyboardConfig.AutocapitalizationType;
@@ -669,12 +677,15 @@ self.accessibilityElements = @[Window.accessibilityContainer];
 		
 			// Remember the setting
 			bSendEscapeOnClose = bInSendEscapeOnClose;
-		
+			
 			// Dismiss the existing keyboard, if one exists, so the style can be overridden.
-			[self endEditing:YES];
+			if (bKeyboardSettingsChanged || !bIsUsingIntegratedKeyboard)
+			{
+				[self endEditing:YES];
+			}
 			[self becomeFirstResponder];
-            
-            FIOSInputInterface::SetKeyboardInhibited(true);
+			
+			FIOSInputInterface::SetKeyboardInhibited(true);
 		}
 		
 		FPlatformAtomics::InterlockedDecrement(&KeyboardShowCount);
