Fixed IME being dismissed and crashing when a new window opens during text input. Call AssociateFocus to setup the new window instead of setting the focus on the existing document.

#rnx

diff --git a/Engine/Source/Runtime/ApplicationCore/Private/Windows/WindowsTextInputMethodSystem.cpp b/Engine/Source/Runtime/ApplicationCore/Private/Windows/WindowsTextInputMethodSystem.cpp
--- a/Engine/Source/Runtime/ApplicationCore/Private/Windows/WindowsTextInputMethodSystem.cpp
+++ b/Engine/Source/Runtime/ApplicationCore/Private/Windows/WindowsTextInputMethodSystem.cpp
@@ -701,6 +701,8 @@ void FWindowsTextInputMethodSystem::ClearStaleWindowHandles()
 
 void FWindowsTextInputMethodSystem::ApplyDefaults(const TSharedRef<FGenericWindow>& InWindow)
 {
+	UE_LOG(LogWindowsTextInputMethodSystem, Verbose, TEXT("Apply Defaults..."));
+
 	SCOPE_CYCLE_COUNTER(STAT_IMEWindowsApplyDefaults);
 	ClearStaleWindowHandles();
 	KnownWindows.Add(InWindow);
@@ -712,14 +714,14 @@ void FWindowsTextInputMethodSystem::ApplyDefaults(const TSharedRef<FGenericWindo
 	// In that case, we set the window to use the same IME context as is currently active. If the window actually takes focus
 	// away from the currently active IME context, then that will be taken care of in DeactivateContext, and all windows using the
 	// IME context will be disabled
-	ITfDocumentMgr* TSFDocumentManagerToSet = nullptr;
+	ITfDocumentMgr* TSFDocumentManagerToSet = TSFDisabledDocumentManager;
 	HIMC IMMContextToSet = nullptr;
 	if(ActiveContext.IsValid())
 	{
 		const HRESULT Result = TSFThreadManager->GetFocus(&TSFDocumentManagerToSet);
 		if (FAILED(Result))
 		{
-			TSFDocumentManagerToSet = nullptr;
+			TSFDocumentManagerToSet = TSFDisabledDocumentManager;
 
 			TCHAR ErrorMsg[1024];
 			FPlatformMisc::GetSystemErrorMessage(ErrorMsg, 1024, Result);
@@ -730,14 +732,13 @@ void FWindowsTextInputMethodSystem::ApplyDefaults(const TSharedRef<FGenericWindo
 	}
 
 	// TSF Implementation
-	if(TSFDocumentManagerToSet)
-	{
-		TSFThreadManager->SetFocus(TSFDocumentManagerToSet);
-	}
-	else
+	ITfDocumentMgr* Unused;
+	const HRESULT Result = TSFThreadManager->AssociateFocus(Hwnd, TSFDocumentManagerToSet, &Unused);
+	if (FAILED(Result))
 	{
-		ITfDocumentMgr* Unused;
-		TSFThreadManager->AssociateFocus(Hwnd, TSFDisabledDocumentManager, &Unused);
+		TCHAR ErrorMsg[1024];
+		FPlatformMisc::GetSystemErrorMessage(ErrorMsg, 1024, Result);
+		UE_LOG(LogWindowsTextInputMethodSystem, Error, TEXT("Couldn't associate focus between window and TSF document manager, %s (0x%08x)"), ErrorMsg, Result);
 	}
 
 	// IMM Implementation
