iOS config rules support
#tests local client, QA run on iOS clients.

diff --git a/Engine/Plugins/Online/Android/AndroidFetchBackgroundDownload/Source/AndroidFetchBackgroundDownload/Private/AndroidPlatformBackgroundHttpManager.cpp b/Engine/Plugins/Online/Android/AndroidFetchBackgroundDownload/Source/AndroidFetchBackgroundDownload/Private/AndroidPlatformBackgroundHttpManager.cpp
--- a/Engine/Plugins/Online/Android/AndroidFetchBackgroundDownload/Source/AndroidFetchBackgroundDownload/Private/AndroidPlatformBackgroundHttpManager.cpp
+++ b/Engine/Plugins/Online/Android/AndroidFetchBackgroundDownload/Source/AndroidFetchBackgroundDownload/Private/AndroidPlatformBackgroundHttpManager.cpp
@@ -996,7 +996,7 @@ void FAndroidPlatformBackgroundHttpManager::FJavaClassInfo::Initialize()
 
 FAndroidPlatformBackgroundHttpManager::EAndroidBackgroundDownloadConfigRulesSetting FAndroidPlatformBackgroundHttpManager::GetAndroidBackgroundDownloadConfigRulesSetting()
 {
-	FString* SettingStringValue = FAndroidMisc::GetConfigRulesVariable(AndroidBackgroundDownloadConfigRulesSettingKey);
+	const FString* SettingStringValue = FPlatformMisc::GetConfigRulesVariable(AndroidBackgroundDownloadConfigRulesSettingKey);
 	if (nullptr == SettingStringValue)
 	{
 		return EAndroidBackgroundDownloadConfigRulesSetting::Unknown;
diff --git a/Engine/Plugins/Runtime/AndroidDeviceProfileSelector/Source/AndroidDeviceProfileSelector/Private/AndroidDeviceProfileSelector.cpp b/Engine/Plugins/Runtime/AndroidDeviceProfileSelector/Source/AndroidDeviceProfileSelector/Private/AndroidDeviceProfileSelector.cpp
--- a/Engine/Plugins/Runtime/AndroidDeviceProfileSelector/Source/AndroidDeviceProfileSelector/Private/AndroidDeviceProfileSelector.cpp
+++ b/Engine/Plugins/Runtime/AndroidDeviceProfileSelector/Source/AndroidDeviceProfileSelector/Private/AndroidDeviceProfileSelector.cpp
@@ -142,7 +142,7 @@ FString FAndroidDeviceProfileSelector::FindMatchingProfile(const FString& Fallba
 						MatchString = VariableValueMatchString;
 #if PLATFORM_ANDROID
 						// TODO: Do this when running device selection with editor builds.
-						if (FString* ConfigRuleVar = FAndroidMisc::GetConfigRulesVariable(ConfigRuleVarName))
+						if (const FString* ConfigRuleVar = FPlatformMisc::GetConfigRulesVariable(ConfigRuleVarName))
 						{
 							ConfigRuleString = *ConfigRuleVar;
 						}
diff --git a/Engine/Plugins/Runtime/AndroidDeviceProfileSelector/Source/AndroidDeviceProfileSelectorRuntime/Private/AndroidDeviceProfileSelectorRuntimeModule.cpp b/Engine/Plugins/Runtime/AndroidDeviceProfileSelector/Source/AndroidDeviceProfileSelectorRuntime/Private/AndroidDeviceProfileSelectorRuntimeModule.cpp
--- a/Engine/Plugins/Runtime/AndroidDeviceProfileSelector/Source/AndroidDeviceProfileSelectorRuntime/Private/AndroidDeviceProfileSelectorRuntimeModule.cpp
+++ b/Engine/Plugins/Runtime/AndroidDeviceProfileSelector/Source/AndroidDeviceProfileSelectorRuntime/Private/AndroidDeviceProfileSelectorRuntimeModule.cpp
@@ -31,7 +31,7 @@ static const TMap<FName,FString>& GetDeviceSelectorParams()
 		auto GetParam = [](const FString& DefaultParam, const TCHAR* ConfRuleVarName)
 		{
 #if PLATFORM_ANDROID
-			if (FString* ConfRuleVarValue = FAndroidMisc::GetConfigRulesVariable(ConfRuleVarName))
+			if (const FString* ConfRuleVarValue = FPlatformMisc::GetConfigRulesVariable(ConfRuleVarName))
 			{
 				return *ConfRuleVarValue;
 			}
@@ -51,7 +51,7 @@ static const TMap<FName,FString>& GetDeviceSelectorParams()
 		// which on Android has a different rounding algo. See GenericPlatformMemory::GetMemorySizeBucket.
 		uint64 MemoryBucketRoundingAddition = 384;
 #if PLATFORM_ANDROID
-		if (FString* MemoryBucketRoundingAdditionVar = FAndroidMisc::GetConfigRulesVariable(TEXT("MemoryBucketRoundingAddition")))
+		if (const FString* MemoryBucketRoundingAdditionVar = FPlatformMisc::GetConfigRulesVariable(TEXT("MemoryBucketRoundingAddition")))
 		{
 			MemoryBucketRoundingAddition = FCString::Atoi64(**MemoryBucketRoundingAdditionVar);
 		}
@@ -81,7 +81,7 @@ static const TMap<FName,FString>& GetDeviceSelectorParams()
 
 #if PLATFORM_ANDROID
 		// allow ConfigRules to override cvars first
-		const TMap<FString, FString>& ConfigRules = FAndroidMisc::GetConfigRulesTMap();
+		const TMap<FString, FString>& ConfigRules = FPlatformMisc::GetConfigRuleVars();
 		for (const TPair<FString, FString>& Pair : ConfigRules)
 		{
 			const FString& VariableName = Pair.Key;
@@ -123,7 +123,7 @@ FString const FAndroidDeviceProfileSelectorRuntimeModule::GetRuntimeDeviceProfil
 		CheckForJavaSurfaceViewWorkaround(DeviceMake, DeviceModel);
 
 		// Use override from ConfigRules if set
-		FString* ConfigProfile = FAndroidMisc::GetConfigRulesVariable(TEXT("Profile"));
+		const FString* ConfigProfile = FPlatformMisc::GetConfigRulesVariable(TEXT("Profile"));
 		if (ConfigProfile != nullptr)
 		{
 			ProfileName = *ConfigProfile;
diff --git a/Engine/Source/Programs/UnrealBuildTool/ProjectFiles/Xcode/XcodeProject.File.cs b/Engine/Source/Programs/UnrealBuildTool/ProjectFiles/Xcode/XcodeProject.File.cs
--- a/Engine/Source/Programs/UnrealBuildTool/ProjectFiles/Xcode/XcodeProject.File.cs
+++ b/Engine/Source/Programs/UnrealBuildTool/ProjectFiles/Xcode/XcodeProject.File.cs
@@ -127,6 +127,7 @@ namespace UnrealBuildTool.XcodeProjectXcconfig
 			".png",
 			".xcprivacy",
 			".plist",
+			".bin",
 			// no extensions will be folders, like "cookeddata" that we want to copy as a resource
 			"",
 		};
diff --git a/Engine/Source/Runtime/ApplicationCore/Private/Android/AndroidApplication.cpp b/Engine/Source/Runtime/ApplicationCore/Private/Android/AndroidApplication.cpp
--- a/Engine/Source/Runtime/ApplicationCore/Private/Android/AndroidApplication.cpp
+++ b/Engine/Source/Runtime/ApplicationCore/Private/Android/AndroidApplication.cpp
@@ -184,7 +184,7 @@ void FDisplayMetrics::RebuildDisplayMetrics( FDisplayMetrics& OutDisplayMetrics
 	FVector4 SafeZoneRect = FAndroidWindow::GetSafezone(bIsPortrait);
 
 	// ConfigRules values override values from device
-	if (FString* SafeZoneVar = FAndroidMisc::GetConfigRulesVariable(bIsPortrait ? TEXT("SafeZone_Portrait") : TEXT("SafeZone_Landscape")))
+	if (const FString* SafeZoneVar = FPlatformMisc::GetConfigRulesVariable(bIsPortrait ? TEXT("SafeZone_Portrait") : TEXT("SafeZone_Landscape")))
 	{
 		TArray<FString> ZoneVector;
 		ZoneVector.Reserve(4);
diff --git a/Engine/Source/Runtime/ApplicationCore/Private/IOS/IOSConfigRules.cpp b/Engine/Source/Runtime/ApplicationCore/Private/IOS/IOSConfigRules.cpp
new file mode 100644
--- /dev/null
+++ b/Engine/Source/Runtime/ApplicationCore/Private/IOS/IOSConfigRules.cpp
@@ -0,0 +1,874 @@
+// Copyright Epic Games, Inc. All Rights Reserved.
+
+#include "IOS/IOSConfigRules.h"
+#include "Misc/FileHelper.h"
+#include "zlib.h"
+#include <CommonCrypto/CommonCryptor.h>
+#include <CommonCrypto/CommonKeyDerivation.h>
+#include "String/ParseLines.h"
+#include "String/ParseTokens.h"
+#include "Containers/StringConv.h"
+#include "Misc/ByteSwap.h"
+#include "Containers/Map.h"
+#include "Misc/Compression.h"
+#include "String/RemoveFrom.h"
+#include "Internationalization/Regex.h"
+
+DEFINE_LOG_CATEGORY_STATIC(LogConfigRules, Log, All);
+
+#define USE_ODS_LOGGING 1
+
+#if USE_ODS_LOGGING
+	#define CONF_LOG(A,B, ...) FPlatformMisc::LowLevelOutputDebugStringf(__VA_ARGS__);
+#else
+	#define CONF_LOG(...) UE_LOG(__VA_ARGS__);
+#endif
+
+@implementation iOSConfigRuleProviders
+static NSMutableArray* ConfigRuleProviders = nil;
++ (void)load
+{
+	ConfigRuleProviders = [[NSMutableArray alloc]init];
+}
+
++ (void)registerRuleProvider:(NSObject<iOSConfigRuleProvider>*)newProvider
+{
+	[ConfigRuleProviders addObject: newProvider];	
+}
+
++ (NSArray *)getRuleProviders
+{
+	return [ConfigRuleProviders copy];
+}
+
++ (void)empty
+{
+	[ConfigRuleProviders removeAllObjects];
+}
+@end
+
+class FConfigRules
+{
+public:
+
+    static const int ExpectedConfRuleSig = 0x39d8;
+
+    FString Path;
+	FString Key;
+	FString IV;
+
+	uint32 Version = -1;
+	uint32 UncompressedSize = 0;
+    TArray<uint8> FileBytes;
+
+#pragma pack(push)
+#pragma pack(1)
+    struct FConfHeader
+    {
+        uint16 Sig;
+        int32 Version;
+        int32 UncompressedSize;
+    };
+#pragma pack(pop)
+    
+    FConfigRules()
+    {
+       
+    }
+    
+    int OpenAndGetVersionNumber(const FString& PathIN, const FString& KeyIn, const FString& IVIn)
+    {
+        check(Path.IsEmpty());
+        Path = PathIN;
+		Key = KeyIn;
+		IV = IVIn;
+		
+        FFileHelper::LoadFileToArray(FileBytes, *Path);
+
+		if (FileBytes.Num() > sizeof(FConfHeader))
+        {
+            FConfHeader ConfHeader = *((FConfHeader*)FileBytes.GetData());
+            if (ByteSwap(ConfHeader.Sig) == ExpectedConfRuleSig)
+            {
+                Version = ByteSwap(ConfHeader.Version);
+                UncompressedSize = ByteSwap(ConfHeader.UncompressedSize);
+            }
+        }
+		else
+		{
+			CONF_LOG(LogConfigRules, Log, TEXT("ConfigRules: %s was not found."), *Path);	
+		}
+        return Version;
+    }
+	
+	TConstArrayView<uint8> GetData() const
+	{
+		return TConstArrayView<uint8>(FileBytes.GetData() + sizeof(FConfHeader), FileBytes.Num() - sizeof(FConfHeader));
+	}
+};
+
+
+bool Decrypt(const TConstArrayView<uint8>& DataIn, TArray<uint8>& DataOut, FString& key, FString& iv )
+{
+	check(!DataOut.IsEmpty());
+	
+	FTCHARToUTF8 UTF8Key(*key);
+	
+	uint8 	GeneratedKey[16] = {0};
+	{
+		static const uint8 salt[] = { 0x23, 0x71, 0xd3, 0xa3, 0x30, 0x71, 0x63, 0xe3};
+		uint    rounds  = 1000;
+		uint    keySize = kCCKeySizeAES128;
+
+		
+		CCKeyDerivationPBKDF(CCPBKDFAlgorithm(kCCPBKDF2),
+							 UTF8Key.Get(),
+							 UTF8Key.Length(),
+							 salt,
+							 sizeof(salt),
+							 kCCPRFHmacAlgSHA1,
+							 rounds,
+							 GeneratedKey,
+							 sizeof(GeneratedKey));
+	}
+	
+	size_t WrittenBytes = 0;
+	
+	char ivPtr[kCCBlockSizeAES128 + 1];
+	FMemory::Memzero(ivPtr);
+	if (!iv.IsEmpty())
+	{
+		FTCHARToUTF8 UTF8iv(*iv);
+		if (UTF8iv.Length() + 1 != sizeof(ivPtr))
+		{
+			return false;
+		}
+		FMemory::Memcpy(ivPtr, UTF8iv.Get(), sizeof(ivPtr));
+	}
+
+	CCCryptorStatus ccStatus   = kCCSuccess;
+	size_t          cryptBytes = 0;
+	
+	ccStatus = CCCrypt(kCCDecrypt,
+						kCCAlgorithmAES128,
+						kCCOptionECBMode | kCCOptionPKCS7Padding,
+                         GeneratedKey,
+                         kCCKeySizeAES128,
+                         ivPtr,
+                         DataIn.GetData(),
+                         DataIn.Num(),
+                         DataOut.GetData(),
+                         DataOut.Num(),
+                         &cryptBytes);
+    
+	if (ccStatus == kCCSuccess)
+	{
+		DataOut.SetNum(cryptBytes);
+	}
+	
+	return ccStatus == kCCSuccess;
+}
+
+namespace ConfigRules
+{
+	TMap<FString, FString> ProcessConfigRules(FConfigRules& Rules, TMap<FString,FString>&& PredefinedVariables);
+};
+
+bool FIOSConfigRules::Init(TMap<FString,FString>&& PredefinedVariables)
+{
+	InitRules();
+
+	// load the configrule with the highest version.
+	int FoundVersion = -1;
+	FConfigRules SelectedRules;
+	for (int i = 0; i < ConfigRulesParams.Num(); i++)
+	{
+		FString Filename = ConfigRulesParams[i].Path;
+		FConfigRules TestRules;
+		int TestVersion = TestRules.OpenAndGetVersionNumber(Filename, ConfigRulesParams[i].Key, ConfigRulesParams[i].IV);
+		if (TestVersion > FoundVersion)
+		{
+			SelectedRules = MoveTemp(TestRules);
+			FoundVersion = TestVersion;
+		}
+	}
+
+	if (FoundVersion > -1)
+	{
+		ConfigRuleVariablesMap = ConfigRules::ProcessConfigRules(SelectedRules, MoveTemp(PredefinedVariables));
+	}
+	else
+	{
+		ConfigRuleVariablesMap = PredefinedVariables;
+	}
+	return FoundVersion > -1;
+}
+
+void FIOSConfigRules::InitRules()
+{
+    check(ConfigRulesParams.IsEmpty());
+   
+	for (NSObject<iOSConfigRuleProvider>* Provider in [iOSConfigRuleProviders getRuleProviders])
+	{
+		for (NSArray* Params in [Provider getRuleData])
+		{
+			if([Params count] == 3)
+			{
+				const char * PathStr = [[Params objectAtIndex:0] UTF8String];
+				const char * KeyStr = [[Params objectAtIndex:1] UTF8String];
+				const char * PathIV = [[Params objectAtIndex:2] UTF8String];
+				
+				FConfigRuleParams ConfigRuleParams;
+				ConfigRuleParams.Path = PathStr;
+				ConfigRuleParams.Key = KeyStr;
+				ConfigRuleParams.IV = PathIV;
+				
+				ConfigRuleParams.Path.ReplaceInline(TEXT("[[cache]]"), *FString([NSSearchPathForDirectoriesInDomains(NSCachesDirectory, NSUserDomainMask, YES) objectAtIndex:0]));
+				ConfigRulesParams.Add(ConfigRuleParams);
+			}
+		}
+	}
+	
+	[iOSConfigRuleProviders empty];
+	
+	if(ConfigRulesParams.IsEmpty())
+	{
+		FConfigRuleParams ConfigRuleParams;
+		ConfigRuleParams.Path = TEXT("~/configrules");		
+		ConfigRulesParams.Add(ConfigRuleParams);
+	}
+}
+
+namespace ConfigRules
+{
+	TMap<FString, FString> ParseConfigRules(TConstArrayView<uint8> ConfigRulesData, TMap<FString,FString>&& PredefinedVariables);
+
+	TMap<FString, FString> ProcessConfigRules(FConfigRules& Rules, TMap<FString,FString>&& PredefinedVariables)
+	{
+		TArray<uint8> DecryptBytes;
+		bool bSuccess = true;
+		if (!Rules.Key.IsEmpty())
+		{
+			DecryptBytes.SetNum(Rules.FileBytes.Num());
+			bSuccess = Decrypt(Rules.GetData(), DecryptBytes, Rules.Key, Rules.IV);
+		}
+
+		TArray<uint8> UncompressedBytes;
+		if (bSuccess)
+		{
+			const TConstArrayView<uint8>& SourceBytes = Rules.Key.IsEmpty() ? Rules.GetData() : DecryptBytes;
+			UncompressedBytes.SetNum(Rules.UncompressedSize);
+			bSuccess = FCompression::UncompressMemory(NAME_Zlib, UncompressedBytes.GetData(), UncompressedBytes.Num(), SourceBytes.GetData(), SourceBytes.Num());
+		}
+
+		if (bSuccess)
+		{
+			return ParseConfigRules(UncompressedBytes, MoveTemp(PredefinedVariables));
+		}
+		else
+		{
+			CONF_LOG(LogConfigRules, Error, TEXT("ConfigRules: file read failed for %s!"), *Rules.Path);
+		}
+		return PredefinedVariables;
+	}
+
+	static FStringView RemoveSurrounds(FStringView Input, FStringView Entry, FStringView Exit)
+	{
+		if (Input.Len() >= 2 && Entry.Len() > 0 && Entry.Len() == Exit.Len())
+		{
+			FStringView Ret = UE::String::RemoveFromEnd(UE::String::RemoveFromStart(Input, Entry), Exit);
+			if (Ret.Len() == (Input.Len() - (Entry.Len() * 2)))
+			{
+				return Ret;
+			}
+		}
+		return Input;
+	}
+
+	static FStringView SubStr(FStringView Input, int32 StartIdx, int32 EndIdx)
+	{
+		return Input.Left(EndIdx).RightChop(StartIdx);
+	};
+
+	static int32 IndexOf(const FString& Input, FStringView Str, int FromIndex = 0)
+	{
+		return Input.Find(Str, ESearchCase::CaseSensitive, ESearchDir::FromStart, FromIndex);		
+	}
+
+	// returns list of strings separated by Split using Entry/Exit as pairing sets (ex. "(" and ")").  The Entry and Exit characters need to be in same order
+	static TArray<FStringView> ParseSegments(FStringView Input, FStringView Split, FStringView Entry, FStringView Exit)
+	{
+		TArray<FStringView> Output;
+		TArray<int32> EntryStack;
+
+		int StartIndex = 0;
+		int ScanIndex = 0;
+		int ExitIndex = -1;
+		int InputLength = Input.Len();
+
+		while (ScanIndex < InputLength)
+		{
+			FStringView Scan = SubStr(Input, ScanIndex, ScanIndex + 1);
+
+			if (Scan.Equals(Split) && EntryStack.Num() == 0)
+			{
+				Output.Add(SubStr(Input, StartIndex, ScanIndex).TrimStartAndEnd());
+				ScanIndex++;
+				StartIndex = ScanIndex;
+				continue;
+			}
+			ScanIndex++;
+
+			if (Scan.Equals(TEXT("\\")))
+			{
+				ScanIndex++;
+				continue;
+			}
+
+			if (EntryStack.Num() > 0 && Exit.Find(Scan) == ExitIndex)
+			{
+				int StackLength = EntryStack.Num() - 1;
+				EntryStack.RemoveAt(StackLength);
+				ExitIndex = StackLength > 0 ? EntryStack[StackLength - 1] : -1;
+				continue;
+			}
+
+			int EntryIndex = Entry.Find(Scan);
+			if (EntryIndex >= 0)
+			{
+				EntryStack.Add(EntryIndex);
+				ExitIndex = EntryIndex;
+				continue;
+			}
+		}
+		if (StartIndex < InputLength)
+		{
+			Output.Add(Input.RightChop(StartIndex).TrimStartAndEnd());
+		}
+
+		return Output;
+	}
+
+	FString ExpandVariables(TMap<FString, FString>& Variables, FStringView Input)
+	{
+		FString Result(Input);
+		int Idx;
+		for (Idx = IndexOf(Result, TEXT("$(")); Idx != -1; Idx = IndexOf(Result, TEXT("$("), Idx))
+		{
+			// Find the end of the variable name
+			int EndIdx = IndexOf(Result, TEXT(")"), Idx + 2);
+			if (EndIdx == -1)
+			{
+				break;
+			}
+
+			// Extract the variable name from the string
+			FStringView VarKey = SubStr(Result, Idx + 2, EndIdx);
+
+			// Find the value for it if it exists
+			if (FString* VarValue = Variables.Find(FString(VarKey)))
+			{
+				// Replace the variable
+				Result = SubStr(Result, 0, Idx) + (*VarValue) + Result.RightChop( EndIdx + 1);
+			}
+			else
+			{
+				// or skip past it
+				Idx = EndIdx + 1;
+				continue;
+			}
+
+		}
+		return Result;
+	}
+
+	/////////////////////
+	static bool EvaluateConditions(TMap<FString, FString>& Variables, TArray<FStringView>& Conditions, FString& PreviousRegexMatch);
+	static bool EvaluateConditions(TMap<FString, FString>& Variables, TArray<FStringView>& Conditions)
+	{
+		FString PrevRegEx(TEXT(""));
+		return EvaluateConditions(Variables, Conditions, PrevRegEx);
+	}
+
+	/////////////////////
+	TMap<FString, FString> ParseConfigRules(TConstArrayView<uint8> ConfigRulesData, TMap<FString,FString>&& PredefinedVariables)
+	{
+		enum class EConfRuleState : uint8
+		{
+			Run = 0,
+			ExecTrue,
+			FindElse,
+			ExecFalse,
+			FindEnd
+		};
+
+		FAnsiStringView ConfigRulesAnsi((const char*)ConfigRulesData.GetData(), ConfigRulesData.Num());
+		// TODO: we should be able to parse this entire file without conversion to tchar.
+		// Some functions such as ParseLines, LexFromString etc. expect tchar.
+		FString ConfigRules = StringCast<TCHAR>(ConfigRulesAnsi.GetData()).Get();
+		TMap<FString, FString> ConfigRuleVars;
+		ConfigRuleVars.Append(MoveTemp(PredefinedVariables));
+		TArray<EConfRuleState> StateStack;
+
+		EConfRuleState CurrentState = EConfRuleState::Run;
+		int NestDepth = 0;
+		bool bAbort = false;
+		UE::String::ParseLines(ConfigRules,
+			[&ConfigRuleVars, &CurrentState, &StateStack, &NestDepth, &bAbort](FStringView Line)
+			{
+				FStringView line = Line.TrimStartAndEnd();
+				if (line.Len() < 1 || bAbort)
+				{
+					return;
+				}
+				if (line.StartsWith(TEXT("//"), ESearchCase::Type::CaseSensitive) || line.StartsWith(TEXT(";"), ESearchCase::Type::CaseSensitive))
+				{
+					if (line.StartsWith(TEXT("// version:")))
+					{
+						int configRulesVersion = 0;
+						LexFromString(configRulesVersion, *FString(line.RightChop(11)));
+						ConfigRuleVars.Add(TEXT("configRulesVersion"), FString::FromInt(configRulesVersion));
+						CONF_LOG(LogConfigRules, Log, TEXT("ConfigRules version: %d"), configRulesVersion);
+					}
+					return;
+				}
+
+				// look for Command
+				int index;
+				if (!line.FindChar(':', index))
+				{
+					return;
+				}
+				FStringView Command = line.Left(index).TrimStartAndEnd();
+				line = line.RightChop(index + 1).TrimStartAndEnd();
+
+				// handle states
+				switch (CurrentState)
+				{
+					case EConfRuleState::Run:
+					{
+						if (Command.Equals(TEXT("else")) || Command.Equals(TEXT("elseif")) || Command.Equals(TEXT("endif")))
+						{
+							CONF_LOG(LogConfigRules, Error, TEXT("ConfigRules: unexpected %s encountered!"), *FString(Command));
+						}
+						break;
+					}
+
+					case EConfRuleState::ExecTrue:
+					{
+						if (Command.Equals(TEXT("else")))
+						{
+							CurrentState = EConfRuleState::FindEnd;
+						}
+						else if (Command.Equals(TEXT("endif")))
+						{
+							CurrentState = StateStack.Pop();
+						}
+						return;
+					}
+
+					case EConfRuleState::FindElse:
+					{
+						if (Command.Equals(TEXT("if")))
+						{
+							NestDepth++;
+							return;
+						}
+						if (NestDepth > 0)
+						{
+							if (Command.Equals(TEXT("endif")))
+							{
+								NestDepth--;
+							}
+							return;
+						}
+						if (Command.Equals(TEXT("endif")))
+						{
+							CurrentState = StateStack.Pop();
+							return;
+						}
+						if (Command.Equals(TEXT("else")))
+						{
+							CurrentState = EConfRuleState::ExecFalse;
+							return;
+						}
+						if (Command.Equals(TEXT("elseif")))
+						{
+							CurrentState = EConfRuleState::FindEnd;
+
+							TArray<FStringView> Conditions;
+							UE::String::ParseTokensMultiple(line, TConstArrayView<TCHAR>({ ',','(',')' }), Conditions);
+							if (Conditions.Num() > 0)
+							{
+								bool bConditionTrue = EvaluateConditions(ConfigRuleVars, Conditions);
+								CurrentState = bConditionTrue ? EConfRuleState::ExecTrue : EConfRuleState::FindElse;
+							}
+						}
+						return;
+					}
+					case EConfRuleState::ExecFalse:
+					{
+						if (Command.Equals(TEXT("endif")))
+						{
+							CurrentState = StateStack.Pop();
+						}
+						if (Command.Equals(TEXT("else")) || Command.Equals(TEXT("elseif")))
+						{
+							CONF_LOG(LogConfigRules, Error, TEXT("ConfigRules: unexpected %s while handling false condition!"), *FString(Command));
+						}
+						return;
+					}
+					case EConfRuleState::FindEnd:
+					{
+						if (Command.Equals(TEXT("if")))
+						{
+							NestDepth++;
+						}
+						else if (Command.Equals(TEXT("endif")))
+						{
+							if (NestDepth > 0)
+							{
+								NestDepth--;
+							}
+							else
+							{
+								CurrentState = StateStack.Pop();
+							}
+						}
+						return;
+					}
+					default:
+						CONF_LOG(LogConfigRules, Error, TEXT("ConfigRules: unknown state!"));
+						return;
+				}
+
+				// handle commands
+				if (Command.Equals(TEXT("set")))
+				{
+					// set:(a=b[,c=d,...])
+					TArray<FStringView> Sets = ParseSegments(RemoveSurrounds(line, TEXT("("), TEXT(")")), TEXT(","), TEXT("(\")"), TEXT(")\")"));
+					for (FStringView Assignment : Sets)
+					{
+						TArray<FStringView> KeyValue = ParseSegments(Assignment, TEXT("="), TEXT("\""), TEXT("\""));
+						if (KeyValue.Num() == 2)
+						{
+							FStringView Key = RemoveSurrounds(KeyValue[0], TEXT("\""), TEXT("\""));
+							FString Value = ExpandVariables(ConfigRuleVars, RemoveSurrounds(KeyValue[1], TEXT("\""), TEXT("\"")));
+							check(!Key.IsEmpty());
+							if (Key.StartsWith(TEXT("APPEND_")))
+							{
+								Key = Key.RightChop(7);
+								if (FString* Found = ConfigRuleVars.Find(FString(Key)))
+								{
+									Value = (*Found) + Value;
+								}
+								ConfigRuleVars.Add(FString(Key), Value);
+							}
+							else
+							{
+								ConfigRuleVars.Add(FString(Key), Value);
+							}
+						}
+					}
+				}
+				else if (Command.Equals(TEXT("clear")))
+				{
+					// clear:(a[,b,...])
+					TArray<FStringView> Sets = ParseSegments(RemoveSurrounds(Line, TEXT("("), TEXT(")")), TEXT(","), TEXT("(\""), TEXT(")\""));
+					for (FStringView Key : Sets)
+					{
+						ConfigRuleVars.Remove(FString(RemoveSurrounds(Key, TEXT("\""), TEXT("\""))));
+					}
+				}
+				else if (Command.Equals(TEXT("chipset")))
+				{
+					// not supported here.
+				}
+				else if (Command.Equals(TEXT("if")))
+				{
+					// if:(SourceType=SRC_DeviceMake,CompareType=CMP_Equal,MatchString="samsung")
+					// ... commands for true for all conditions
+					// elseif:(SourceType=SRC_DeviceMake,CompareType=CMP_Equal,MatchString="Google")
+					// ... commands for true for all conditions
+					// else:
+					// ... commands for false for any condition
+					// end:
+					StateStack.Push(CurrentState);
+					CurrentState = EConfRuleState::FindEnd;
+
+					TArray<FStringView> Conditions = ParseSegments(Line, TEXT(","), TEXT("(\""), TEXT(")\""));
+					if (!Conditions.IsEmpty())
+					{
+						const bool bConditionTrue = EvaluateConditions(ConfigRuleVars, Conditions);
+						CurrentState = bConditionTrue ? EConfRuleState::ExecTrue : EConfRuleState::FindElse;
+					}
+				}
+				else if (Command.Equals(TEXT("condition")))
+				{
+					// condition:((SourceType=SRC_DeviceMake,CompareType=CMP_Equal,MatchString="samsung")),(SourceType=,CompareType=,MatchString=),...]),(a=b[,c=d,...]),(a[,b,...])
+					// if all the conditions are true, execute the optional sets and/or clears
+					TArray<FStringView> ConditionAndSets = ParseSegments(line, TEXT(","), TEXT("(\""), TEXT(")\""));
+					int SetSize = ConditionAndSets.Num();
+					if (SetSize == 2 || SetSize == 3)
+					{
+						TArray<FStringView> Conditions = ParseSegments(RemoveSurrounds(ConditionAndSets[0], TEXT("("), TEXT(")")), TEXT(","), TEXT("(\""), TEXT(")\""));
+						TArray<FStringView> Sets = ParseSegments(RemoveSurrounds(ConditionAndSets[1], TEXT("("), TEXT(")")), TEXT(","), TEXT("(\""), TEXT(")\""));
+						TArray<FStringView> Clears = (SetSize == 3) ? ParseSegments(RemoveSurrounds(ConditionAndSets[2], TEXT("("), TEXT(")")), TEXT(","), TEXT("(\""), TEXT(")\"")) : TArray<FStringView>();
+
+						bool bConditionTrue = EvaluateConditions(ConfigRuleVars, Conditions);
+
+						if (bConditionTrue)
+						{
+							// run the sets
+							for (FStringView Assignment : Sets)
+							{
+								TArray<FStringView> KeyValue = ParseSegments(Assignment, TEXT("="), TEXT("\""), TEXT("\""));
+								if (KeyValue.Num() == 2)
+								{
+									FStringView Key = RemoveSurrounds(KeyValue[0], TEXT("\""), TEXT("\""));
+									FString Value = ExpandVariables(ConfigRuleVars, RemoveSurrounds(KeyValue[1], TEXT("\""), TEXT("\"")));
+									check(!Key.IsEmpty());
+									if (Key.StartsWith(TEXT("APPEND_")))
+									{
+										Key = Key.RightChop(7);
+										if (FString* Found = ConfigRuleVars.Find(FString(Key)))
+										{
+											Value = (*Found) + Value;
+										}
+										ConfigRuleVars.Add(FString(Key), Value);
+									}
+									else
+									{
+										ConfigRuleVars.Add(FString(Key), Value);
+									}
+								}
+							}
+
+							// run the clears
+							for (FStringView Key : Clears)
+							{
+								ConfigRuleVars.Remove(FString(RemoveSurrounds(Key, TEXT("\""), TEXT("\""))));
+							}
+						}
+					}
+				}
+				// see if log message requested
+				static const FString LogStr(TEXT("log"));
+				if (FString* Found = ConfigRuleVars.Find(LogStr))
+				{
+					CONF_LOG(LogConfigRules, Log, TEXT("ConfigRules log output:\n %s"), **Found);
+					ConfigRuleVars.Remove(LogStr);
+				}
+
+				// check if requested to dump variables to the log
+				static const FString DumpVarsStr(TEXT("dumpvars"));
+				if (FString* Found = ConfigRuleVars.Find(DumpVarsStr))
+				{
+					ConfigRuleVars.Remove(DumpVarsStr);
+					CONF_LOG(LogConfigRules, Log, TEXT("ConfigRules vars:"));
+					for ( TPair<FString, FString> VarEntry: ConfigRuleVars)
+					{
+						CONF_LOG(LogConfigRules, Log, TEXT("%s = %s"), *VarEntry.Key, *VarEntry.Value);
+					}
+				}
+
+				// if there was a raised error or break, stop
+				const bool HasError = ConfigRuleVars.Contains(TEXT("error"));
+				const bool HasBreak = ConfigRuleVars.Contains(TEXT("break"));
+
+				// stop if user wants to break
+				if (HasBreak || HasError)
+				{
+					CONF_LOG(LogConfigRules, Warning, TEXT("Config rules aborting parse due to %s."), HasBreak ? TEXT("break command") : TEXT("error"));
+					bAbort = true;
+				}
+			});
+
+		return ConfigRuleVars;
+	}
+
+	static bool EvaluateConditions(TMap<FString, FString>& Variables, TArray<FStringView>& Conditions, FString& PreviousRegexMatch)
+	{
+		bool bConditionTrue = true;
+		for (FStringView Condition : Conditions)
+		{
+			FString SourceType = "";
+			FString CompareType = "";
+			FString MatchString = "";
+
+			// deal with Condition group (src,cmp,match)
+			TArray<FStringView> groups = ParseSegments(RemoveSurrounds(Condition, TEXT("("), TEXT(")")), TEXT(","), TEXT("\""), TEXT("\""));
+			for (FStringView group : groups)
+			{
+				TArray<FStringView> KeyValue = ParseSegments(group, TEXT("="), TEXT("\""), TEXT("\""));
+				if (KeyValue.Num() == 2)
+				{
+					FStringView Key = RemoveSurrounds(KeyValue[0], TEXT("\""), TEXT("\""));
+					FStringView value = RemoveSurrounds(KeyValue[1], TEXT("\""), TEXT("\""));
+
+					if (Key.Equals(TEXT("SourceType")))
+					{
+						SourceType = value;
+					}
+					else if (Key.Equals(TEXT("CompareType")))
+					{
+						CompareType = value;
+					}
+					else if (Key.Equals(TEXT("MatchString")))
+					{
+						MatchString = value;
+					}
+				}
+			}
+
+			FString Source = "";
+			if (SourceType.Equals(TEXT("SRC_PreviousRegexMatch")))
+			{
+				Source = PreviousRegexMatch;
+			}
+			else if (SourceType.Equals(TEXT("SRC_CommandLine")))
+			{
+				checkNoEntry();
+			}
+			else if (FString* Found = Variables.Find(SourceType))
+			{
+				Source = *Found;
+			}
+			else if (SourceType.Equals(TEXT("[EXIST]")))
+			{
+				Source = MatchString;
+			}
+			else
+			{
+				bConditionTrue = false;
+				break;
+			}
+
+			// apply operation
+			if (CompareType.Equals(TEXT("CMP_Exist")))
+			{
+				if (!Variables.Contains(Source))
+				{
+					bConditionTrue = false;
+					break;
+				}
+			}
+			else if (CompareType.Equals(TEXT("CMP_NotExist")))
+			{
+				if (Variables.Contains(Source))
+				{
+					bConditionTrue = false;
+					break;
+				}
+			}
+			else if (CompareType.Equals(TEXT("CMP_Equal")))
+			{
+				if (!Source.Equals(MatchString))
+				{
+					bConditionTrue = false;
+					break;
+				}
+			}
+			else if (CompareType.Equals(TEXT("CMP_NotEqual")))
+			{
+				if (Source.Equals(MatchString))
+				{
+					bConditionTrue = false;
+					break;
+				}
+			}
+			else if (CompareType.Equals(TEXT("CMP_EqualIgnore")))
+			{
+				if (!Source.ToLower().Equals(MatchString.ToLower()))
+				{
+					bConditionTrue = false;
+					break;
+				}
+			}
+			else if (CompareType.Equals(TEXT("CMP_NotEqualIgnore")))
+			{
+				if (Source.ToLower().Equals(MatchString.ToLower()))
+				{
+					bConditionTrue = false;
+					break;
+				}
+			}
+			else if (CompareType.Equals(TEXT("CMP_Regex")))
+			{
+				const FRegexPattern RegexPattern(MatchString);
+				FRegexMatcher RegexMatcher(RegexPattern, Source);
+
+				if (RegexMatcher.FindNext())
+				{
+					FString Captured = RegexMatcher.GetCaptureGroup(1);
+					PreviousRegexMatch = Captured.IsEmpty() ? RegexMatcher.GetCaptureGroup(0) : Captured;
+				}
+				else
+				{
+					bConditionTrue = false;
+					break;
+				}
+			}
+			else
+			{
+				bool bNumericOperands = true;
+				float SourceFloat = 0.0f;
+				float MatchFloat = 0.0f;
+
+				// convert source and match to float if numeric
+				bNumericOperands = LexTryParseString(SourceFloat, *Source);
+				bNumericOperands = LexTryParseString(MatchFloat, *MatchString) && bNumericOperands;
+
+				// if comparison ends with Ignore, do case-insensitive compare by converting both to lowercase
+				if (CompareType.EndsWith(TEXT("Ignore")))
+				{
+					bNumericOperands = false;
+					CompareType = SubStr(CompareType, 0, CompareType.Len() - 6);
+
+					Source = Source.ToLower();
+					MatchString = MatchString.ToLower();
+				}
+
+				if (CompareType.Equals(TEXT("CMP_Less")))
+				{
+					if ((bNumericOperands && (SourceFloat >= MatchFloat)) || (!bNumericOperands && (Source.Compare(MatchString) >= 0)))
+					{
+						bConditionTrue = false;
+						break;
+					}
+				}
+				else if (CompareType.Equals(TEXT("CMP_LessEqual")))
+				{
+					if ((bNumericOperands && (SourceFloat > MatchFloat)) || (!bNumericOperands && (Source.Compare(MatchString) > 0)))
+					{
+						bConditionTrue = false;
+						break;
+					}
+				}
+				else if (CompareType.Equals(TEXT("CMP_Greater")))
+				{
+					if ((bNumericOperands && (SourceFloat <= MatchFloat)) || (!bNumericOperands && (Source.Compare(MatchString) <= 0)))
+					{
+						bConditionTrue = false;
+						break;
+					}
+				}
+				else if (CompareType.Equals(TEXT("CMP_GreaterEqual")))
+				{
+					if ((bNumericOperands && (SourceFloat < MatchFloat)) || (!bNumericOperands && (Source.Compare(MatchString) < 0)))
+					{
+						bConditionTrue = false;
+						break;
+					}
+				}
+				else
+				{
+					bConditionTrue = false;
+					break;
+				}
+			}
+		}
+
+		return bConditionTrue;
+	}
+}
diff --git a/Engine/Source/Runtime/ApplicationCore/Public/IOS/IOSConfigRules.h b/Engine/Source/Runtime/ApplicationCore/Public/IOS/IOSConfigRules.h
new file mode 100644
--- /dev/null
+++ b/Engine/Source/Runtime/ApplicationCore/Public/IOS/IOSConfigRules.h
@@ -0,0 +1,43 @@
+// Copyright Epic Games, Inc. All Rights Reserved.
+#pragma once
+#include "CoreMinimal.h"
+
+// iOSConfigRuleProvider represents the interface used to return the info required to load multiple config files
+// instances are registered with iOSConfigRuleProviders during +load time and queried during FIOSConfigRules::InitRules
+@protocol iOSConfigRuleProvider <NSObject>
+// getRuleData returns an array of arrays, the sub arrays must contain 3 elements
+// of filename, key, iv. e.g. [["filename", "key", "IV"],["filename2", "key2", "IV2"], etc.. 
+- (NSArray *)getRuleData;
+@end
+
+// iOSConfigRuleProviders contains the list of iOSConfigRuleProviders
+// instances of iOSConfigRuleProvider are registered here to be 
+// queried during Init
+@interface iOSConfigRuleProviders : NSObject
++ (void)load;
++ (void)registerRuleProvider:(NSObject<iOSConfigRuleProvider>*)newProvider;
++ (NSArray *)getRuleProviders;
++ (void)empty;
+@end
+
+class FIOSConfigRules
+{
+	struct FConfigRuleParams
+	{
+		FString Path;
+		FString Key;
+		FString IV;
+	};
+    static inline TArray<FConfigRuleParams> ConfigRulesParams;
+	
+	static void InitRules();	
+	static inline TMap<FString, FString> ConfigRuleVariablesMap;
+	
+public:
+	// Initialize the config rules, 
+	// PredefinedVariables contains a set of K/V pairs that are used while
+	// processing the rules. They are included in the map returned by GetConfigRulesMap.
+	// returns true if a config rule file was processed.
+	static bool Init(TMap<FString, FString>&& PredefinedVariables);
+	static const TMap<FString, FString>& GetConfigRulesMap() { return ConfigRuleVariablesMap; }
+};
diff --git a/Engine/Source/Runtime/Core/Private/Android/AndroidPlatformMisc.cpp b/Engine/Source/Runtime/Core/Private/Android/AndroidPlatformMisc.cpp
--- a/Engine/Source/Runtime/Core/Private/Android/AndroidPlatformMisc.cpp
+++ b/Engine/Source/Runtime/Core/Private/Android/AndroidPlatformMisc.cpp
@@ -2529,16 +2529,11 @@ FString FAndroidMisc::GetVulkanVersion()
 	return VulkanVersionString;
 }
 
-const TMap<FString, FString>& FAndroidMisc::GetConfigRulesTMap()
+const TMap<FString, FString>& FAndroidMisc::GetConfigRuleVars()
 {
 	return ConfigRulesVariables;
 }
 
-FString* FAndroidMisc::GetConfigRulesVariable(const FString& Key)
-{
-	return ConfigRulesVariables.Find(Key);
-}
-
 bool FAndroidMisc::AllowThreadHeartBeat()
 {
 	static uint32 AllowThreadHeartBeatOnce = -1;
@@ -2978,7 +2973,7 @@ FString FAndroidMisc::GetCPUBrand()
 
 FString FAndroidMisc::GetCPUChipset()
 {
-	static FString *Chipset = FAndroidMisc::GetConfigRulesVariable(TEXT("hardware"));
+	static const FString *Chipset = FAndroidMisc::GetConfigRulesVariable(TEXT("hardware"));
 	return (Chipset == NULL) ? FGenericPlatformMisc::GetCPUChipset() : *Chipset;
 }
 
diff --git a/Engine/Source/Runtime/Core/Private/GenericPlatform/GenericPlatformMemory.cpp b/Engine/Source/Runtime/Core/Private/GenericPlatform/GenericPlatformMemory.cpp
--- a/Engine/Source/Runtime/Core/Private/GenericPlatform/GenericPlatformMemory.cpp
+++ b/Engine/Source/Runtime/Core/Private/GenericPlatform/GenericPlatformMemory.cpp
@@ -390,7 +390,7 @@ EPlatformMemorySizeBucket FGenericPlatformMemory::GetMemorySizeBucket()
 #if PLATFORM_ANDROID
 		// we don't exactly want to round up on android
 		uint64 MemoryBucketRoundingAddition = 384;
-		if (FString* MemoryBucketRoundingAdditionVar = FAndroidMisc::GetConfigRulesVariable(TEXT("MemoryBucketRoundingAddition")))
+		if (const FString* MemoryBucketRoundingAdditionVar = FPlatformMisc::GetConfigRulesVariable(TEXT("MemoryBucketRoundingAddition")))
 		{
 			MemoryBucketRoundingAddition = FCString::Atoi64(**MemoryBucketRoundingAdditionVar);
 		}
@@ -616,4 +616,4 @@ FString FGenericPlatformMemory::PrettyMemory(uint64 Memory)
 	}
 
 	return FString::Printf(TEXT("%llu%s"), Memory, Units[UnitIndex]);
-}
\ No newline at end of file
+}
diff --git a/Engine/Source/Runtime/Core/Private/GenericPlatform/GenericPlatformMisc.cpp b/Engine/Source/Runtime/Core/Private/GenericPlatform/GenericPlatformMisc.cpp
--- a/Engine/Source/Runtime/Core/Private/GenericPlatform/GenericPlatformMisc.cpp
+++ b/Engine/Source/Runtime/Core/Private/GenericPlatform/GenericPlatformMisc.cpp
@@ -2058,3 +2058,14 @@ bool FGenericPlatformMisc::RemoveNetworkListener(FDelegateHandle Handle)
 
 	return false;
 }
+
+const TMap<FString, FString>& FGenericPlatformMisc::GetConfigRuleVars()
+{
+	static TMap<FString, FString> Temp;
+	return Temp;
+}
+
+const FString* FGenericPlatformMisc::GetConfigRulesVariable(const FString& Key) 
+{
+	return FPlatformMisc::GetConfigRuleVars().Find(Key);
+}
diff --git a/Engine/Source/Runtime/Core/Private/IOS/IOSPlatformMisc.cpp b/Engine/Source/Runtime/Core/Private/IOS/IOSPlatformMisc.cpp
--- a/Engine/Source/Runtime/Core/Private/IOS/IOSPlatformMisc.cpp
+++ b/Engine/Source/Runtime/Core/Private/IOS/IOSPlatformMisc.cpp
@@ -56,6 +56,9 @@
 
 #include "Apple/PostAppleSystemHeaders.h"
 
+#include "IOSConfigRules.h"
+#include <Metal/Metal.h>
+
 #if !defined ENABLE_ADVERTISING_IDENTIFIER
 	#define ENABLE_ADVERTISING_IDENTIFIER 0
 #endif
@@ -1815,14 +1818,77 @@ static void GracefulTerminationHandler(int32 Signal, siginfo_t* Info, void* Cont
     }
 }
 
+static TMap<FString, FString> GetPredefinedConfigRuleVars()
+{
+	TMap<FString, FString> PredefinedConfigRuleVars;
+	const FPlatformMemoryConstants& MemConstants = FPlatformMemory::GetConstants();
+	uint64 PhysMemBytes = [NSProcessInfo processInfo].physicalMemory;
+	int32 totalMemoryMB = (int32)(MemConstants.TotalPhysical / 1024 / 1024);
+	int32 totalMemoryGB = (int32)(MemConstants.TotalPhysicalGB);	
+
+	FString CpuVendor = FIOSPlatformMisc::GetCPUVendor();
+	FString CPUBrand = FIOSPlatformMisc::GetCPUBrand();
+	FString ProjectVersion = FIOSPlatformMisc::GetProjectVersion();
+	FString BuildNumber = FIOSPlatformMisc::GetBuildNumber();
+	FString OSVersionLabel,OSSubVersionLabel;
+	FIOSPlatformMisc::GetOSVersions(OSVersionLabel,OSSubVersionLabel);
+	
+	PredefinedConfigRuleVars.Add(FString(TEXT("SRC_OSVersion")), OSSubVersionLabel);
+	PredefinedConfigRuleVars.Add(FString(TEXT("SRC_OSLanguage")), *FString([[NSLocale preferredLanguages] objectAtIndex:0]));
+	
+	PredefinedConfigRuleVars.Add(FString(TEXT("SRC_GPUFamily")), CpuVendor);
+	PredefinedConfigRuleVars.Add(FString(TEXT("SRC_DeviceMake")), CpuVendor);
+	PredefinedConfigRuleVars.Add(FString(TEXT("SRC_DeviceModel")), CPUBrand);
+	PredefinedConfigRuleVars.Add(FString(TEXT("SRC_TotalPhysicalGB")), FString::FromInt(totalMemoryGB));
+	PredefinedConfigRuleVars.Add(FString(TEXT("memory")), FString::FromInt(totalMemoryMB));
+
+	PredefinedConfigRuleVars.Add(FString(TEXT("versionCode")), BuildNumber);
+	PredefinedConfigRuleVars.Add(FString(TEXT("versionName")), ProjectVersion);
+	
+	PredefinedConfigRuleVars.Add(FString(TEXT("processorCount")), FString::FromInt(FIOSPlatformMisc::NumberOfCores()));
+	
+	id<MTLDevice> DefaultMetalDevice = MTLCreateSystemDefaultDevice();
+	if(DefaultMetalDevice)
+	{
+		FString GpuName([DefaultMetalDevice name]);
+		PredefinedConfigRuleVars.Add(FString(TEXT("SRC_GPUFamily")), GpuName);
+		MTLGPUFamily Families[] = {
+			MTLGPUFamilyApple1,
+			MTLGPUFamilyApple2,
+			MTLGPUFamilyApple3,
+			MTLGPUFamilyApple4,
+			MTLGPUFamilyApple5,
+			MTLGPUFamilyApple6,
+			MTLGPUFamilyApple7,
+			MTLGPUFamilyApple8,
+			MTLGPUFamilyApple9
+		};
+		
+		for(int i = 0; i< UE_ARRAY_COUNT(Families); i++)
+		{
+			FString SupportsGPUFamily = FString::FromInt([DefaultMetalDevice supportsFamily:Families[i]]);
+			PredefinedConfigRuleVars.Add(FString::Format(TEXT("SRC_SupportsGPUFamily{0}"), {i+1}), SupportsGPUFamily);
+		}
+	}
+	[DefaultMetalDevice release];
+	return PredefinedConfigRuleVars;
+}
+
+const TMap<FString, FString>& FIOSPlatformMisc::GetConfigRuleVars()
+{
+	return FIOSConfigRules::GetConfigRulesMap();
+}
+
 void FIOSPlatformMisc::PlatformPreInit()
 {
     FGenericPlatformMisc::PlatformPreInit();
     
     GIOSAppInfo.Init();
-    
+
     // turn off SIGPIPE crashes
     signal(SIGPIPE, SIG_IGN);
+
+	FIOSConfigRules::Init(GetPredefinedConfigRuleVars());
 }
 
 // Make sure that SetStoredValue and GetStoredValue generate the same key
diff --git a/Engine/Source/Runtime/Core/Public/Android/AndroidPlatformMisc.h b/Engine/Source/Runtime/Core/Public/Android/AndroidPlatformMisc.h
--- a/Engine/Source/Runtime/Core/Public/Android/AndroidPlatformMisc.h
+++ b/Engine/Source/Runtime/Core/Public/Android/AndroidPlatformMisc.h
@@ -191,8 +191,8 @@ public:
 #endif
 	static CORE_API bool IsSupportedAndroidDevice();
 	static CORE_API void SetForceUnsupported(bool bInOverride);
-	static CORE_API const TMap<FString, FString>& GetConfigRulesTMap();
-	static CORE_API FString* GetConfigRulesVariable(const FString& Key);
+	
+	static CORE_API const TMap<FString, FString>& GetConfigRuleVars();
 
 	/* HasVulkanDriverSupport
 	 * @return true if this Android device supports a Vulkan API Unreal could use
diff --git a/Engine/Source/Runtime/Core/Public/GenericPlatform/GenericPlatformMisc.h b/Engine/Source/Runtime/Core/Public/GenericPlatform/GenericPlatformMisc.h
--- a/Engine/Source/Runtime/Core/Public/GenericPlatform/GenericPlatformMisc.h
+++ b/Engine/Source/Runtime/Core/Public/GenericPlatform/GenericPlatformMisc.h
@@ -1870,6 +1870,16 @@ public:
 		return MAX_int32;
 	}
 	
+	/*
+	* Returns the value of config rule variable or null if not found.
+	*/
+	static CORE_API const FString* GetConfigRulesVariable(const FString& Key);
+
+	/*
+	* Returns the map containing all config rule variables.
+	*/
+	static CORE_API const class TMap<FString, FString>& GetConfigRuleVars();
+ 
 	/**
 	 * Returns true if PGO is currently enabled
 	 */
diff --git a/Engine/Source/Runtime/Core/Public/IOS/IOSPlatformMisc.h b/Engine/Source/Runtime/Core/Public/IOS/IOSPlatformMisc.h
--- a/Engine/Source/Runtime/Core/Public/IOS/IOSPlatformMisc.h
+++ b/Engine/Source/Runtime/Core/Public/IOS/IOSPlatformMisc.h
@@ -245,6 +245,8 @@ struct CORE_API FIOSPlatformMisc : public FApplePlatformMisc
 #if !UE_BUILD_SHIPPING
 	static bool IsConsoleOpen();
 #endif
+
+	static const class TMap<FString, FString>& GetConfigRuleVars();
 };
 
 typedef FIOSPlatformMisc FPlatformMisc;
diff --git a/Engine/Source/Runtime/Engine/Private/DeviceProfiles/DeviceProfileManager.cpp b/Engine/Source/Runtime/Engine/Private/DeviceProfiles/DeviceProfileManager.cpp
--- a/Engine/Source/Runtime/Engine/Private/DeviceProfiles/DeviceProfileManager.cpp
+++ b/Engine/Source/Runtime/Engine/Private/DeviceProfiles/DeviceProfileManager.cpp
@@ -470,9 +470,8 @@ void UDeviceProfileManager::SetDeviceProfileCVars(const FString& DeviceProfileNa
 
 	// now allow for ovverrides in !SHIPPING
 #if !UE_BUILD_SHIPPING
-#if PLATFORM_ANDROID
 	// allow ConfigRules to override cvars first
-	const TMap<FString, FString>& ConfigRules = FAndroidMisc::GetConfigRulesTMap();
+	const TMap<FString, FString>& ConfigRules = FPlatformMisc::GetConfigRuleVars();
 	for (const TPair<FString, FString>& Pair : ConfigRules)
 	{
 		FString Key = Pair.Key;
@@ -487,7 +486,7 @@ void UDeviceProfileManager::SetDeviceProfileCVars(const FString& DeviceProfileNa
 			UE::ConfigUtilities::OnSetCVarFromIniEntry(*GDeviceProfilesIni, *CVarKey, *CVarValue, ECVF_SetByDeviceProfile);
 		}
 	}
-#endif
+
 	// pre-apply any -dpcvars= items, so that they override anything in the DPs
 	// Search for all occurrences of dpcvars and dpcvar on the command line.
 	static const TCHAR* DPCVarTags[]{ TEXT("DPCVars="), TEXT("DPCVar="), TEXT("ForceDPCVars=") };
diff --git a/Engine/Source/Runtime/Launch/Private/Android/LaunchAndroid.cpp b/Engine/Source/Runtime/Launch/Private/Android/LaunchAndroid.cpp
--- a/Engine/Source/Runtime/Launch/Private/Android/LaunchAndroid.cpp
+++ b/Engine/Source/Runtime/Launch/Private/Android/LaunchAndroid.cpp
@@ -453,7 +453,7 @@ static void InitCommandLine()
 	}
 
 #if !UE_BUILD_SHIPPING
-	if (FString* ConfigRulesCmdLineAppend = FAndroidMisc::GetConfigRulesVariable(TEXT("cmdline")))
+	if (const FString* ConfigRulesCmdLineAppend = FAndroidMisc::GetConfigRulesVariable(TEXT("cmdline")))
 	{
 		FCommandLine::Append(**ConfigRulesCmdLineAppend);
 		FPlatformMisc::LowLevelOutputDebugStringf(TEXT("ConfigRules appended: %s"), **ConfigRulesCmdLineAppend);
@@ -496,7 +496,7 @@ static void ApplyAndroidCompatConfigRules()
 	if (GConfig->GetArray(TEXT("AndroidCompatCVars"), TEXT("CVars"), AndroidCompatCVars, GEngineIni))
 	{
 		TSet<FString> AllowedCompatCVars(AndroidCompatCVars);
-		for (const TTuple<FString, FString>& Pair : FAndroidMisc::GetConfigRulesTMap())
+		for (const TTuple<FString, FString>& Pair : FPlatformMisc::GetConfigRuleVars())
 		{
 			const FString& Key = Pair.Key;
 			const FString& Value = Pair.Value;
diff --git a/Engine/Source/Runtime/Launch/Private/IOS/LaunchIOS.cpp b/Engine/Source/Runtime/Launch/Private/IOS/LaunchIOS.cpp
--- a/Engine/Source/Runtime/Launch/Private/IOS/LaunchIOS.cpp
+++ b/Engine/Source/Runtime/Launch/Private/IOS/LaunchIOS.cpp
@@ -417,10 +417,43 @@ void FAppEntry::PlatformInit()
 
 extern TcpConsoleListener *ConsoleListener;
 
+static void ApplyIOSCompatConfigRules()
+{
+	TArray<FString> IOSCompatCVars;
+	if (GConfig->GetArray(TEXT("IOSCompatCVars"), TEXT("CVars"), IOSCompatCVars, GEngineIni))
+	{		
+		TSet<FString> AllowedCompatCVars(IOSCompatCVars);
+		for (const TTuple<FString, FString>& Pair : FPlatformMisc::GetConfigRuleVars())
+		{
+			const FString& Key = Pair.Key;
+			const FString& Value = Pair.Value;
+			static const TCHAR IOSCompat[] = TEXT("IOSCompat.");
+			if (Key.StartsWith(IOSCompat))
+			{
+				FString CVarName = Key.Mid(UE_ARRAY_COUNT(IOSCompat) - 1);
+				if (AllowedCompatCVars.Contains(CVarName))
+				{
+					auto* CVar = IConsoleManager::Get().FindConsoleVariable(*CVarName);
+					if (CVar)
+					{
+						// set with HF priority
+						// configrules are therefore higher priority than deviceprofiles. (e.g. -dpcvars)
+						CVar->Set(*Value, ECVF_SetByHotfix);
+						FPlatformMisc::LowLevelOutputDebugStringf(TEXT("Config rule compat Setting %s = %s"), *CVarName, *Value);
+					}
+				}
+			}
+		}
+	}
+}
+
+
 void FAppEntry::Init()
 {
 	SCOPED_BOOT_TIMING("FAppEntry::Init()");
 	
+	FDelegateHandle ConfigReadyHandle = FCoreDelegates::TSConfigReadyForUse().AddStatic(&ApplyIOSCompatConfigRules);
+
 	//extern TCHAR GCmdLine[16384];
 	GEngineLoop.PreInit(FCommandLine::Get());
 
diff --git a/Engine/Source/Runtime/WebBrowser/Private/MobileJS/MobileJSScripting.cpp b/Engine/Source/Runtime/WebBrowser/Private/MobileJS/MobileJSScripting.cpp
--- a/Engine/Source/Runtime/WebBrowser/Private/MobileJS/MobileJSScripting.cpp
+++ b/Engine/Source/Runtime/WebBrowser/Private/MobileJS/MobileJSScripting.cpp
@@ -589,11 +589,11 @@ void FMobileJSScripting::InitializeScript(TSharedRef<class IWebBrowserWindow> In
 	int32 ScreenWidth = Viewport.X;
 	int32 ScreenHeight = Viewport.Y;
 #if PLATFORM_ANDROID
-	if (FString* ScreenWidthVar = FAndroidMisc::GetConfigRulesVariable(TEXT("ScreenWidth")))
+	if (const FString* ScreenWidthVar = FAndroidMisc::GetConfigRulesVariable(TEXT("ScreenWidth")))
 	{
 		ScreenWidth = FCString::Atoi(**ScreenWidthVar);
 	}
-	if (FString* ScreenHeightVar = FAndroidMisc::GetConfigRulesVariable(TEXT("ScreenHeight")))
+	if (const FString* ScreenHeightVar = FAndroidMisc::GetConfigRulesVariable(TEXT("ScreenHeight")))
 	{
 		ScreenHeight = FCString::Atoi(**ScreenHeightVar);
 	}
