Possible iOS intermittent/random controller crash fix. Crash appears to be happening after extended play sessions in the GameController Framework.
- remove dead code
- add some extra error checking

#rnx

diff --git a/Engine/Source/Runtime/ApplicationCore/Private/Apple/AppleControllerInterface.cpp b/Engine/Source/Runtime/ApplicationCore/Private/Apple/AppleControllerInterface.cpp
--- a/Engine/Source/Runtime/ApplicationCore/Private/Apple/AppleControllerInterface.cpp
+++ b/Engine/Source/Runtime/ApplicationCore/Private/Apple/AppleControllerInterface.cpp
@@ -260,60 +260,65 @@ void FAppleControllerInterface::HandleDisconnect(GCController* Controller)
 void FAppleControllerInterface::SendControllerEvents()
 {
 	@autoreleasepool{
-    for(int32 i = 0; i < UE_ARRAY_COUNT(Controllers); ++i)
- 	{
+	for(int32 i = 0; i < UE_ARRAY_COUNT(Controllers); ++i)
+	{
 		FUserController& Controller = Controllers[i];
 		
 		// make sure the connection handler has run on this
-		if (Controller.PlayerIndex == PlayerIndex::PlayerUnset || Controller.Controller == nil)
+		GCController* ControllerImpl = Controller.Controller;
+		if (Controller.PlayerIndex == PlayerIndex::PlayerUnset || ControllerImpl == nil)
 		{
-            continue;
+			continue;
 		}
-		
-		IPlatformInputDeviceMapper& DeviceMapper = IPlatformInputDeviceMapper::Get();
-		FPlatformUserId UserId = FGenericPlatformMisc::GetPlatformUserForUserIndex(Controller.PlayerIndex);
-		FInputDeviceId DeviceId = DeviceMapper.GetPrimaryInputDeviceForUser(UserId);
-		
-		GCController* ControllerImpl = Controller.Controller;
-		GCExtendedGamepad* ExtendedGamepad = [ControllerImpl capture].extendedGamepad;
-        
-		if (ExtendedGamepad != nil)
+
+		HandleButtonGamepad(FGamepadKeyNames::FaceButtonBottom, i);
+		HandleButtonGamepad(FGamepadKeyNames::FaceButtonLeft, i);
+		HandleButtonGamepad(FGamepadKeyNames::FaceButtonRight, i);
+		HandleButtonGamepad(FGamepadKeyNames::FaceButtonTop, i);
+		HandleButtonGamepad(FGamepadKeyNames::LeftShoulder, i);
+		HandleButtonGamepad(FGamepadKeyNames::RightShoulder, i);
+		HandleButtonGamepad(FGamepadKeyNames::LeftTriggerThreshold, i);
+		HandleButtonGamepad(FGamepadKeyNames::RightTriggerThreshold, i);
+		HandleButtonGamepad(FGamepadKeyNames::DPadUp, i);
+		HandleButtonGamepad(FGamepadKeyNames::DPadDown, i);
+		HandleButtonGamepad(FGamepadKeyNames::DPadRight, i);
+		HandleButtonGamepad(FGamepadKeyNames::DPadLeft, i);
+		HandleButtonGamepad(FGamepadKeyNames::SpecialRight, i);
+		HandleButtonGamepad(FGamepadKeyNames::SpecialLeft, i);
+
+		HandleAnalogGamepad(FGamepadKeyNames::LeftAnalogX, i);
+		HandleAnalogGamepad(FGamepadKeyNames::LeftAnalogY, i);
+		HandleAnalogGamepad(FGamepadKeyNames::RightAnalogX, i);
+		HandleAnalogGamepad(FGamepadKeyNames::RightAnalogY, i);
+		HandleAnalogGamepad(FGamepadKeyNames::RightTriggerAnalog, i);
+		HandleAnalogGamepad(FGamepadKeyNames::LeftTriggerAnalog, i);
+
+		HandleVirtualButtonGamepad(FGamepadKeyNames::LeftStickRight, FGamepadKeyNames::LeftStickLeft, i);
+		HandleVirtualButtonGamepad(FGamepadKeyNames::LeftStickDown, FGamepadKeyNames::LeftStickUp, i);
+		HandleVirtualButtonGamepad(FGamepadKeyNames::RightStickLeft, FGamepadKeyNames::RightStickRight, i);
+		HandleVirtualButtonGamepad(FGamepadKeyNames::RightStickDown, FGamepadKeyNames::RightStickUp, i);
+		HandleButtonGamepad(FGamepadKeyNames::LeftThumb, i);
+		HandleButtonGamepad(FGamepadKeyNames::RightThumb, i);
+
+		// Update the saved previous captured game controller. This shouldn't fail, but it does occasionally in the
+		// GameController Framework, so add some extra error checking
+		GCController* CapturedController = [ControllerImpl capture];
+		if (CapturedController != nil)
 		{
-            const GCExtendedGamepad* PreviousExtendedGamepad = Controller.PreviousExtendedGamepad;
-
-            HandleButtonGamepad(FGamepadKeyNames::FaceButtonBottom, i);
-            HandleButtonGamepad(FGamepadKeyNames::FaceButtonLeft, i);
-            HandleButtonGamepad(FGamepadKeyNames::FaceButtonRight, i);
-            HandleButtonGamepad(FGamepadKeyNames::FaceButtonTop, i);
-            HandleButtonGamepad(FGamepadKeyNames::LeftShoulder, i);
-            HandleButtonGamepad(FGamepadKeyNames::RightShoulder, i);
-            HandleButtonGamepad(FGamepadKeyNames::LeftTriggerThreshold, i);
-            HandleButtonGamepad(FGamepadKeyNames::RightTriggerThreshold, i);
-            HandleButtonGamepad(FGamepadKeyNames::DPadUp, i);
-            HandleButtonGamepad(FGamepadKeyNames::DPadDown, i);
-            HandleButtonGamepad(FGamepadKeyNames::DPadRight, i);
-            HandleButtonGamepad(FGamepadKeyNames::DPadLeft, i);
-            HandleButtonGamepad(FGamepadKeyNames::SpecialRight, i);
-            HandleButtonGamepad(FGamepadKeyNames::SpecialLeft, i);
-            
-            HandleAnalogGamepad(FGamepadKeyNames::LeftAnalogX, i);
-            HandleAnalogGamepad(FGamepadKeyNames::LeftAnalogY, i);
-            HandleAnalogGamepad(FGamepadKeyNames::RightAnalogX, i);
-            HandleAnalogGamepad(FGamepadKeyNames::RightAnalogY, i);
-            HandleAnalogGamepad(FGamepadKeyNames::RightTriggerAnalog, i);
-            HandleAnalogGamepad(FGamepadKeyNames::LeftTriggerAnalog, i);
-
-
-            HandleVirtualButtonGamepad(FGamepadKeyNames::LeftStickRight, FGamepadKeyNames::LeftStickLeft, i);
-            HandleVirtualButtonGamepad(FGamepadKeyNames::LeftStickDown, FGamepadKeyNames::LeftStickUp, i);
-            HandleVirtualButtonGamepad(FGamepadKeyNames::RightStickLeft, FGamepadKeyNames::RightStickRight, i);
-            HandleVirtualButtonGamepad(FGamepadKeyNames::RightStickDown, FGamepadKeyNames::RightStickUp, i);
-            HandleButtonGamepad(FGamepadKeyNames::LeftThumb, i);
-            HandleButtonGamepad(FGamepadKeyNames::RightThumb, i);
-
-            [Controller.PreviousExtendedGamepad release];
-            Controller.PreviousExtendedGamepad = ExtendedGamepad;
-            [Controller.PreviousExtendedGamepad retain];
+			GCExtendedGamepad* ExtendedGamepad = CapturedController.extendedGamepad;
+			if (ExtendedGamepad != nil)
+			{
+				[Controller.PreviousExtendedGamepad release];
+				Controller.PreviousExtendedGamepad = [ExtendedGamepad retain];
+			}
+			else
+			{
+				UE_LOG(LogAppleController, Warning, TEXT("Unable to access extendedGamepad"));
+		    }
+		}
+		else
+		{
+			UE_LOG(LogAppleController, Warning, TEXT("Unable to capture controller"));
 		}
 	}
 	} //@autoreleasepool
