iOS screen density from device profile: take the view's contentScaleFactor into account

#rnx

diff --git a/Engine/Source/Runtime/ApplicationCore/Private/IOS/IOSPlatformApplicationMisc.cpp b/Engine/Source/Runtime/ApplicationCore/Private/IOS/IOSPlatformApplicationMisc.cpp
--- a/Engine/Source/Runtime/ApplicationCore/Private/IOS/IOSPlatformApplicationMisc.cpp
+++ b/Engine/Source/Runtime/ApplicationCore/Private/IOS/IOSPlatformApplicationMisc.cpp
@@ -275,11 +275,27 @@ void FIOSPlatformApplicationMisc::ClipboardPaste(class FString& Result)
 
 EScreenPhysicalAccuracy FIOSPlatformApplicationMisc::ComputePhysicalScreenDensity(int32& ScreenDensity)
 {
+	// look up the current scale factor
+	double ContentScaleFactor = 0;
+	if (const UIView* View = [IOSAppDelegate GetDelegate].IOSView)
+	{
+		ContentScaleFactor = View.contentScaleFactor;
+	}
+
 	// If we have a cvar set by a device profile, return it.
 	ScreenDensity = CVarPhysicalScreenDensity.GetValueOnAnyThread();
-	if (ScreenDensity >= 0)
+	if (ScreenDensity == 0)
+	{
+		return EScreenPhysicalAccuracy::Unknown;
+	}
+	else if (ScreenDensity > 0)
 	{
-		return ScreenDensity == 0 ? EScreenPhysicalAccuracy::Unknown : EScreenPhysicalAccuracy::Truth;
+		if (ContentScaleFactor != 0)
+		{
+			// The view may not use the full resolution of the screen so we scale the screen density accordingly
+			ScreenDensity = FMath::TruncToInt(ScreenDensity / ContentScaleFactor);
+		}
+		return EScreenPhysicalAccuracy::Truth;
 	}
 
 	// If it hasn't been set, assume that the density is a multiple of the 
@@ -291,11 +307,7 @@ EScreenPhysicalAccuracy FIOSPlatformApplicationMisc::ComputePhysicalScreenDensit
 #endif
 	ScreenDensity = FMath::TruncToInt(163 * NativeScale);
 
-	// look up the current scale factor
-	UIView* View = [IOSAppDelegate GetDelegate].IOSView;
-	const double ContentScaleFactor = View.contentScaleFactor;
-
-	if ( ContentScaleFactor != 0 )
+	if (ContentScaleFactor != 0)
 	{
 		ScreenDensity = FMath::TruncToInt(ScreenDensity * (ContentScaleFactor / NativeScale));
 	}
