Fix non-modal windows being resizable on Mac while there is a modal window active

#tests Local testing with the editor and using repro-steps from the associated jira

diff --git a/Engine/Source/Runtime/ApplicationCore/Private/Mac/MacPlatformApplicationMisc.cpp b/Engine/Source/Runtime/ApplicationCore/Private/Mac/MacPlatformApplicationMisc.cpp
--- a/Engine/Source/Runtime/ApplicationCore/Private/Mac/MacPlatformApplicationMisc.cpp
+++ b/Engine/Source/Runtime/ApplicationCore/Private/Mac/MacPlatformApplicationMisc.cpp
@@ -498,7 +498,7 @@ void FMacPlatformApplicationMisc::PumpMessages(bool bFromMainLoop)
 					
 					UpdateApplicationMenu();
 					UpdateWindowMenu();
-					UpdateCocoaButtons();
+					UpdateCocoaWindows();
 				}, NSDefaultRunLoopMode, false);
 
 				UpdateCachedMacMenuStateCb();
@@ -582,7 +582,7 @@ void FMacPlatformApplicationMisc::UpdateWindowMenu()
 }
 
 
-void FMacPlatformApplicationMisc::UpdateCocoaButtons()
+void FMacPlatformApplicationMisc::UpdateCocoaWindows()
 {
 	MacApplication->GetWindowsArrayMutex().Lock();
 
@@ -599,12 +599,22 @@ void FMacPlatformApplicationMisc::UpdateCocoaButtons()
                 CloseButton.enabled = false;
                 MinimizeButton.enabled = false;
                 MaximizeButton.enabled = false;
+
+            	if (!Window->GetDefinition().IsModalWindow)
+            	{
+            		Window->GetWindowHandle().styleMask &= ~NSWindowStyleMaskResizable;
+            	}
             }
             else if(!bMacApplicationModalMode)
             {
                 CloseButton.enabled = Window->GetDefinition().HasCloseButton;
                 MinimizeButton.enabled = Window->GetDefinition().SupportsMinimize;
                 MaximizeButton.enabled = Window->GetDefinition().SupportsMaximize;
+            	
+            	if (!Window->GetDefinition().IsModalWindow && (Window->GetDefinition().SupportsMaximize || Window->GetDefinition().HasSizingFrame))
+            	{
+            		Window->GetWindowHandle().styleMask |= NSWindowStyleMaskResizable;
+            	}
             }
         }
     }
diff --git a/Engine/Source/Runtime/ApplicationCore/Public/Mac/MacPlatformApplicationMisc.h b/Engine/Source/Runtime/ApplicationCore/Public/Mac/MacPlatformApplicationMisc.h
--- a/Engine/Source/Runtime/ApplicationCore/Public/Mac/MacPlatformApplicationMisc.h
+++ b/Engine/Source/Runtime/ApplicationCore/Public/Mac/MacPlatformApplicationMisc.h
@@ -25,7 +25,7 @@ struct APPLICATIONCORE_API FMacPlatformApplicationMisc : public FGenericPlatform
     static void UpdateApplicationMenu();
     static void LanguageChanged();
 	static void UpdateWindowMenu();
-    static void UpdateCocoaButtons();
+    static void UpdateCocoaWindows();
 	static struct FLinearColor GetScreenPixelColor(const FVector2D& InScreenPos, float InGamma = 1.0f);
 	static float GetDPIScaleFactorAtPoint(float X, float Y);
 	static bool IsHighDPIModeEnabled() { return bIsHighResolutionCapable && IsHighDPIAwarenessEnabled(); }
