Fix a couple of MainThreadCheck issues located in FMacWindow by using cached properties

#tests Local testing for about 1 week

diff --git a/Engine/Source/Runtime/ApplicationCore/Private/Mac/CocoaWindow.cpp b/Engine/Source/Runtime/ApplicationCore/Private/Mac/CocoaWindow.cpp
--- a/Engine/Source/Runtime/ApplicationCore/Private/Mac/CocoaWindow.cpp
+++ b/Engine/Source/Runtime/ApplicationCore/Private/Mac/CocoaWindow.cpp
@@ -35,6 +35,10 @@ NSString* NSPerformDragOperation = @"NSPerformDragOperation";
 	if (NewSelf)
 	{
 		bZoomed = [super isZoomed];
+		bMinimized = [super isMiniaturized];
+		bKeyWindow = [super isKeyWindow];
+		CachedScreen = [super screen];
+		CachedBackingScaleFactor = [super backingScaleFactor];
 		bIsOnActiveSpace = [super isOnActiveSpace];
 		self.TargetWindowMode = EWindowMode::Windowed;
 		[super setAlphaValue:Opacity];
@@ -262,6 +266,16 @@ NSString* NSPerformDragOperation = @"NSPerformDragOperation";
 	}
 }
 
+- (void)windowDidBecomeKey:(NSNotification *)notification
+{
+	bKeyWindow = true;
+}
+
+- (void)windowDidResignKey:(NSNotification *)notification
+{
+	bKeyWindow = false;
+}
+
 - (void)windowDidBecomeMain:(NSNotification*)Notification
 {
 	SCOPED_AUTORELEASE_POOL;
@@ -345,6 +359,7 @@ NSString* NSPerformDragOperation = @"NSPerformDragOperation";
 
 - (void)windowDidChangeScreen:(NSNotification*)Notification
 {
+	CachedScreen = [self screen];
 	// The windowdidChangeScreen notification only comes after you finish dragging.
 	// It does however, work fine for handling display arrangement changes that cause a window to go offscreen.
 	if (bDisplayReconfiguring)
@@ -429,6 +444,11 @@ NSString* NSPerformDragOperation = @"NSPerformDragOperation";
 	}, @[ NSDefaultRunLoopMode ], false);
 }
 
+- (void)windowDidChangeBackingProperties:(NSNotification *)notification
+{
+	CachedBackingScaleFactor = [self backingScaleFactor];
+}
+
 - (void)windowWillStartLiveResize:(NSNotification*)Notification
 {
 	SCOPED_AUTORELEASE_POOL;
@@ -574,4 +594,14 @@ NSString* NSPerformDragOperation = @"NSPerformDragOperation";
 	}
 }
 
+- (void)windowDidMiniaturize:(NSNotification *)notification
+{
+	bMinimized = true;
+}
+
+- (void)windowDidDeminiaturize:(NSNotification *)notification
+{
+	bMinimized = false;
+}
+
 @end
diff --git a/Engine/Source/Runtime/ApplicationCore/Private/Mac/MacWindow.cpp b/Engine/Source/Runtime/ApplicationCore/Private/Mac/MacWindow.cpp
--- a/Engine/Source/Runtime/ApplicationCore/Private/Mac/MacWindow.cpp
+++ b/Engine/Source/Runtime/ApplicationCore/Private/Mac/MacWindow.cpp
@@ -403,7 +403,9 @@ bool FMacWindow::IsMaximized() const
 bool FMacWindow::IsMinimized() const
 {
 	SCOPED_AUTORELEASE_POOL;
-	return WindowHandle.miniaturized;
+	// @TODO: UE-227397
+	// Currently using cached property to avoid triggering a MainThreadChecker issue.
+	return WindowHandle->bMinimized;
 }
 
 bool FMacWindow::IsVisible() const
@@ -456,7 +458,9 @@ bool FMacWindow::IsPointInWindow( int32 X, int32 Y ) const
 	SCOPED_AUTORELEASE_POOL;
 
 	bool PointInWindow = false;
-	if (!WindowHandle.miniaturized)
+	// @TODO: UE-227397
+	// Currently using cached property to avoid triggering a MainThreadChecker issue.
+	if (!WindowHandle->bMinimized)
 	{
 		NSRect WindowFrame = [WindowHandle frame];
 		WindowFrame.size = [WindowHandle openGLFrame].size;
@@ -477,7 +481,9 @@ bool FMacWindow::IsPointInWindow( int32 X, int32 Y ) const
 		}
 		if (bScreensHaveSeparateSpaces)
 		{
-			NSRect ScreenFrame = [[WindowHandle screen] frame];
+			// @TODO: UE-227397
+			// Currently using cached property to avoid triggering a MainThreadChecker issue.
+			NSRect ScreenFrame = WindowHandle->CachedScreen.frame;
 			NSRect Intersection = NSIntersectionRect(ScreenFrame, WindowFrame);
 			VisibleFrame.size = Intersection.size;
 			VisibleFrame.origin.x = Intersection.origin.x - WindowFrame.origin.x;
@@ -502,25 +508,27 @@ int32 FMacWindow::GetWindowBorderSize() const
 
 bool FMacWindow::IsForegroundWindow() const
 {
-	return WindowHandle.keyWindow;
+	// @TODO: UE-227397
+	// Currently using cached property to avoid triggering a MainThreadChecker issue.
+	return WindowHandle->bKeyWindow;
 }
 
 void FMacWindow::SetText(const TCHAR* const Text)
 {
 	SCOPED_AUTORELEASE_POOL;
-	if (FString(WindowHandle.title) != FString(Text))
-	{
-		CFStringRef CFName = FPlatformString::TCHARToCFString( Text );
-		MainThreadCall(^{
-			SCOPED_AUTORELEASE_POOL;
+	MainThreadCall(^{
+		if (FString(WindowHandle.title) != FString(Text))
+		{
+			CFStringRef CFName = FPlatformString::TCHARToCFString( Text );
+
 			[WindowHandle setTitle: (NSString*)CFName];
 			if (IsRegularWindow())
 			{
 				[NSApp changeWindowsItem: WindowHandle title: (NSString*)CFName filename: NO];
 			}
 			CFRelease( CFName );
-		}, UnrealNilEventMode, true);
-	}
+		}
+	}, UnrealNilEventMode, true);
 }
 
 bool FMacWindow::IsRegularWindow() const
@@ -530,7 +538,9 @@ bool FMacWindow::IsRegularWindow() const
 
 float FMacWindow::GetDPIScaleFactor() const
 {
-	return FPlatformApplicationMisc::IsHighDPIModeEnabled() ? (float)WindowHandle.backingScaleFactor : 1.0f;
+	// @TODO: UE-227397
+	// Currently using cached property to avoid triggering a MainThreadChecker issue.
+	return FPlatformApplicationMisc::IsHighDPIModeEnabled() ? (float)WindowHandle->CachedBackingScaleFactor : 1.0f;
 }
 
 void FMacWindow::SetNativeWindowButtonsVisibility(bool bVisible)
diff --git a/Engine/Source/Runtime/ApplicationCore/Public/Mac/CocoaWindow.h b/Engine/Source/Runtime/ApplicationCore/Public/Mac/CocoaWindow.h
--- a/Engine/Source/Runtime/ApplicationCore/Public/Mac/CocoaWindow.h
+++ b/Engine/Source/Runtime/ApplicationCore/Public/Mac/CocoaWindow.h
@@ -25,6 +25,10 @@
 	bool bZoomed;
 	bool bIsOnActiveSpace;
 	bool bIsBeingResized;
+	bool bMinimized;
+	bool bKeyWindow;
+	NSScreen* CachedScreen;
+	CGFloat CachedBackingScaleFactor;
 }
 
 @property (assign) EWindowMode::Type TargetWindowMode;
