Fix deadlock happening sporadically on Mac when opening LyraGame from a fullscreen Xcode
The problem was that WindowsMutex is locked in GameThread before calling a MainThreadCall when MainThread is waiting for the WindowsMutex before going back to its event loop.
By moving the mutex lock inside the MainThreadCall, we avoid this potential deadlock and avoid running multiple MainThreadCall to process all Windows to one that encapsulate all Windows.

#tests Run LyraGame & LyraEditor locally

diff --git a/Engine/Source/Runtime/ApplicationCore/Private/Mac/MacApplication.cpp b/Engine/Source/Runtime/ApplicationCore/Private/Mac/MacApplication.cpp
--- a/Engine/Source/Runtime/ApplicationCore/Private/Mac/MacApplication.cpp
+++ b/Engine/Source/Runtime/ApplicationCore/Private/Mac/MacApplication.cpp
@@ -704,12 +704,14 @@ void FMacApplication::OnDisplayReconfiguration(CGDirectDisplayID Display, CGDisp
 		App->BroadcastDisplayMetricsChanged(DisplayMetrics);
 	}
 
-	FScopeLock Lock(&App->WindowsMutex);
-	for (int32 WindowIndex=0; WindowIndex < App->Windows.Num(); ++WindowIndex)
-	{
-		TSharedRef<FMacWindow> WindowRef = App->Windows[WindowIndex];
-		WindowRef->OnDisplayReconfiguration(Display, Flags);
-	}
+	MainThreadCall(^{
+		FScopeLock Lock(&App->WindowsMutex);
+		for (int32 WindowIndex=0; WindowIndex < App->Windows.Num(); ++WindowIndex)
+		{
+			TSharedRef<FMacWindow> WindowRef = App->Windows[WindowIndex];
+			WindowRef->OnDisplayReconfiguration(Display, Flags);
+		}
+	});
 }
 
 #if WITH_EDITOR
diff --git a/Engine/Source/Runtime/ApplicationCore/Private/Mac/MacWindow.cpp b/Engine/Source/Runtime/ApplicationCore/Private/Mac/MacWindow.cpp
--- a/Engine/Source/Runtime/ApplicationCore/Private/Mac/MacWindow.cpp
+++ b/Engine/Source/Runtime/ApplicationCore/Private/Mac/MacWindow.cpp
@@ -571,22 +571,20 @@ void FMacWindow::SetNativeWindowButtonsVisibility(bool bVisible)
 
 void FMacWindow::OnDisplayReconfiguration(CGDirectDisplayID Display, CGDisplayChangeSummaryFlags Flags)
 {
+	SCOPED_AUTORELEASE_POOL;
 	if(WindowHandle)
 	{
-		MainThreadCall(^{
-			SCOPED_AUTORELEASE_POOL;
-			if(Flags & kCGDisplayBeginConfigurationFlag)
-			{
-				[WindowHandle setMovable: YES];
-				[WindowHandle setMovableByWindowBackground: NO];
-				
-				[WindowHandle setDisplayReconfiguring: true];
-			}
-			else if(Flags & kCGDisplayDesktopShapeChangedFlag)
-			{
-				[WindowHandle setDisplayReconfiguring: false];
-			}
-		});
+		if(Flags & kCGDisplayBeginConfigurationFlag)
+		{
+			[WindowHandle setMovable: YES];
+			[WindowHandle setMovableByWindowBackground: NO];
+			
+			[WindowHandle setDisplayReconfiguring: true];
+		}
+		else if(Flags & kCGDisplayDesktopShapeChangedFlag)
+		{
+			[WindowHandle setDisplayReconfiguring: false];
+		}
 	}
 }
 
