[Windows] Add command line argument for initial OutputDevice2 filter.

#rnx

diff --git a/Engine/Source/Runtime/ApplicationCore/Private/Windows/WindowsConsoleOutputDevice2.cpp b/Engine/Source/Runtime/ApplicationCore/Private/Windows/WindowsConsoleOutputDevice2.cpp
--- a/Engine/Source/Runtime/ApplicationCore/Private/Windows/WindowsConsoleOutputDevice2.cpp
+++ b/Engine/Source/Runtime/ApplicationCore/Private/Windows/WindowsConsoleOutputDevice2.cpp
@@ -2799,6 +2799,8 @@ void FWindowsConsoleOutputDevice2::Show( bool ShowWindow )
 			ConsolePosY = CW_USEDEFAULT;
 		}
 
+		FParse::Value(FCommandLine::Get(), TEXT("ConsoleIncludeFilter="), ConsoleIncludeFilterStr);
+
 		FString ConsoleTitle;
 
 		if (!FParse::Value(FCommandLine::Get(), TEXT("ConsoleTitle="), ConsoleTitle))
