Fixed the issue of the BP node "Get Game Controller Name" and "Is Controller Assigned to Gamepad" not working when launching the game while a gamepad is already connected.

Note: some gamepads will have a start/activate button, e.g. the XBOX or BEITONG controller. So you have to press that button to "power on" the controller so that the game can detect it.

Fixed the issue of gamepad disconnection delegate not being called.

#android

diff --git a/Engine/Build/Android/Java/src/com/epicgames/unreal/GameActivity.java.template b/Engine/Build/Android/Java/src/com/epicgames/unreal/GameActivity.java.template
--- a/Engine/Build/Android/Java/src/com/epicgames/unreal/GameActivity.java.template
+++ b/Engine/Build/Android/Java/src/com/epicgames/unreal/GameActivity.java.template
@@ -610,13 +610,42 @@ public class GameActivity extends $${gameActivitySuperClass}$$ implements Surfac
 	/** App wide Multicast lock used to allow broadcast/multicast communication */
 	private WifiManager.MulticastLock MultiCastLock = null;
 
+	private GameControllerManager gameControllerManager;
+
 //$${gameActivityClassAdditions}$$
 
+	// This method will be called from native and after the GEngineLoop::PreInit()
+	public boolean createGameControllerManager()
+	{
+		if (gameControllerManager == null)
+		{
+			if (inputManager == null)
+			{
+				inputManager = (InputManager) getSystemService(Context.INPUT_SERVICE);
+			}
+			if (inputManager != null)
+			{
+				gameControllerManager = new GameControllerManager(inputManager);
+				if (gameControllerManager != null)
+				{
+					gameControllerManager.onStart();
+					return true;
+				}
+			}
+		}
+		return false;
+	}
+
 	@Override
 	public void onStart()
 	{
 		super.onStart();
 		
+		if (gameControllerManager != null)
+		{
+			gameControllerManager.onStart();
+		}
+
 		if (!BuildConfiguration.equals("Shipping"))
 		{
 			// Create console command broadcast listener
@@ -4042,6 +4071,11 @@ public class GameActivity extends $${gameActivitySuperClass}$$ implements Surfac
 	@Override
 	public void onStop()
 	{
+		if (gameControllerManager != null)
+		{
+			gameControllerManager.onStop();
+		}
+
 		if (consoleCmdReceiver != null)
 		{
 			unregisterReceiver(consoleCmdReceiver);
diff --git a/Engine/Build/Android/Java/src/com/epicgames/unreal/GameControllerManager.java b/Engine/Build/Android/Java/src/com/epicgames/unreal/GameControllerManager.java
new file mode 100644
--- /dev/null
+++ b/Engine/Build/Android/Java/src/com/epicgames/unreal/GameControllerManager.java
@@ -0,0 +1,205 @@
+// Copyright Epic Games, Inc. All Rights Reserved.
+
+package com.epicgames.unreal;
+
+import android.content.Context;
+import android.hardware.input.InputManager;
+import android.util.Log;
+import android.view.InputDevice;
+import android.view.MotionEvent;
+
+import java.util.ArrayList;
+import java.util.List;
+
+public class GameControllerManager
+{
+    private static final String FINGERPRINT_DEVICE_NAME = "uinput-fpc";
+    private static final int GAMECONTROLLER_SOURCE_MASK = InputDevice.SOURCE_DPAD | InputDevice.SOURCE_JOYSTICK | InputDevice.SOURCE_GAMEPAD;
+
+    private final InputManager inputManager;
+	private static Logger Log = new Logger("UE", "GameControllerManager");;
+
+	private GameControllerThread gameControllerThread;
+
+	private ArrayList<Integer> listDeviceId = new ArrayList<Integer>();
+
+	enum InputDeviceState
+	{
+		Added,
+		Removed,
+		Changed
+	};
+
+
+    public GameControllerManager(InputManager inputMgr)
+    {
+        inputManager = inputMgr;
+    }
+    
+	public InputManager getAppInputManager()
+    {
+        return inputManager;
+    }
+
+	private native void nativeOnZeroInputDevices();
+
+    private void scanDevices()
+    {
+        int[] deviceIds = inputManager.getInputDeviceIds();
+
+		// add new devices
+		for (int deviceId : deviceIds)
+		{
+			if (listDeviceId.indexOf(deviceId) == -1)
+			{
+				onInputDeviceAdded(deviceId);
+			}
+		}
+
+		Log.debug("scanDevices stage 1, listDeviceId size: " + listDeviceId.size());
+
+		// remove disconnected devices
+		for (int index = 0; index < listDeviceId.size(); ++index)
+		{
+			int prevId = listDeviceId.get(index);
+			boolean found = false;
+			for (int deviceId : deviceIds) 
+			{
+				if (prevId == deviceId)
+				{
+					found = true;
+					break;
+				}
+			}
+
+			if (!found)
+			{
+				nativeOnInputDeviceStateEvent(prevId, InputDeviceState.Removed.ordinal());
+				listDeviceId.remove(index--);
+
+				Log.debug("scanDevices removes device: " + prevId);
+			}
+		}
+		
+		Log.debug("scanDevices stage 2, listDeviceId size: " + listDeviceId.size());
+    }
+
+    public void onStop() 
+	{
+        if (gameControllerThread != null) 
+		{
+            gameControllerThread.stopListening();
+			Log.debug("stopListening...");
+        }
+    }
+
+    public void onStart() 
+	{
+        if (gameControllerThread != null) 
+		{
+			// back to foreground and update all connected devices
+            scanDevices();
+            gameControllerThread.startListening();
+			Log.debug("startListening...");
+        } 
+		else 
+		{
+			Log = new Logger("UE", "GameControllerManager");
+
+            gameControllerThread = new GameControllerThread();
+            gameControllerThread.init(this);
+            gameControllerThread.start();
+        }
+    }
+
+    private static boolean isGameController(int deviceId) 
+	{
+        InputDevice inputDevice = InputDevice.getDevice(deviceId);
+        if (inputDevice != null) 
+		{
+            int inputDeviceSources = inputDevice.getSources();
+            int sourceMask = InputDevice.SOURCE_ANY & GAMECONTROLLER_SOURCE_MASK;
+            if ((inputDevice.isVirtual() == false) && ((inputDeviceSources & sourceMask) != 0))
+			{
+				String deviceName = inputDevice.getName();
+				List<InputDevice.MotionRange> motionRanges = inputDevice.getMotionRanges();
+				if (motionRanges != null && motionRanges.size() > 0)
+				{
+					// Some physical keyboards include DPAD and JOYSTICK sources, but
+					// only report a single AXIS_GENERIC_1 motion range, screen them out
+					// as they aren't really game controllers
+					if (motionRanges.size() == 1 && motionRanges.get(0).getAxis() == MotionEvent.AXIS_GENERIC_1)
+					{
+						// Ignore the fingerprint reader, which can also identify
+						// as a game controller for some reason
+						if (deviceName.equalsIgnoreCase(FINGERPRINT_DEVICE_NAME))
+						{
+							return false;
+						}
+					}
+
+					return true;
+				}
+            }
+        }
+
+        return false;
+    }
+
+	private native void nativeOnInputDeviceStateEvent(int deviceId, int state);
+
+    public void onInputDeviceAdded(int deviceId)
+    {
+	    Log.debug("onInputDeviceAdded, id: " + deviceId);
+
+		if (listDeviceId.indexOf(deviceId) != -1)
+		{
+			Log.warn("adding an input device already existed, id: " + deviceId);
+			return;
+		}
+
+        if (!isGameController(deviceId))
+		{
+            return;
+		}
+			
+		listDeviceId.add(deviceId);
+		nativeOnInputDeviceStateEvent(deviceId, InputDeviceState.Added.ordinal());
+	    Log.debug("nativeOnInputDeviceStateEvent, id: " + deviceId);
+    }
+
+    public void onInputDeviceRemoved(int deviceId)
+    {
+        Log.debug("onInputDeviceRemoved, id: " + deviceId);
+
+		for (int index = 0; index < listDeviceId.size(); ++index)
+		{
+			if (listDeviceId.get(index) == deviceId)
+			{
+				listDeviceId.remove(index);
+				nativeOnInputDeviceStateEvent(deviceId, InputDeviceState.Removed.ordinal());
+				Log.debug("nativeOnInputDeviceStateEvent, id: " + deviceId);
+				break;
+			}
+		}
+    }
+
+    public void onInputDeviceChanged(int deviceId)
+    {
+        Log.debug("onInputDeviceChanged, id: " + deviceId);
+		
+        if (!isGameController(deviceId))
+		{
+            return;
+		}
+
+		if (listDeviceId.indexOf(deviceId) == -1)
+		{
+			Log.warn("The changed input device wasn't recorded before, id: " + deviceId);
+			listDeviceId.add(deviceId);
+		}
+		
+		nativeOnInputDeviceStateEvent(deviceId, InputDeviceState.Changed.ordinal());
+		Log.debug("nativeOnInputDeviceStateEvent, id: " + deviceId);
+    }
+}
diff --git a/Engine/Build/Android/Java/src/com/epicgames/unreal/GameControllerThread.java b/Engine/Build/Android/Java/src/com/epicgames/unreal/GameControllerThread.java
new file mode 100644
--- /dev/null
+++ b/Engine/Build/Android/Java/src/com/epicgames/unreal/GameControllerThread.java
@@ -0,0 +1,67 @@
+// Copyright Epic Games, Inc. All Rights Reserved.
+
+package com.epicgames.unreal;
+
+import android.hardware.input.InputManager;
+import android.os.Handler;
+import android.os.Looper;
+import android.util.Log;
+
+import java.lang.Thread;
+
+public class GameControllerThread extends Thread implements InputManager.InputDeviceListener
+{
+    private GameControllerManager mGameControllerManager;
+    private boolean activeInputDeviceListener = false;
+    private Handler mHandler;
+
+    public void init(GameControllerManager gcManager)
+    {
+        mGameControllerManager = gcManager;
+    }
+
+    @Override
+    public void run ()
+    {
+        Looper.prepare();
+        mHandler = new Handler(Looper.myLooper());
+        startListening();
+        Looper.loop();
+    }
+
+    public void startListening()
+    {
+        if (!activeInputDeviceListener)
+        {
+            mGameControllerManager.getAppInputManager().registerInputDeviceListener(this, mHandler);
+            activeInputDeviceListener = true;
+        }
+    }
+
+    public void stopListening()
+    {
+        if (activeInputDeviceListener)
+        {
+            mGameControllerManager.getAppInputManager().unregisterInputDeviceListener(this);
+            activeInputDeviceListener = false;
+        }
+    }
+
+    @Override
+    public void onInputDeviceAdded(int deviceId)
+    {
+        mGameControllerManager.onInputDeviceAdded(deviceId);
+    }
+
+    @Override
+    public void onInputDeviceRemoved(int deviceId)
+    {
+        mGameControllerManager.onInputDeviceRemoved(deviceId);
+    }
+
+    @Override
+    public void onInputDeviceChanged(int deviceId)
+    {
+        mGameControllerManager.onInputDeviceChanged(deviceId);
+    }
+}
diff --git a/Engine/Source/Runtime/ApplicationCore/Private/Android/AndroidInputInterface.cpp b/Engine/Source/Runtime/ApplicationCore/Private/Android/AndroidInputInterface.cpp
--- a/Engine/Source/Runtime/ApplicationCore/Private/Android/AndroidInputInterface.cpp
+++ b/Engine/Source/Runtime/ApplicationCore/Private/Android/AndroidInputInterface.cpp
@@ -1,5 +1,7 @@
 // Copyright Epic Games, Inc. All Rights Reserved.
 
+#include "Android/AndroidJNI.h"
+#include "Android/AndroidApplication.h"
 #include "Android/AndroidInputInterface.h"
 #if USE_ANDROID_INPUT
 #include "Android/AndroidEventManager.h"
@@ -78,6 +80,51 @@ static FAutoConsoleVariableRef CVarAndroidUnifyMotionSpace(
 	TEXT("If set to non-zero, acceleration, gravity, and rotation rate will all be in the same coordinate space. 0 for legacy behaviour. 1 (default as of 5.5) will match Unreal's coordinate space (left-handed, z-up, etc). 2 will be right-handed by swapping x and y. Non-zero also forces rotation rate units to be radians/s and acceleration units to be g."),
 	ECVF_Default);
 
+bool AndroidEnableInputDeviceListener = true;
+static FAutoConsoleVariableRef CVarAndroidEnableInputDeviceListener(
+	TEXT("Android.EnableInputDeviceListener"),
+	AndroidEnableInputDeviceListener,
+	TEXT("Determines how to detect gamepad connection/disconnection. true for using InputDeviceListener. false for using the gamepad input events (Default: true)"),
+	ECVF_Default);
+
+
+bool AndroidThunkCpp_InitGameControllerManager()
+{
+	if (JNIEnv* Env = FAndroidApplication::GetJavaEnv())
+	{
+#if USE_ANDROID_STANDALONE
+		constexpr const ANSICHAR ClassName[] = "com/epicgames/makeaar/GameControllerManager";
+		constexpr const ANSICHAR ClassSig[] = "Lcom/epicgames/makeaar/GameControllerManager;";
+#else
+		constexpr const ANSICHAR ClassName[] = "com/epicgames/unreal/GameControllerManager";
+		constexpr const ANSICHAR ClassSig[] = "Lcom/epicgames/unreal/GameControllerManager;";
+#endif
+		jmethodID CreateGameControllerManagerMethodID = FJavaWrapper::FindMethod(Env, FJavaWrapper::GameActivityClassID, "createGameControllerManager", "()Z", false);
+		bool bGCMCreated = FJavaWrapper::CallBooleanMethod(Env, FJavaWrapper::GameActivityThis, CreateGameControllerManagerMethodID);
+		if (bGCMCreated)
+		{
+			jclass GameControllerManagerClassID = AndroidJavaEnv::FindJavaClass(ClassName);
+			jmethodID ScanDeviceMethodID = FJavaWrapper::FindMethod(Env, GameControllerManagerClassID, "scanDevices", "()V", false);
+			jfieldID GameControllerManagerFieldId = FJavaWrapper::FindField(Env, FJavaWrapper::GameActivityClassID, "gameControllerManager", ClassSig, false);
+			jobject GameControllerManagerObj = Env->GetObjectField(FJavaWrapper::GameActivityThis, GameControllerManagerFieldId);
+			if (GameControllerManagerObj != nullptr)
+			{
+				FJavaWrapper::CallVoidMethod(Env, GameControllerManagerObj, ScanDeviceMethodID);
+				return true;
+			}
+		}
+	}
+
+	return false;
+}
+
+JNI_METHOD void Java_com_epicgames_unreal_GameControllerManager_nativeOnInputDeviceStateEvent(JNIEnv* env, jobject thiz, jint device_id, jint state)
+{
+	InputDeviceStateEvent StateEvent = static_cast<InputDeviceStateEvent>(state);
+	FAndroidInputInterface::HandleInputDeviceStateEvent(device_id, StateEvent);
+}
+
+
 TSharedRef< FAndroidInputInterface > FAndroidInputInterface::Create(const TSharedRef< FGenericApplicationMessageHandler >& InMessageHandler, const TSharedPtr< ICursor >& InCursor)
 {
 	return MakeShareable(new FAndroidInputInterface(InMessageHandler, InCursor));
@@ -143,21 +190,24 @@ FAndroidInputInterface::FAndroidInputInterface(const TSharedRef< FGenericApplica
 	FMemory::Memset(OldControllerData, 0);
 	FMemory::Memset(NewControllerData, 0);
 	FMemory::Memset(ControllerVibeState, 0);
+
+	if (AndroidEnableInputDeviceListener)
+	{
+		extern bool AndroidThunkCpp_InitGameControllerManager();
+		if (!AndroidThunkCpp_InitGameControllerManager())
+		{
+			UE_LOG(LogAndroid, Error, TEXT("GameControllerManager initialization failed!"));
+		}
+	}
 }
 
 void FAndroidInputInterface::ResetGamepadAssignments()
 {
-	IPlatformInputDeviceMapper& DeviceMapper = IPlatformInputDeviceMapper::Get();
-	
 	for (int32 DeviceIndex = 0; DeviceIndex < MAX_NUM_CONTROLLERS; DeviceIndex++)
 	{
 		if (DeviceMapping[DeviceIndex].DeviceState == MappingState::Valid)
 		{
-			FPlatformUserId PlatformUserId = PLATFORMUSERID_NONE;
-			FInputDeviceId DeviceId = INPUTDEVICEID_NONE;
-			
-			DeviceMapper.RemapControllerIdToPlatformUserAndDevice(DeviceIndex, OUT PlatformUserId, OUT DeviceId);
-			DeviceMapper.Internal_MapInputDeviceToUser(DeviceId, PlatformUserId, EInputDeviceConnectionState::Disconnected);
+			MapControllerToPlayer(DeviceIndex, EInputDeviceConnectionState::Disconnected);
 		}
 
 		DeviceMapping[DeviceIndex].DeviceInfo.DeviceId = 0;
@@ -172,13 +222,7 @@ void FAndroidInputInterface::ResetGamepadAssignmentToController(int32 Controller
 
 	if (DeviceMapping[ControllerId].DeviceState == MappingState::Valid)
 	{
-		IPlatformInputDeviceMapper& DeviceMapper = IPlatformInputDeviceMapper::Get();
-
-		FPlatformUserId PlatformUserId = PLATFORMUSERID_NONE;
-		FInputDeviceId DeviceId = INPUTDEVICEID_NONE;
-			
-		DeviceMapper.RemapControllerIdToPlatformUserAndDevice(ControllerId, OUT PlatformUserId, OUT DeviceId);
-		DeviceMapper.Internal_MapInputDeviceToUser(DeviceId, PlatformUserId, EInputDeviceConnectionState::Disconnected);
+		MapControllerToPlayer(ControllerId, EInputDeviceConnectionState::Disconnected);
 	}
 
 	DeviceMapping[ControllerId].DeviceInfo.DeviceId = 0;
@@ -190,7 +234,8 @@ bool FAndroidInputInterface::IsControllerAssignedToGamepad(int32 ControllerId)
 	if (ControllerId < 0 || ControllerId >= MAX_NUM_CONTROLLERS)
 		return false;
 
-	return (DeviceMapping[ControllerId].DeviceState == MappingState::Valid);
+	MappingState PotentState = AndroidEnableInputDeviceListener ? MappingState::ToActivate : MappingState::ToValidate;
+	return (DeviceMapping[ControllerId].DeviceState > PotentState);
 }
 
 FString FAndroidInputInterface::GetGamepadControllerName(int32 ControllerId)
@@ -200,7 +245,8 @@ FString FAndroidInputInterface::GetGamepadControllerName(int32 ControllerId)
 		return FString(TEXT("None"));
 	}
 
-	if (DeviceMapping[ControllerId].DeviceState != MappingState::Valid)
+	MappingState PotentState = AndroidEnableInputDeviceListener ? MappingState::ToValidate : MappingState::Valid;
+	if (DeviceMapping[ControllerId].DeviceState < PotentState)
 	{
 		return FString(TEXT("None"));
 	}
@@ -226,6 +272,15 @@ void FAndroidInputInterface::AddExternalInputDevice(TSharedPtr<IInputDevice> Inp
 	}
 }
 
+void FAndroidInputInterface::MapControllerToPlayer(int32 ControllerId, EInputDeviceConnectionState State)
+{
+	IPlatformInputDeviceMapper& DeviceMapper = IPlatformInputDeviceMapper::Get();
+	FPlatformUserId PlatformUserId = PLATFORMUSERID_NONE;
+	FInputDeviceId DeviceId = INPUTDEVICEID_NONE;
+	DeviceMapper.RemapControllerIdToPlatformUserAndDevice(ControllerId, OUT PlatformUserId, OUT DeviceId);
+	DeviceMapper.Internal_MapInputDeviceToUser(DeviceId, PlatformUserId, State);
+}
+
 extern bool AndroidThunkCpp_GetInputDeviceInfo(int32 deviceId, FAndroidInputDeviceInfo &results);
 
 void FAndroidInputInterface::Tick(float DeltaTime)
@@ -984,11 +1039,13 @@ void FAndroidInputInterface::SendControllerEvents()
 	// Check for gamepads needing validation if enabled
 	if (bAllowControllers)
 	{
+		MappingState StateToCheck = AndroidEnableInputDeviceListener ? MappingState::ToActivate : MappingState::ToValidate;
+		MappingState StateSwitchTo = AndroidEnableInputDeviceListener ? MappingState::ToValidate : MappingState::Valid;
+
 		for (int32 DeviceIndex = 0; DeviceIndex < MAX_NUM_CONTROLLERS; DeviceIndex++)
 		{
 			FAndroidGamepadDeviceMapping& CurrentDevice = DeviceMapping[DeviceIndex];
-
-			if (CurrentDevice.DeviceState == MappingState::ToValidate)
+			if (CurrentDevice.DeviceState == StateToCheck)
 			{
 				// Query for the device type from Java side
 				if (AndroidThunkCpp_GetInputDeviceInfo(CurrentDevice.DeviceInfo.DeviceId, CurrentDevice.DeviceInfo))
@@ -997,7 +1054,7 @@ void FAndroidInputInterface::SendControllerEvents()
 					int32 FoundMatch = -1;
 					for (int32 DeviceScan = 0; DeviceScan < MAX_NUM_CONTROLLERS; DeviceScan++)
 					{
-						if (DeviceMapping[DeviceScan].DeviceState != MappingState::Valid)
+						if (DeviceMapping[DeviceScan].DeviceState < StateSwitchTo)
 							continue;
 
 						if (DeviceMapping[DeviceScan].DeviceInfo.Descriptor.Equals(CurrentDevice.DeviceInfo.Descriptor))
@@ -1010,7 +1067,7 @@ void FAndroidInputInterface::SendControllerEvents()
 					// Deal with new controller
 					if (FoundMatch == -1)
 					{
-						CurrentDevice.DeviceState = MappingState::Valid;
+						CurrentDevice.DeviceState = StateSwitchTo;
 
 						// Generic mappings
 						CurrentDevice.ControllerClass = ControllerClassType::Generic;
@@ -1165,12 +1222,10 @@ void FAndroidInputInterface::SendControllerEvents()
 							CurrentDevice.bTriggersUseThresholdForClick = true;
 						}
 
-						IPlatformInputDeviceMapper& DeviceMapper = IPlatformInputDeviceMapper::Get();
-						FPlatformUserId PlatformUserId = PLATFORMUSERID_NONE;
-						FInputDeviceId DeviceId = INPUTDEVICEID_NONE;
-			
-						DeviceMapper.RemapControllerIdToPlatformUserAndDevice(DeviceIndex, OUT PlatformUserId, OUT DeviceId);
-						DeviceMapper.Internal_MapInputDeviceToUser(DeviceId, PlatformUserId, EInputDeviceConnectionState::Connected);
+						if (CurrentDevice.DeviceState == MappingState::Valid)
+						{
+							MapControllerToPlayer(DeviceIndex, EInputDeviceConnectionState::Connected);
+						}
 
 						FPlatformMisc::LowLevelOutputDebugStringf(TEXT("Assigned new gamepad controller %d: DeviceId=%d, ControllerId=%d, DeviceName=%s, Descriptor=%s"),
 							DeviceIndex, CurrentDevice.DeviceInfo.DeviceId, CurrentDevice.DeviceInfo.ControllerId, *CurrentDevice.DeviceInfo.Name, *CurrentDevice.DeviceInfo.Descriptor);
@@ -1188,14 +1243,10 @@ void FAndroidInputInterface::SendControllerEvents()
 						NewControllerData[FoundMatch].DeviceId = FoundMatch;
 						OldControllerData[FoundMatch].DeviceId = FoundMatch;
 
-						//@TODO: uncomment these line in the future when disconnects are detected
-						// IPlatformInputDeviceMapper& DeviceMapper = IPlatformInputDeviceMapper::Get();
-						//
-						// FPlatformUserId PlatformUserId = PLATFORMUSERID_NONE;
-						// FInputDeviceId DeviceId = INPUTDEVICEID_NONE;
-						//
-						// DeviceMapper.RemapControllerIdToPlatformUserAndDevice(FoundMatch, OUT PlatformUserId, OUT DeviceId);
-						// DeviceMapper.Internal_MapInputDeviceToUser(DeviceId, PlatformUserId, EInputDeviceConnectionState::Connected);
+						if (DeviceMapping[FoundMatch].DeviceState == MappingState::Valid)
+						{
+							MapControllerToPlayer(FoundMatch, EInputDeviceConnectionState::Connected);
+						}
 
 						FPlatformMisc::LowLevelOutputDebugStringf(TEXT("Reconnected gamepad controller %d: DeviceId=%d, ControllerId=%d, DeviceName=%s, Descriptor=%s"),
 							FoundMatch, DeviceMapping[FoundMatch].DeviceInfo.DeviceId, CurrentDevice.DeviceInfo.ControllerId, *CurrentDevice.DeviceInfo.Name, *CurrentDevice.DeviceInfo.Descriptor);
@@ -1513,11 +1564,67 @@ int32 FAndroidInputInterface::GetControllerIndex(int32 deviceId)
 	NewControllerData[UnassignedIndex].DeviceId = UnassignedIndex;
 
 	// Mark it for validation later
-	DeviceMapping[UnassignedIndex].DeviceState = MappingState::ToValidate;
+	MappingState PotentState = AndroidEnableInputDeviceListener ? MappingState::ToActivate : MappingState::ToValidate;
+	DeviceMapping[UnassignedIndex].DeviceState = PotentState;
 
 	return UnassignedIndex;
 }
 
+void FAndroidInputInterface::HandleInputDeviceStateEvent(int32 DeviceID, InputDeviceStateEvent StateEvent)
+{
+	if (!AndroidEnableInputDeviceListener)
+	{
+		return;
+	}
+
+	FScopeLock Lock(&TouchInputCriticalSection);
+
+	switch (StateEvent)
+	{
+	case InputDeviceStateEvent::Added:
+	case InputDeviceStateEvent::Changed:
+		if (GetControllerIndex(DeviceID) < 0)
+		{
+			UE_LOG(LogAndroid, Error, TEXT("Gamepad connected but failed to add, DeviceID = %d"), DeviceID);
+		}
+		break;
+	case InputDeviceStateEvent::Removed:
+		{
+			int32 ControllerIndex = -1;
+			for (int32 Index = 0; Index < MAX_NUM_CONTROLLERS; Index++)
+			{
+				if (DeviceMapping[Index].DeviceInfo.DeviceId == DeviceID && DeviceMapping[Index].DeviceState > MappingState::Unassigned)
+				{
+					ControllerIndex = Index;
+					break;
+				}
+			}
+
+			if (ControllerIndex > -1)
+			{
+				if (DeviceMapping[ControllerIndex].DeviceState == MappingState::Valid)
+				{
+					ExecuteOnGameThread(TEXT("Broadcast the EInputDeviceConnectionState::Disconnected status"), [ControllerIndex]()
+						{
+							MapControllerToPlayer(ControllerIndex, EInputDeviceConnectionState::Disconnected);
+						});
+				}
+
+				DeviceMapping[ControllerIndex].DeviceState = MappingState::Unassigned;
+				DeviceMapping[ControllerIndex].DeviceInfo.DeviceId = 0;
+				FMemory::Memset(NewControllerData[ControllerIndex], 0);
+				FMemory::Memset(OldControllerData[ControllerIndex], 0);
+			}
+			else
+			{
+				UE_LOG(LogAndroid, Error, TEXT("Nonexistent gamepad removed, DeviceID = %d"), DeviceID);
+			}
+
+			break;
+		}
+	}
+}
+
 void FAndroidInputInterface::JoystickAxisEvent(int32 deviceId, int32 axisId, float axisValue)
 {
 	FScopeLock Lock(&TouchInputCriticalSection);
@@ -1527,6 +1634,13 @@ void FAndroidInputInterface::JoystickAxisEvent(int32 deviceId, int32 axisId, flo
 	if (deviceId == -1)
 		return;
 
+	if (AndroidEnableInputDeviceListener && DeviceMapping[deviceId].DeviceState == MappingState::ToValidate)
+	{
+		MapControllerToPlayer(deviceId, EInputDeviceConnectionState::Connected);
+		DeviceMapping[deviceId].DeviceState = MappingState::Valid;
+		//UE_LOG(LogAndroid, Log, TEXT("Gamepad state change from ToValidate to Valid, DeviceID = %d"), deviceId);
+	}
+
 	auto RemapTriggerFunction = [](const float Minimum, const float Value)
 	{
 		if(Minimum != 0.0f)
@@ -1657,6 +1771,13 @@ void FAndroidInputInterface::JoystickButtonEvent(int32 deviceId, int32 buttonId,
 	if (deviceId == -1)
 		return;
 
+	if (AndroidEnableInputDeviceListener && DeviceMapping[deviceId].DeviceState == MappingState::ToValidate)
+	{
+		MapControllerToPlayer(deviceId, EInputDeviceConnectionState::Connected);
+		DeviceMapping[deviceId].DeviceState = MappingState::Valid;
+		//UE_LOG(LogAndroid, Log, TEXT("Gamepad state change from ToValidate to Valid, DeviceID = %d"), deviceId);
+	}
+
 	//FPlatformMisc::LowLevelOutputDebugStringf(TEXT("JoystickButtonEvent[%d]: %d"), (int)DeviceMapping[deviceId].ButtonRemapping, buttonId);
 
 	if (DeviceMapping[deviceId].ControllerClass == ControllerClassType::PlaystationWireless)
diff --git a/Engine/Source/Runtime/ApplicationCore/Public/Android/AndroidInputInterface.h b/Engine/Source/Runtime/ApplicationCore/Public/Android/AndroidInputInterface.h
--- a/Engine/Source/Runtime/ApplicationCore/Public/Android/AndroidInputInterface.h
+++ b/Engine/Source/Runtime/ApplicationCore/Public/Android/AndroidInputInterface.h
@@ -81,9 +81,17 @@ enum TouchType
 	TouchEnded,
 };
 
+enum class InputDeviceStateEvent
+{
+	Added = 0,
+	Removed,
+	Changed
+};
+
 enum MappingState
 {
 	Unassigned = 0,
+	ToActivate,
 	ToValidate,
 	Valid
 };
@@ -252,6 +260,8 @@ public:
 	static bool IsControllerAssignedToGamepad(int32 ControllerId);
 	static FString GetGamepadControllerName(int32 ControllerId);
 
+	static void HandleInputDeviceStateEvent(int32 DeviceID, InputDeviceStateEvent StateEvent);
+
 	static void JoystickAxisEvent(int32 deviceId, int32 axisId, float axisValue);
 	static void JoystickButtonEvent(int32 deviceId, int32 buttonId, bool buttonDown);
 
@@ -287,6 +297,8 @@ private:
 
 	FAndroidInputInterface( const TSharedRef< FGenericApplicationMessageHandler >& InMessageHandler, const TSharedPtr< ICursor >& InCursor);
 
+	static void MapControllerToPlayer(int32 ControllerId, EInputDeviceConnectionState State);
+
 public:
 
 	FAndroidGamepadDeviceMapping* GetDeviceMapping(int32 ControllerId);
