fix for taking a long time to background

diff --git a/Engine/Source/Runtime/ApplicationCore/Private/IOS/IOSApplication.cpp b/Engine/Source/Runtime/ApplicationCore/Private/IOS/IOSApplication.cpp
--- a/Engine/Source/Runtime/ApplicationCore/Private/IOS/IOSApplication.cpp
+++ b/Engine/Source/Runtime/ApplicationCore/Private/IOS/IOSApplication.cpp
@@ -152,14 +152,12 @@ UIEdgeInsets CachedInsets;
 
 void FDisplayMetrics::RebuildDisplayMetrics(FDisplayMetrics& OutDisplayMetrics)
 {
-	__block FPlatformRect Rect;
-	__block CGFloat ContentScaleFactor;
-	dispatch_sync(dispatch_get_main_queue(), ^{
-		Rect = FIOSWindow::GetScreenRect();
-		const FIOSView *View = [[IOSAppDelegate GetDelegate] IOSView];
-		[View CalculateContentScaleFactor:Rect.Right ScreenHeight:Rect.Bottom];
-		ContentScaleFactor = View.contentScaleFactor;
-	});
+	FPlatformRect Rect;
+	CGFloat ContentScaleFactor;
+	Rect = FIOSWindow::GetScreenRect();
+	const FIOSView *View = [[IOSAppDelegate GetDelegate] IOSView];
+	[View CalculateContentScaleFactor:Rect.Right ScreenHeight:Rect.Bottom];
+	ContentScaleFactor = View.contentScaleFactor;
 	
 	// Get screen rect
 	OutDisplayMetrics.PrimaryDisplayWorkAreaRect = Rect;
@@ -215,10 +213,8 @@ void FDisplayMetrics::RebuildDisplayMetrics(FDisplayMetrics& OutDisplayMetrics)
 void FIOSApplication::CacheDisplayMetrics()
 {
 #if !PLATFORM_TVOS
-	dispatch_sync(dispatch_get_main_queue(), ^{
-		CachedInsets = [[[[UIApplication sharedApplication] delegate] window] safeAreaInsets];
-		CachedOrientation = [[[[[UIApplication sharedApplication] delegate] window] windowScene] interfaceOrientation];
-	});
+	CachedInsets = [[[[UIApplication sharedApplication] delegate] window] safeAreaInsets];
+	CachedOrientation = [[[[[UIApplication sharedApplication] delegate] window] windowScene] interfaceOrientation];
 #endif
 }
 
diff --git a/Engine/Source/Runtime/ApplicationCore/Private/IOS/IOSView.cpp b/Engine/Source/Runtime/ApplicationCore/Private/IOS/IOSView.cpp
--- a/Engine/Source/Runtime/ApplicationCore/Private/IOS/IOSView.cpp
+++ b/Engine/Source/Runtime/ApplicationCore/Private/IOS/IOSView.cpp
@@ -19,6 +19,14 @@
 #include "IOS/Accessibility/IOSAccessibilityElement.h"
 #endif
 
+static int32 GIOSMetalCacheLayer = 1;
+static FAutoConsoleVariableRef CVarIOSMetalCacheLayer(
+	TEXT("ios.Metal.CacheLayer"),
+	GIOSMetalCacheLayer,
+	TEXT("Whether to cache Metal layer instance. (default = 1)\n"),
+	ECVF_ReadOnly
+);
+
 namespace MTL
 {
     class Device;
@@ -187,6 +195,11 @@ MTL::Device* GMetalDevice = nullptr;
 		MetalLayer.pixelFormat = MTLPixelFormatBGRA8Unorm;
 		MetalLayer.framebufferOnly = NO;
 		
+		if (GIOSMetalCacheLayer > 0)
+		{
+			CachedMetalLayer = MetalLayer;
+		}
+
 		NSLog(@"::: Created a UIView that will support Metal :::");
 	}
 	
@@ -339,9 +352,17 @@ MTL::Device* GMetalDevice = nullptr;
 - (id<CAMetalDrawable>)MakeDrawable
 {
 	__block CAMetalLayer* MetalLayer = nil;
-	dispatch_sync(dispatch_get_main_queue(), ^{
-		MetalLayer = (CAMetalLayer*)[self layer];
-	});
+	if (GIOSMetalCacheLayer > 0)
+	{
+		MetalLayer = CachedMetalLayer;
+		check(MetalLayer);
+	}
+	else
+	{
+		dispatch_sync(dispatch_get_main_queue(), ^ {
+			MetalLayer = (CAMetalLayer*)[self layer];
+			});
+	}
 	// this call cannot be made on the MainThread
 	// thus requiring the code above to MainThreadCall here
 	return [MetalLayer nextDrawable];
diff --git a/Engine/Source/Runtime/ApplicationCore/Public/IOS/IOSView.h b/Engine/Source/Runtime/ApplicationCore/Public/IOS/IOSView.h
--- a/Engine/Source/Runtime/ApplicationCore/Public/IOS/IOSView.h
+++ b/Engine/Source/Runtime/ApplicationCore/Public/IOS/IOSView.h
@@ -62,6 +62,8 @@ APPLICATIONCORE_API
 	MTL::Device* MetalDevice;
 #endif
 	id<CAMetalDrawable> PanicDrawable;
+	// Might be null if caching is disabled.
+	CAMetalLayer* CachedMetalLayer;
 
 	//// KEYBOARD MEMBERS
 	
