iOS Universal link support

diff --git a/Engine/Source/Runtime/ApplicationCore/Private/IOS/IOSAppDelegate.cpp b/Engine/Source/Runtime/ApplicationCore/Private/IOS/IOSAppDelegate.cpp
--- a/Engine/Source/Runtime/ApplicationCore/Private/IOS/IOSAppDelegate.cpp
+++ b/Engine/Source/Runtime/ApplicationCore/Private/IOS/IOSAppDelegate.cpp
@@ -27,6 +27,7 @@
 #include "Algo/AllOf.h"
 #include "Misc/App.h"
 #include "Misc/EmbeddedCommunication.h"
+#include "Containers/Ticker.h"
 #if USE_MUTE_SWITCH_DETECTION
 #include "SharkfoodMuteSwitchDetector.h"
 #include "SharkfoodMuteSwitchDetector.m"
@@ -48,6 +49,9 @@
 #define GAME_THREAD_STACK_SIZE 16 * 1024 * 1024
 #endif
 
+#ifndef WITH_IOS_UNIVERSAL_LINKS
+#define WITH_IOS_UNIVERSAL_LINKS 0
+#endif
 
 
 DEFINE_LOG_CATEGORY(LogIOSAudioSession);
@@ -1135,6 +1139,54 @@ static FAutoConsoleVariableRef CVarGEnableThermalsReport(
 	return YES;
 }
 
+#if WITH_IOS_UNIVERSAL_LINKS
+static void HandleUniversalLink(const FString& UniversalLink)
+{
+	static FString ProtocolActivationUri;
+	static FTSTicker::FDelegateHandle ProtocolActivationTickerHandle;
+	
+	ProtocolActivationUri = UniversalLink;
+
+	if (!ProtocolActivationUri.IsEmpty())
+	{
+		if (!ProtocolActivationTickerHandle.IsValid())
+		{
+			ProtocolActivationTickerHandle = FTSTicker::GetCoreTicker().AddTicker(FTickerDelegate::CreateLambda(
+				[](float)
+				{
+					if (FCoreDelegates::OnActivatedByProtocol.IsBound())
+					{
+						FCoreDelegates::OnActivatedByProtocol.Broadcast(ProtocolActivationUri, PLATFORMUSERID_NONE);
+
+						ProtocolActivationUri.Reset();
+						ProtocolActivationTickerHandle.Reset();
+						return false; // remove from ticker
+					}
+
+					return true; // try again next tick
+				}
+			));
+		}
+	}
+	else if (ProtocolActivationTickerHandle.IsValid())
+	{
+		FTSTicker::GetCoreTicker().RemoveTicker(ProtocolActivationTickerHandle);
+		ProtocolActivationTickerHandle.Reset();
+	}
+}
+
+-(BOOL)application:(UIApplication*)application continueUserActivity:(NSUserActivity*)userActivity restorationHandler:(void(^)(NSArray<id<UIUserActivityRestoring>>* _Nullable)) restorationHandler 
+{
+	if ([userActivity.activityType isEqualToString:NSUserActivityTypeBrowsingWeb])
+	{
+		NSString* UniversalLink = userActivity.webpageURL.absoluteString;
+		HandleUniversalLink(FString(UniversalLink));
+		return YES;
+	}
+	return NO;
+}
+#endif // WITH_IOS_UNIVERSAL_LINKS
+
 #if !PLATFORM_TVOS
 - (UIInterfaceOrientationMask)application:(UIApplication *)application supportedInterfaceOrientationsForWindow:(UIWindow*)window
 {
