Fix class name for GameControllerManager in case of ASIS plugin usage

#rnx

diff --git a/Engine/Source/Runtime/ApplicationCore/Private/Android/AndroidInputInterface.cpp b/Engine/Source/Runtime/ApplicationCore/Private/Android/AndroidInputInterface.cpp
--- a/Engine/Source/Runtime/ApplicationCore/Private/Android/AndroidInputInterface.cpp
+++ b/Engine/Source/Runtime/ApplicationCore/Private/Android/AndroidInputInterface.cpp
@@ -93,13 +93,9 @@ bool AndroidThunkCpp_InitGameControllerManager()
 {
 	if (JNIEnv* Env = FAndroidApplication::GetJavaEnv())
 	{
-#if USE_ANDROID_STANDALONE
-		constexpr const ANSICHAR ClassName[] = "com/epicgames/makeaar/GameControllerManager";
-		constexpr const ANSICHAR ClassSig[] = "Lcom/epicgames/makeaar/GameControllerManager;";
-#else
 		constexpr const ANSICHAR ClassName[] = "com/epicgames/unreal/GameControllerManager";
 		constexpr const ANSICHAR ClassSig[] = "Lcom/epicgames/unreal/GameControllerManager;";
-#endif
+
 		jmethodID CreateGameControllerManagerMethodID = FJavaWrapper::FindMethod(Env, FJavaWrapper::GameActivityClassID, "createGameControllerManager", "()Z", false);
 		bool bGCMCreated = FJavaWrapper::CallBooleanMethod(Env, FJavaWrapper::GameActivityThis, CreateGameControllerManagerMethodID);
 		if (bGCMCreated)
