Enable safe zone CVar overwrites for Xbox

#rnx
#tests Tested on Xcloud DevKits and on Windows streaming platforms

diff --git a/Engine/Source/Runtime/ApplicationCore/Private/GenericPlatform/GenericApplication.cpp b/Engine/Source/Runtime/ApplicationCore/Private/GenericPlatform/GenericApplication.cpp
--- a/Engine/Source/Runtime/ApplicationCore/Private/GenericPlatform/GenericApplication.cpp
+++ b/Engine/Source/Runtime/ApplicationCore/Private/GenericPlatform/GenericApplication.cpp
@@ -87,6 +87,12 @@ const FInputDeviceScope* FInputDeviceScope::GetCurrent()
 float GDebugSafeZoneRatio = 1.0f;
 float GDebugActionZoneRatio = 1.0f;
 
+static TAutoConsoleVariable<bool> CVarEnableSafeZoneOverrides(TEXT("SafeZone.EnableOverrides"), false, TEXT("Enables usage of individual SafeZone.Ratio CVars"));
+static TAutoConsoleVariable<float> CVarSafeZoneRatioLeft(TEXT("SafeZone.Ratio.Left"), 1.0f, TEXT("Safe Zone Ratio - Left"));
+static TAutoConsoleVariable<float> CVarSafeZoneRatioTop(TEXT("SafeZone.Ratio.Top"), 1.0f, TEXT("Safe Zone Ratio - Top"));
+static TAutoConsoleVariable<float> CVarSafeZoneRatioRight(TEXT("SafeZone.Ratio.Right"), 1.0f, TEXT("Safe Zone Ratio - Right"));
+static TAutoConsoleVariable<float> CVarSafeZoneRatioBottom(TEXT("SafeZone.Ratio.Bottom"), 1.0f, TEXT("Safe Zone Ratio - Bottom"));
+
 struct FSafeZoneConsoleVariables
 {
 	FAutoConsoleVariableRef DebugSafeZoneRatioCVar;
@@ -135,6 +141,21 @@ float FDisplayMetrics::GetDebugTitleSafeZoneRatio()
 	return GDebugSafeZoneRatio;
 }
 
+bool FDisplayMetrics::TryGetTitleSafeZoneOverwrite(FVector4& InOutTitleSafePaddingSize)
+{
+	if (CVarEnableSafeZoneOverrides.GetValueOnGameThread())
+	{
+		InOutTitleSafePaddingSize.X = (float)PrimaryDisplayWidth * (1.0f - CVarSafeZoneRatioLeft.GetValueOnGameThread()) * 0.5f;
+		InOutTitleSafePaddingSize.Y = (float)PrimaryDisplayHeight * (1.0f - CVarSafeZoneRatioTop.GetValueOnGameThread()) * 0.5f;
+		InOutTitleSafePaddingSize.Z = (float)PrimaryDisplayWidth * (1.0f - CVarSafeZoneRatioRight.GetValueOnGameThread()) * 0.5f;
+		InOutTitleSafePaddingSize.W = (float)PrimaryDisplayHeight * (1.0f - CVarSafeZoneRatioBottom.GetValueOnGameThread()) * 0.5f;
+
+		return true;
+	}
+
+	return false;
+}
+
 float FDisplayMetrics::GetDebugActionSafeZoneRatio()
 {
 	return GDebugActionZoneRatio;
@@ -142,15 +163,9 @@ float FDisplayMetrics::GetDebugActionSafeZoneRatio()
 
 void FDisplayMetrics::ApplyDefaultSafeZones()
 {
-	// Allow safe zones to be overridden by the command line. Used by mobile PIE.
+	// Allow safe zones to be overridden by CVars. Used by streaming and mobile PIE.
 	TitleSafePaddingSize = FVector4(0.0f, 0.0f, 0.0f, 0.0f);
-	bool bSetByCommandLine;
-	bSetByCommandLine = FParse::Value(FCommandLine::Get(), TEXT("SafeZonePaddingLeft="),   TitleSafePaddingSize.X);
-	bSetByCommandLine = FParse::Value(FCommandLine::Get(), TEXT("SafeZonePaddingTop="),  TitleSafePaddingSize.Y) || bSetByCommandLine;
-	bSetByCommandLine = FParse::Value(FCommandLine::Get(), TEXT("SafeZonePaddingRight="),    TitleSafePaddingSize.Z) || bSetByCommandLine;
-	bSetByCommandLine = FParse::Value(FCommandLine::Get(), TEXT("SafeZonePaddingBottom="), TitleSafePaddingSize.W) || bSetByCommandLine;
-
-	if (!bSetByCommandLine)
+	if (!TryGetTitleSafeZoneOverwrite(TitleSafePaddingSize))
 	{
 		const float SafeZoneRatio = GetDebugTitleSafeZoneRatio();
 		if (SafeZoneRatio < 1.0f)
@@ -159,7 +174,7 @@ void FDisplayMetrics::ApplyDefaultSafeZones()
 			TitleSafePaddingSize = FVector4((float)PrimaryDisplayWidth * HalfUnsafeRatio, (float)PrimaryDisplayHeight * HalfUnsafeRatio, (float)PrimaryDisplayWidth * HalfUnsafeRatio, (float)PrimaryDisplayHeight * HalfUnsafeRatio);
 		}
 	}
-
+	
 	const float ActionSafeZoneRatio = GetDebugActionSafeZoneRatio();
 	if (ActionSafeZoneRatio < 1.0f)
 	{
diff --git a/Engine/Source/Runtime/ApplicationCore/Public/GenericPlatform/GenericApplication.h b/Engine/Source/Runtime/ApplicationCore/Public/GenericPlatform/GenericApplication.h
--- a/Engine/Source/Runtime/ApplicationCore/Public/GenericPlatform/GenericApplication.h
+++ b/Engine/Source/Runtime/ApplicationCore/Public/GenericPlatform/GenericApplication.h
@@ -402,6 +402,9 @@ struct FDisplayMetrics
 	// The title safe zone ratio that will be returned by FDisplayMetrics::GetDisplayMetrics on platforms that don't have a defined safe zone
 	APPLICATIONCORE_API static float GetDebugTitleSafeZoneRatio();
 protected:
+	// Fills the InOutTitleSafePaddingSize with CVar overwrites (Sets padding to 0 for values that are not set).
+	bool TryGetTitleSafeZoneOverwrite(FVector4& InOutTitleSafePaddingSize);
+
 	// The action safe zone ratio that will be returned by FDisplayMetrics::GetDisplayMetrics on platforms that don't have a defined safe zone
 	static float GetDebugActionSafeZoneRatio();
 
