[ASIS] Do not apply a scale factor to get screen rect in case of override Android window

#rnx

diff --git a/Engine/Source/Runtime/ApplicationCore/Private/Android/AndroidWindow.cpp b/Engine/Source/Runtime/ApplicationCore/Private/Android/AndroidWindow.cpp
--- a/Engine/Source/Runtime/ApplicationCore/Private/Android/AndroidWindow.cpp
+++ b/Engine/Source/Runtime/ApplicationCore/Private/Android/AndroidWindow.cpp
@@ -621,7 +621,7 @@ FPlatformRect FAndroidWindow::GetScreenRect(bool bUseEventThreadWindow)
 		const bool bMobileHDR = (MobileHDRCvar && MobileHDRCvar->GetValueOnAnyThread() == 1);
 		UE_LOG(LogAndroid, Log, TEXT("Mobile HDR: %s"), bMobileHDR ? TEXT("YES") : TEXT("no"));
 
-		if (!bIsOculusMobileApp)
+		if (!bIsOculusMobileApp && (GAndroidWindowOverride == nullptr))
 		{
 			if (CurrentParams.WindowDPI && RequestedResX == 0 && RequestedResY == 0)
 			{
