[Engine][IOS] Added new bEnableVirtualKeyboardVisibilityEvent config to enable VirtualKeyboard Shown/Hidden event to be broadcasted regardless of the IntegratedKeyboard config.

diff --git a/Engine/Source/Runtime/ApplicationCore/Private/IOS/IOSView.cpp b/Engine/Source/Runtime/ApplicationCore/Private/IOS/IOSView.cpp
--- a/Engine/Source/Runtime/ApplicationCore/Private/IOS/IOSView.cpp
+++ b/Engine/Source/Runtime/ApplicationCore/Private/IOS/IOSView.cpp
@@ -1019,10 +1019,13 @@ self.accessibilityElements = @[Window.accessibilityContainer];
 	
 	bool bUseIntegratedKeyboard = false;
 	GConfig->GetBool(TEXT("/Script/IOSRuntimeSettings.IOSRuntimeSettings"), TEXT("bUseIntegratedKeyboard"), bUseIntegratedKeyboard, GEngineIni);
+
+	bool bEnableVirtualKeyboardVisibilityEvent = false;
+	GConfig->GetBool(TEXT("/Script/IOSRuntimeSettings.IOSRuntimeSettings"), TEXT("bEnableVirtualKeyboardVisibilityEvent"), bEnableVirtualKeyboardVisibilityEvent, GEngineIni);
 	
 	// get notifications when the keyboard is in view
 	bIsUsingIntegratedKeyboard = FParse::Param(FCommandLine::Get(), TEXT("NewKeyboard")) || bUseIntegratedKeyboard;
-	if (bIsUsingIntegratedKeyboard)
+	if (bEnableVirtualKeyboardVisibilityEvent || bIsUsingIntegratedKeyboard)
 	{
 		[[NSNotificationCenter defaultCenter] addObserver:self
 												 selector:@selector(keyboardWasShown:)
@@ -1031,7 +1034,6 @@ self.accessibilityElements = @[Window.accessibilityContainer];
 		[[NSNotificationCenter defaultCenter] addObserver:self
 												 selector:@selector(keyboardWillBeHidden:)
 													 name:UIKeyboardWillHideNotification object:nil];
-		
 	}
 #endif
 }
diff --git a/Engine/Source/Runtime/IOS/IOSRuntimeSettings/Classes/IOSRuntimeSettings.h b/Engine/Source/Runtime/IOS/IOSRuntimeSettings/Classes/IOSRuntimeSettings.h
--- a/Engine/Source/Runtime/IOS/IOSRuntimeSettings/Classes/IOSRuntimeSettings.h
+++ b/Engine/Source/Runtime/IOS/IOSRuntimeSettings/Classes/IOSRuntimeSettings.h
@@ -521,6 +521,10 @@ public:
 	UPROPERTY(GlobalConfig, EditAnywhere, Category = Input)
 	bool bUseIntegratedKeyboard;
 
+	// Whether or not we should receive the keyboard's visibility event (Shown/Hidden) regardless of the IntegratedKeyboard being activated or not.
+	UPROPERTY(GlobalConfig, EditAnywhere, Category = Input)
+	bool bEnableVirtualKeyboardVisibilityEvent;
+
 	/** Sample rate to run the audio mixer with. */
 	UPROPERTY(config, EditAnywhere, Category = "Audio", Meta = (DisplayName = "Audio Mixer Sample Rate"))
 	int32 AudioSampleRate;
