ForceMaxTouchpadSensitivity
-The problem is that the default touchpad sensitivity setting on windows is designed to reject accidental touchpad touches while typing.  This is sensible for most apps, but not games where the touchpad and keyboard can be used simultaneously.
-We have two solutions for this.  In the long run microsoft is adding a new api to force the sensitivity to maximum for the current process.  However it is not yet available.  This change has stubbed out code to support that when it arrives.  In the short run there is a function to set system parameters available for windows 10 22h2 on.  Versions before windows 10 22h2 we do nothing.
-Added ApplyForceMaxTouchpadSensitivity() and RemoveForceMaxTouchpadSensitivity() to FWindowsApplication to change the sensitivity.
-Because the system call can take 1 second to complete I added a parameter to ApplyForceMaxTouchpadSensitivity which will allow the function to do the long call in a task worker thread.  All the calls will wait for this task to complete, if there is one running.  Currently only the async Apply exists, in order to avoid slowing down app startup and tab-in.
-Added Cvars WindowsApplication.ForceMaxTouchpadSensitivity and WindowsApplication.ForceMaxTouchpadSensitivityEditor which will call those functions (using the async apply) at startup and when the app gains/loses focus and on exit.  These default to off.
-One can simply set the cvar to have this feature work automatically.  To activate it at runtime I would suggest setting the cvar to handle the focus/defocus and shutdown behavior and then calling ApplyForceMaxTouchpadSensitivity(AllowAsynchronous) to immediately change the setting.  To deactivate at runtime set the cvar to false, then call RemoveForceMaxTouchpadSensitivity.

diff --git a/Engine/Source/Runtime/ApplicationCore/Private/Windows/WindowsApplication.cpp b/Engine/Source/Runtime/ApplicationCore/Private/Windows/WindowsApplication.cpp
--- a/Engine/Source/Runtime/ApplicationCore/Private/Windows/WindowsApplication.cpp
+++ b/Engine/Source/Runtime/ApplicationCore/Private/Windows/WindowsApplication.cpp
@@ -20,6 +20,7 @@
 #include "Null/NullPlatformApplicationMisc.h"
 #include "Stats/Stats.h"
 #include "HAL/IConsoleManager.h"
+#include "Tasks/Task.h"
 
 #if WITH_EDITOR
 #include "Modules/ModuleManager.h"
@@ -42,6 +43,7 @@ THIRD_PARTY_INCLUDES_START
 #include <cfgmgr32.h>
 #include <windowsx.h>
 #include <VersionHelpers.h>
+#include "winuser.h"
 THIRD_PARTY_INCLUDES_END
 
 // Platform code uses IsMaximized which is defined to IsZoomed by windowsx.h
@@ -64,6 +66,482 @@ FAutoConsoleVariableRef	CVarPreventDuplicateMouseEventsForTouch(
 	TEXT("Hack to get around multiple mouse events being triggered for touch events on Windows 7 and lower.  Enabling this will prevent pen tablets from working on windows 7 since until we switch to the windows 8 sdk (and can use WM_POINTER* events) we cannot detect the difference")
 );
 
+namespace UE::WindowsApplication::TouchpadSensitivity
+{
+	DECLARE_LOG_CATEGORY_EXTERN(LogWindowsTouchpad, Log, VeryVerbose);
+	DEFINE_LOG_CATEGORY(LogWindowsTouchpad);
+
+	static bool bForceMaxTouchpadSensitivity = false;
+	FAutoConsoleVariableRef	CVarForceMaxTouchpadSensitivity(
+		TEXT("WindowsApplication.ForceMaxTouchpadSensitivity"),
+		bForceMaxTouchpadSensitivity,
+		TEXT("Force Touchpad sensitivity to maximum in game when the app is active (aka has focus)."
+			" The sensitivity setting can reject touchpad touches while typing on the theory that they are accidental, but that is not a good assumption for games."
+			" The setting defaults to 'Medium'. The 'Most sensitive' setting does not reject any touches, so that is what we want.."
+			" On Windows10 22H2 we force the global system setting to 'Most Sensitive', but do not write that to the registry persistently."
+			" We will set it back when the app loses focus and on exit, but it may not happen in the case of a crash leaving the setting altered until a reboot."
+/*			" On Window11 (24H2? CHECKVERSION NUMBER)+ a new api, ForceMaxTouchpadSensitivityForCurrentProcess, exists and we use that."*/
+		)
+	);
+
+	enum class EForceMaxTouchpadSensitivityMode
+	{
+		NotCached,
+		DoNothing,
+		SetTouchpadParameters,
+		ForceMaxTouchpadSensitivityForCurrentProcess
+	};
+
+	static EForceMaxTouchpadSensitivityMode CachedForceMaxTouchpadSensitivityMode = EForceMaxTouchpadSensitivityMode::NotCached;
+
+	void CacheForceMaxTouchpadSensitivityMode()
+	{
+		// Do nothing if we already cached it.
+		if (CachedForceMaxTouchpadSensitivityMode != EForceMaxTouchpadSensitivityMode::NotCached)
+		{
+			return;
+		}
+
+		// Need to compile switch this one on the windows sdk that has the api call, I assume.
+#if 0
+		if (FPlatformMisc::VerifyWindowsVersion(10, 0, 26100)) //WINDOWS 11 24H2+ (or whatever version ForceMaxTouchpadSensitivityForCurrentProcess appears in, this might be correct)
+		{
+			CachedForceMaxTouchpadSensitivityMode = EForceMaxTouchpadSensitivityMode::ForceMaxTouchpadSensitivityForCurrentProcess;
+			UE_LOG(LogWindowsTouchpad, Log, TEXT("CacheForceMaxTouchpadSensitivityMode ForceMaxTouchpadSensitivityForCurrentProcess"));
+		}
+		else
+#endif
+		if (FPlatformMisc::VerifyWindowsVersion(10, 0, 19045)) // WINDOWS 10 22H2+
+		{
+			CachedForceMaxTouchpadSensitivityMode = EForceMaxTouchpadSensitivityMode::SetTouchpadParameters;
+			UE_LOG(LogWindowsTouchpad, Log, TEXT("CacheForceMaxTouchpadSensitivityMode SetTouchpadParameters"));
+		}
+		else
+		{
+			CachedForceMaxTouchpadSensitivityMode = EForceMaxTouchpadSensitivityMode::DoNothing;
+			UE_LOG(LogWindowsTouchpad, Log, TEXT("CacheForceMaxTouchpadSensitivityMode DoNothing"));
+		}
+	}
+
+	EForceMaxTouchpadSensitivityMode GetForceMaxTouchpadSensitivityMode()
+	{
+		check(CachedForceMaxTouchpadSensitivityMode != EForceMaxTouchpadSensitivityMode::NotCached);
+
+		UE_LOG(LogWindowsTouchpad, Verbose, TEXT("GetForceMaxTouchpadSensitivityMode %i"), CachedForceMaxTouchpadSensitivityMode);
+		return CachedForceMaxTouchpadSensitivityMode;
+	}
+
+	// Define constants, enums, and structs so that we can use the SetTouchpadParameters api on versions of windows
+	// where it was not yet public.
+
+	// In Windows 10, version 22H2, and later versions, "version 0" of the TOUCHPAD_PARAMETERS
+	// settings structure can be used to query and update the user's touchpad settings.
+
+	// IMPORTANT: In some later versions, new settings were introduced to this structure. These
+	// settings are omitted from the below definition, since there is no way to distinguish when
+	// retrieving the settings whether they are supported by that OS version. To ensure that these
+	// settings are preserved when updating other settings, the structure  passed with
+	// SPI_SETTOUCHPADPARAMETERS should first be populated using SPI_GETTOUCHPADPARAMETERS, with the
+	// desired settings updates then applied to that structure.
+
+	// The following are from the 24H2 winuser.h public header.
+
+	const UINT SPI_GETTOUCHPADPARAMETERS = 0x00AE;
+	const UINT SPI_SETTOUCHPADPARAMETERS = 0x00AF;
+
+	typedef enum LEGACY_TOUCHPAD_FEATURES {
+		LEGACY_TOUCHPAD_FEATURE_NONE = 0x00000000,
+		LEGACY_TOUCHPAD_FEATURE_ENABLE_DISABLE = 0x00000001,
+		LEGACY_TOUCHPAD_FEATURE_REVERSE_SCROLL_DIRECTION = 0x00000004,
+	} LEGACY_TOUCHPAD_FEATURES;
+
+	typedef enum TOUCHPAD_SENSITIVITY_LEVEL {
+		TOUCHPAD_SENSITIVITY_LEVEL_MOST_SENSITIVE = 0x00000000,
+		TOUCHPAD_SENSITIVITY_LEVEL_HIGH_SENSITIVITY = 0x00000001,
+		TOUCHPAD_SENSITIVITY_LEVEL_MEDIUM_SENSITIVITY = 0x00000002,
+		TOUCHPAD_SENSITIVITY_LEVEL_LOW_SENSITIVITY = 0x00000003,
+		TOUCHPAD_SENSITIVITY_LEVEL_LEAST_SENSITIVE = 0x00000004,
+	} TOUCHPAD_SENSITIVITY_LEVEL;
+	static_assert((int32)FWindowsApplication::EForceMaxTouchpadSensitivityRestorationValues::TOUCHPAD_SENSITIVITY_LEVEL_MOST_SENSITIVE == (int32)TOUCHPAD_SENSITIVITY_LEVEL::TOUCHPAD_SENSITIVITY_LEVEL_MOST_SENSITIVE);
+	static_assert((int32)FWindowsApplication::EForceMaxTouchpadSensitivityRestorationValues::TOUCHPAD_SENSITIVITY_LEVEL_LEAST_SENSITIVE == (int32)TOUCHPAD_SENSITIVITY_LEVEL::TOUCHPAD_SENSITIVITY_LEVEL_LEAST_SENSITIVE);
+
+	// This is the v0 structure for getting/setting touchpad settings prior to Win11 24H2.
+	typedef struct tagTOUCHPAD_PARAMETERS_V0 {
+		UINT versionNumber;
+		LEGACY_TOUCHPAD_FEATURES legacyTouchpadFeatures;
+		TOUCHPAD_SENSITIVITY_LEVEL sensitivityLevel;
+		BOOL touchpadEnabled : 1;
+		BOOL Reserved1 : 1;
+		BOOL Reserved2 : 1;
+		BOOL scrollDirectionReversed : 1;
+		BOOL legacyTouchpadPresent : 1;
+		BOOL Reserved3 : 1;
+		BOOL touchpadPresent : 1;
+		BOOL panEnabled : 1;
+		BOOL zoomEnabled : 1;
+		BOOL tapEnabled : 1;
+		BOOL twoFingerTapEnabled : 1;
+		BOOL Reserved4 : 21;
+		UINT cursorSpeed;
+		BOOL allowActiveWhenMousePresent : 1;
+		BOOL rightClickZoneEnabled : 1;
+		BOOL tapAndDragEnabled : 1;
+		BOOL touchpadActive : 1;
+		BOOL externalMousePresent : 1;
+		BOOL fourFingersSupported : 1;
+		BOOL mouseAccelSettingHonored : 1;
+		BOOL Reserved5 : 25;
+	} TOUCHPAD_PARAMETERS_V0, * PTOUCHPAD_PARAMETERS_V0;
+	static_assert(sizeof(TOUCHPAD_PARAMETERS_V0) == 0x18);
+
+	static TOptional<TOUCHPAD_SENSITIVITY_LEVEL> ForceMaxTouchpadSensitivity_SavedSensitivity;
+
+	FWindowsApplication::EForceMaxTouchpadSensitivityRestorationValues GetForceMaxTouchpadSensitivityRestorationValue()
+	{
+		CacheForceMaxTouchpadSensitivityMode();
+
+		if (GetForceMaxTouchpadSensitivityMode() == EForceMaxTouchpadSensitivityMode::SetTouchpadParameters)
+		{
+			// If we have a cached value return that.
+			if (ForceMaxTouchpadSensitivity_SavedSensitivity.IsSet())
+			{
+				UE_LOG(LogWindowsTouchpad, Verbose, TEXT("GetForceMaxTouchpadSensitivityRestorationValue returning SavedSensitivity %i"), ForceMaxTouchpadSensitivity_SavedSensitivity.GetValue());
+				return static_cast<FWindowsApplication::EForceMaxTouchpadSensitivityRestorationValues>(ForceMaxTouchpadSensitivity_SavedSensitivity.GetValue());
+			}
+			else
+			{
+				// If no cached value return the current system setting.
+				// Get the user's touchpad settings.
+				TOUCHPAD_PARAMETERS_V0 params = {};
+				params.versionNumber = 0;
+				BOOL bSuccess = ::SystemParametersInfo(SPI_GETTOUCHPADPARAMETERS, sizeof(params), &params, 0 /*fWinIni*/);
+				if (!bSuccess)
+				{
+					UE_LOG(LogWindowsTouchpad, Warning, TEXT("GetForceMaxTouchpadSensitivityRestorationValue SystemParametersInfo SPI_GETTOUCHPADPARAMETERS failed"));
+					return FWindowsApplication::EForceMaxTouchpadSensitivityRestorationValues::SystemCallFailed;
+				}
+
+				if (params.sensitivityLevel == TOUCHPAD_SENSITIVITY_LEVEL_MOST_SENSITIVE)
+				{
+					UE_LOG(LogWindowsTouchpad, Verbose, TEXT("GetForceMaxTouchpadSensitivityRestorationValue returning RestorationNotNeeded because the system setting was already MostSensitive."));
+					return FWindowsApplication::EForceMaxTouchpadSensitivityRestorationValues::RestorationNotNeeded;
+				}
+				else
+				{
+					UE_LOG(LogWindowsTouchpad, Verbose, TEXT("GetPreExistingTouchpadSensitivity returning Current System Sensitivity %i"), params.sensitivityLevel);
+					return static_cast<FWindowsApplication::EForceMaxTouchpadSensitivityRestorationValues>(params.sensitivityLevel);
+				}
+			}
+		}
+		else
+		{
+			UE_LOG(LogWindowsTouchpad, Verbose, TEXT("GetForceMaxTouchpadSensitivityRestorationValue returning RestorationNotNeeded"));
+			return FWindowsApplication::EForceMaxTouchpadSensitivityRestorationValues::RestorationNotNeeded;
+		}
+	}
+
+	/*
+	* The multithreading model here is very constrained.  Only the system call in Apply-Async is running in a worker thread.  All of the
+	* rest runs in the game thread and all of the state within unreal is updated in the game thread. All of the
+	* apply/remove functions wait on the system call task if it is in progress. This means that we can call Apply-Async and not end up blocking 
+	* for over a second.  But if we make any other call, including another Apply-Async while the task is running we will block for the task and then
+	* block for the system call.  One should keep in mind that the system setting can theoretically change even outside of unreal, though it is probably 
+	* uncommon.
+	* 
+	* The assumption here is that blocking during startup is bad, blocking during shutdown is ok, and these functions should be called infrequently so
+	* blocking on the task will be very rare in reality.
+	*/
+	static UE::Tasks::TTask<void> ApplyForceMaxTouchpadSensitivityAsyncTask;
+	void WaitOnPreviousApplyAsync()
+	{
+		// If we are using the async apply the task could still be outstanding.  Wait until it is done to keep the state sane.
+		// Timeout to make sure we don't block forever if this all goes wrong.
+		if (ApplyForceMaxTouchpadSensitivityAsyncTask.IsValid() && !ApplyForceMaxTouchpadSensitivityAsyncTask.IsCompleted())
+		{
+			UE_LOG(LogWindowsTouchpad, Verbose, TEXT("WaitOnPreviousApplyAsync"));
+			ApplyForceMaxTouchpadSensitivityAsyncTask.Wait(FTimespan::FromSeconds(2.0));
+		}
+	}
+
+	void ApplyForceMaxTouchpadSensitivity(FWindowsApplication::EForceMaxTouchpadSensitivityAsyncBehavior AsyncBehavior)
+	{
+		check(IsInGameThread());
+
+		switch (GetForceMaxTouchpadSensitivityMode())
+		{
+		case EForceMaxTouchpadSensitivityMode::NotCached:
+			UE_LOG(LogWindowsTouchpad, Log, TEXT("ApplyForceMaxTouchpadSensitivity NotCached, this is a bug"));
+			check(false);
+			break;
+		case EForceMaxTouchpadSensitivityMode::DoNothing:
+			UE_LOG(LogWindowsTouchpad, Verbose, TEXT("ApplyForceMaxTouchpadSensitivity DoNothing"));
+			break;
+		case EForceMaxTouchpadSensitivityMode::SetTouchpadParameters:
+		{
+			if (ForceMaxTouchpadSensitivity_SavedSensitivity.IsSet())
+			{
+				UE_LOG(LogWindowsTouchpad, Verbose, TEXT("ApplyForceMaxTouchpadSensitivity sensitivity was already changed, doing nothing."));
+				return;
+			}
+			
+			// I think its actually impossible to need to wait here.  We would early out just above.  If there was an interveening remove it would have done the wait.
+			WaitOnPreviousApplyAsync();
+
+			// Set the global system runtime touchpad sensitivity to MostSensitive
+			// Note: we will get the params, change just the one, then write them back so that the reserved param block is preserved.
+			UE_LOG(LogWindowsTouchpad, Log, TEXT("ApplyForceMaxTouchpadSensitivity setting to MostSensitive"));
+			// Get the user's touchpad settings.
+			TOUCHPAD_PARAMETERS_V0 params = {};
+			params.versionNumber = 0;
+			BOOL bSuccess = ::SystemParametersInfo(SPI_GETTOUCHPADPARAMETERS, sizeof(params), &params, 0 /*fWinIni*/);
+			if (!bSuccess)
+			{
+				UE_LOG(LogWindowsTouchpad, Warning, TEXT("ApplyForceMaxTouchpadSensitivity SystemParametersInfo SPI_GETTOUCHPADPARAMETERS failed"));
+				return;
+			}
+
+			if (params.sensitivityLevel == TOUCHPAD_SENSITIVITY_LEVEL_MOST_SENSITIVE)
+			{
+				// If the setting is already MOST_SENSTIVE we will do nothing.
+				// This will also stop us overwriting the SavedSensitivity if we call this function twice.
+				UE_LOG(LogWindowsTouchpad, Verbose, TEXT("ApplyForceMaxTouchpadSensitivity sensitivity was already Most Sensitive"));
+				return;
+			}
+			else
+			{
+				ForceMaxTouchpadSensitivity_SavedSensitivity = params.sensitivityLevel;
+			}
+
+			// Update the touchpad sensitivity level to "most sensitive" so that simultaneous
+			// keyboard input does not interfere with the touchpad's mouse input generation.
+			// Leave all other bytes in the struct untouched.
+			params.sensitivityLevel = TOUCHPAD_SENSITIVITY_LEVEL_MOST_SENSITIVE;
+
+			if (AsyncBehavior == FWindowsApplication::EForceMaxTouchpadSensitivityAsyncBehavior::AllowAsynchronous)
+			{
+				UE_LOG(LogWindowsTouchpad, Log, TEXT("ApplyForceMaxTouchpadSensitivity Task launching."));
+				ApplyForceMaxTouchpadSensitivityAsyncTask = UE::Tasks::Launch(
+					TEXT("ApplyForceMaxTouchpadSensitivityTask"),
+					[params] 
+					{ 
+						UE_LOG(LogWindowsTouchpad, Verbose, TEXT("ApplyForceMaxTouchpadSensitivity Task running."));
+
+						// Do not specify SPIF_UPDATEINIFILE so that the registry is not updated.
+						// We will revert this update later, and not writing to the registry is a
+						// failsafe in case we fail to do so (e.g. due to a crash).
+						bool bSuccess = ::SystemParametersInfo(SPI_SETTOUCHPADPARAMETERS, sizeof(params), const_cast<TOUCHPAD_PARAMETERS_V0*>(&params), SPIF_SENDCHANGE);
+						if (!bSuccess)
+						{
+							UE_LOG(LogWindowsTouchpad, Warning, TEXT("ApplyForceMaxTouchpadSensitivityTask SystemParametersInfo SPI_SETTOUCHPADPARAMETERS Failed"));
+							return;
+						}
+						UE_LOG(LogWindowsTouchpad, Log, TEXT("ApplyForceMaxTouchpadSensitivityTask sucessfully changed setting to MostSensitive"));
+
+					},
+					UE::Tasks::ETaskPriority::Normal);
+			}
+			else
+			{
+				check(AsyncBehavior == FWindowsApplication::EForceMaxTouchpadSensitivityAsyncBehavior::RequireBlocking);
+
+				// Do not specify SPIF_UPDATEINIFILE so that the registry is not updated.
+				// We will revert this update later, and not writing to the registry is a
+				// failsafe in case we fail to do so (e.g. due to a crash).
+				bSuccess = ::SystemParametersInfo(SPI_SETTOUCHPADPARAMETERS, sizeof(params), &params, SPIF_SENDCHANGE);
+				if (!bSuccess)
+				{
+					UE_LOG(LogWindowsTouchpad, Warning, TEXT("ApplyForceMaxTouchpadSensitivity SystemParametersInfo SPI_SETTOUCHPADPARAMETERS Failed"));
+					return;
+				}
+				UE_LOG(LogWindowsTouchpad, Log, TEXT("ApplyForceMaxTouchpadSensitivity sucessfully changed setting to MostSensitive"));
+			}
+		}
+		break;
+		case EForceMaxTouchpadSensitivityMode::ForceMaxTouchpadSensitivityForCurrentProcess:
+		{
+			// Use ForceMaxTouchpadSensitivityForCurrentProcess API
+			// Need to compile switch this one on the windows sdk that has the api call, I assume.
+#if 0
+			//BOOL bSuccess = ForceMaxTouchpadSensitivityForCurrentProcess(TRUE);
+			//if (!bSuccess)
+			//{
+			//	UE_LOG(LogWindowsTouchpad, Warning, TEXT("ApplyForceMaxTouchpadSensitivity ForceMaxTouchpadSensitivityForCurrentProcess Failed"));
+			//	return;
+			//}
+#else
+			UE_LOG(LogWindowsTouchpad, Log, TEXT("ApplyForceMaxTouchpadSensitivity ForceMaxTouchpadSensitivityForCurrentProcess -- Not Yet Implemented!"));
+			check(false);
+#endif
+		}
+		break;
+		default:
+			check(false);
+			break;
+		}
+	}
+
+	void RemoveForceMaxTouchpadSensitivity()
+	{
+		check(IsInGameThread());
+
+		// Only if we did set the sensitivity using the EForceMaxTouchpadSensitivityMode::TouchpadParameters implementation
+		// will we have cached a value, and need to restore it.
+		if (ForceMaxTouchpadSensitivity_SavedSensitivity.IsSet())
+		{
+			WaitOnPreviousApplyAsync();
+
+			// Note: we will get the params, change just the one, then write them back so that the reserved param block is preserved.
+			UE_LOG(LogWindowsTouchpad, Log, TEXT("RemoveForceMaxTouchpadSensitivity reverting to SavedSensitivity %i"), ForceMaxTouchpadSensitivity_SavedSensitivity.GetValue());
+
+			// Get the user's touchpad settings.
+			TOUCHPAD_PARAMETERS_V0 params = {};
+			params.versionNumber = 0;
+			BOOL bSuccess = ::SystemParametersInfo(SPI_GETTOUCHPADPARAMETERS, sizeof(params), &params, 0 /*fWinIni*/);
+			if (!bSuccess)
+			{
+				UE_LOG(LogWindowsTouchpad, Warning, TEXT("RemoveForceMaxTouchpadSensitivity SystemParametersInfo SPI_GETTOUCHPADPARAMETERS failed"));
+				return;
+			}
+
+			// Restore the saved sensitivity level.
+			// Leave all other bytes in the struct untouched.
+			params.sensitivityLevel = ForceMaxTouchpadSensitivity_SavedSensitivity.GetValue();
+
+			// Do not specify SPIF_UPDATEINIFILE so that the registry is not updated.
+			// We will revert this update later, and not writing to the registry is a
+			// failsafe in case we fail to do so (e.g. due to a crash).
+			bSuccess = ::SystemParametersInfo(SPI_SETTOUCHPADPARAMETERS, sizeof(params), &params, SPIF_SENDCHANGE);
+			if (!bSuccess)
+			{
+				UE_LOG(LogWindowsTouchpad, Warning, TEXT("RemoveForceMaxTouchpadSensitivity SystemParametersInfo SPI_SETTOUCHPADPARAMETERS failed"));
+				return;
+			}
+			UE_LOG(LogWindowsTouchpad, Log, TEXT("RemoveForceMaxTouchpadSensitivity sucessfully reverted to SavedSensitivity %i"), ForceMaxTouchpadSensitivity_SavedSensitivity.GetValue());
+
+			ForceMaxTouchpadSensitivity_SavedSensitivity.Reset();
+		}
+		else if (GetForceMaxTouchpadSensitivityMode() == EForceMaxTouchpadSensitivityMode::ForceMaxTouchpadSensitivityForCurrentProcess)
+		{
+			// Use ForceMaxTouchpadSensitivityForCurrentProcess API
+			// Need to compile switch this one on the windows sdk that has the api call, I assume.
+#if 0
+
+			//BOOL bSuccess = ForceMaxTouchpadSensitivityForCurrentProcess(FALSE);
+			//if (!bSuccess)
+			//{
+			//	UE_LOG(LogWindowsTouchpad, Warning, TEXT("RemoveForceMaxTouchpadSensitivity ForceMaxTouchpadSensitivityForCurrentProcess Failed"));
+			//	return;
+			//}
+#else
+			UE_LOG(LogWindowsTouchpad, Log, TEXT("RemoveForceMaxTouchpadSensitivity ForceMaxTouchpadSensitivityForCurrentProcess -- Not Yet Implemented!"));
+			check(false);
+#endif	
+		}
+		else
+		{
+			UE_LOG(LogWindowsTouchpad, Verbose, TEXT("RemoveForceMaxTouchpadSensitivity Doing nothing because the sensitivity was already maximized"));
+		}
+	}
+
+	// Note: this function could be called from threads other than the main game thread.
+	// It is NOT fully safe to do so, ForceMaxTouchpadSensitivity_SavedSensitivity is not threadsafe and its SET and VALUE could be mismatched and we are not going to wait for 
+	// a potential worker-thread system call to complete.
+	// However we are in an error state already, it is very likely that it will work fine, we can do some simple checks to make sure it doesn't do anything really strange,
+	// and I want to avoid adding locks to this error case.
+	// TOptional, at the time of writing, is implemented such that a Reset does not change a saved enum value, so cases where another thread is removing the setting and this code also runs should
+	// result in a double-restoration of the saved value.  This is OK.
+	// TOptional, at the time of writing, is written so that the value is saved before the Set flag is set however if it does actually run the operations in the reverse order then the very first set passes through
+	// a state where the flag is set and the value is not saved.  In that case we might incorrectly restore to the MostSensitive setting.  In other cases we would restore the correct value,
+	// regardless of the flag value.
+	// To almost perfectly handle error cases we need an external process watching our game and restoring the setting on any abnormal termination.
+	void RemoveForceMaxTouchpadSensitivityOnErrorShutdown()
+	{
+		// Make a local copy of the saved sensitivity to work with.  Now we know it won't change out from under us.  We only have to worry about it being half-set.
+		TOptional<TOUCHPAD_SENSITIVITY_LEVEL> LocalSavedSensitivity = ForceMaxTouchpadSensitivity_SavedSensitivity;
+
+		// Only if we did set the sensitivity using the EForceMaxTouchpadSensitivityMode::TouchpadParameters implementation
+		// will we have cached a value, and need to restore it.
+		if (LocalSavedSensitivity.IsSet())
+		{
+			// The value should be in the valid range, and not MOST_SENSITIVE, because in that case we would not have needed to change it.
+			// If it is MOST_SENSITIVE or too large the game thread might have been setting or resetting it.
+			if (LocalSavedSensitivity.GetValue() <= TOUCHPAD_SENSITIVITY_LEVEL_MOST_SENSITIVE || LocalSavedSensitivity.GetValue() > TOUCHPAD_SENSITIVITY_LEVEL_LEAST_SENSITIVE)
+			{
+				// No restore needed, or invalid value.
+				return;
+			}
+
+			// We are not doing this wait in this error version.
+			//WaitOnPreviousApplyAsync();
+
+			// Note: we will get the params, change just the one, then write them back so that the reserved param block is preserved.
+			UE_LOG(LogWindowsTouchpad, Verbose, TEXT("RemoveForceMaxTouchpadSensitivityOnErrorShutdown reverting to SavedSensitivity %i"), LocalSavedSensitivity.GetValue());
+
+			// Get the user's touchpad settings.
+			TOUCHPAD_PARAMETERS_V0 params = {};
+			params.versionNumber = 0;
+			BOOL bSuccess = ::SystemParametersInfo(SPI_GETTOUCHPADPARAMETERS, sizeof(params), &params, 0 /*fWinIni*/);
+			if (!bSuccess)
+			{
+				UE_LOG(LogWindowsTouchpad, Warning, TEXT("RemoveForceMaxTouchpadSensitivityOnErrorShutdown SystemParametersInfo SPI_GETTOUCHPADPARAMETERS failed"));
+				return;
+			}
+
+			// Restore the saved sensitivity level.
+			// Leave all other bytes in the struct untouched.
+			params.sensitivityLevel = LocalSavedSensitivity.GetValue();
+
+			// Do not specify SPIF_UPDATEINIFILE so that the registry is not updated.
+			// We will revert this update later, and not writing to the registry is a
+			// failsafe in case we fail to do so (e.g. due to a crash).
+			bSuccess = ::SystemParametersInfo(SPI_SETTOUCHPADPARAMETERS, sizeof(params), &params, SPIF_SENDCHANGE);
+			if (!bSuccess)
+			{
+				UE_LOG(LogWindowsTouchpad, Warning, TEXT("RemoveForceMaxTouchpadSensitivityOnErrorShutdown SystemParametersInfo SPI_SETTOUCHPADPARAMETERS failed"));
+				return;
+			}
+			UE_LOG(LogWindowsTouchpad, Log, TEXT("RemoveForceMaxTouchpadSensitivityOnErrorShutdown sucessfully reverted to SavedSensitivity %i"), LocalSavedSensitivity.GetValue());
+		}
+	}
+
+	void HandleActivateApp( bool bActivate)
+	{
+		if (bForceMaxTouchpadSensitivity)
+		{
+			UE_LOG(LogWindowsTouchpad, Verbose, TEXT("HandleActivateApp %i"), bActivate);
+			if (bActivate)
+			{
+				ApplyForceMaxTouchpadSensitivity(FWindowsApplication::EForceMaxTouchpadSensitivityAsyncBehavior::AllowAsynchronous);
+			}
+			else
+			{
+				RemoveForceMaxTouchpadSensitivity();
+			}
+		}
+	}
+
+	void HandleStartupApp()
+	{
+		CacheForceMaxTouchpadSensitivityMode();
+		if (bForceMaxTouchpadSensitivity)
+		{
+			UE_LOG(LogWindowsTouchpad, Verbose, TEXT("HandleStartupApp"));
+			ApplyForceMaxTouchpadSensitivity(FWindowsApplication::EForceMaxTouchpadSensitivityAsyncBehavior::AllowAsynchronous);
+		}
+	}
+
+	void HandleDestroyApplication()
+	{
+		UE_LOG(LogWindowsTouchpad, Verbose, TEXT("HandleDestroyApplication"));
+		RemoveForceMaxTouchpadSensitivity();
+	}
+
+	void HandleShutDownAfterError()
+	{
+		UE_LOG(LogWindowsTouchpad, Verbose, TEXT("HandleShutDownAfterError"));
+		RemoveForceMaxTouchpadSensitivityOnErrorShutdown();
+	}
+}
 
 #if !UE_BUILD_SHIPPING
 static int32 EnableRawInputSimulationOverRDP = true;
@@ -248,6 +726,23 @@ FWindowsApplication::FWindowsApplication( const HINSTANCE HInstance, const HICON
 	}	
 	IConsoleManager::Get().RegisterConsoleCommand(TEXT("WindowsApplication.ApplyLowLevelMouseFilter"), TEXT("Applies Low Level mouse filter that filters out mouse inputs that act like touch inputs"), FConsoleCommandDelegate::CreateRaw(this, &FWindowsApplication::ApplyLowLevelMouseFilter));
 	IConsoleManager::Get().RegisterConsoleCommand(TEXT("WindowsApplication.RemoveLowLevelMouseFilter"), TEXT("Removes Low Level mouse filter that filters out mouse inputs that act like touch inputs"), FConsoleCommandDelegate::CreateRaw(this, &FWindowsApplication::RemoveLowLevelMouseFilter));
+
+	UE::WindowsApplication::TouchpadSensitivity::HandleStartupApp();
+}
+
+void FWindowsApplication::ApplyForceMaxTouchpadSensitivity(FWindowsApplication::EForceMaxTouchpadSensitivityAsyncBehavior AsyncBehavior)
+{
+	UE::WindowsApplication::TouchpadSensitivity::ApplyForceMaxTouchpadSensitivity(AsyncBehavior);
+}
+
+void FWindowsApplication::RemoveForceMaxTouchpadSensitivity()
+{
+	UE::WindowsApplication::TouchpadSensitivity::RemoveForceMaxTouchpadSensitivity();
+}
+
+FWindowsApplication::EForceMaxTouchpadSensitivityRestorationValues FWindowsApplication::GetForceMaxTouchpadSensitivityRestorationValue()
+{
+	return UE::WindowsApplication::TouchpadSensitivity::GetForceMaxTouchpadSensitivityRestorationValue();
 }
 
 void FWindowsApplication::ApplyLowLevelMouseFilter()
@@ -326,6 +821,7 @@ void FWindowsApplication::DestroyApplication()
 	RemoveLowLevelMouseFilter();
 	IConsoleManager::Get().UnregisterConsoleObject(TEXT("WindowsApplication.ApplyLowLevelMouseFilter"));
 	IConsoleManager::Get().UnregisterConsoleObject(TEXT("WindowsApplication.RemoveLowLevelMouseFilter"));
+	UE::WindowsApplication::TouchpadSensitivity::HandleDestroyApplication();
 }
 
 void FWindowsApplication::ShutDownAfterError()
@@ -339,6 +835,7 @@ void FWindowsApplication::ShutDownAfterError()
 	RemoveLowLevelMouseFilter();
 	IConsoleManager::Get().UnregisterConsoleObject(TEXT("WindowsApplication.ApplyLowLevelMouseFilter"));
 	IConsoleManager::Get().UnregisterConsoleObject(TEXT("WindowsApplication.RemoveLowLevelMouseFilter"));
+	UE::WindowsApplication::TouchpadSensitivity::HandleShutDownAfterError();
 }
 
 bool FWindowsApplication::RegisterClass( const HINSTANCE HInstance, const HICON HIcon )
@@ -2478,6 +2975,7 @@ int32 FWindowsApplication::ProcessDeferredMessage( const FDeferredWindowsMessage
 		case WM_ACTIVATEAPP:
 			UpdateAllModifierKeyStates();
 			MessageHandler->OnApplicationActivationChanged( !!wParam );
+			UE::WindowsApplication::TouchpadSensitivity::HandleActivateApp(!!wParam);
 			break;
 
 		case WM_SETTINGCHANGE:
diff --git a/Engine/Source/Runtime/ApplicationCore/Public/Windows/WindowsApplication.h b/Engine/Source/Runtime/ApplicationCore/Public/Windows/WindowsApplication.h
--- a/Engine/Source/Runtime/ApplicationCore/Public/Windows/WindowsApplication.h
+++ b/Engine/Source/Runtime/ApplicationCore/Public/Windows/WindowsApplication.h
@@ -384,6 +384,38 @@ public:
 	virtual void ResetLightColor(int32 ControllerId) override { }
 	APPLICATIONCORE_API virtual void SetDeviceProperty(int32 ControllerId, const FInputDeviceProperty* Property) override;
 
+	// Touchpad Sensitivity
+	/** On Windows 10 22H2 and later where we can set the touchpad sensitivity system wide setting force it to most sensitive so that touch swipes are not ignored while keyboard keys are being depressed.
+	 * IF the setting is actually changed the system call that changes it can block for approximately 1 second.  Because that is a pretty long time, potentially delaying game startup
+	 * an argument exists on the Apply function which will allow the application of the setting to be done asynchronously.  The Remove function has no such option. Both functions will block for 
+	 * an outstanding Apply system call (With a timeout).  The end result here is that one can use the parameter to prevent the Apply function from potentially blocking for a long time, but the Remove function can then 
+	 * potentially block for two system calls.  It is expected and reccomended that the real world usage patterns make the double-block rare.
+	 */
+	enum class EForceMaxTouchpadSensitivityAsyncBehavior : uint8
+	{
+		AllowAsynchronous, // This function may defer the system call that actually changes the setting in order to avoid blocking for a long time.
+		RequireBlocking // This function will not return until the system call returns.
+	};
+	APPLICATIONCORE_API void ApplyForceMaxTouchpadSensitivity(EForceMaxTouchpadSensitivityAsyncBehavior AsyncBehavior);
+	APPLICATIONCORE_API void RemoveForceMaxTouchpadSensitivity();
+
+	/** If we are using the system wide setting it is possible that we could crash and fail to set it back.  Call this function to
+	 * find out if we have a system setting value we would like to restore after a crash.
+	 * If this information was passed to a separate process (perhaps by writing it to a file) that process might be able to restore the setting
+	 * after a crash.
+	 */
+	enum class EForceMaxTouchpadSensitivityRestorationValues : int32
+	{
+		SystemCallFailed = -2, // Indicates we are in a mode where restoration is desirable, but the system call failed.  Most likely callers should do nothing.
+		RestorationNotNeeded = -1,  // Indicates restoration is unnecessary (either the system setting was already MostSensitive, we can't change the setting, or the system is fully capable of restoring itself.)
+		TOUCHPAD_SENSITIVITY_LEVEL_MOST_SENSITIVE = 0x00000000,  // These values one would restore with ::SystemParametersInfo(SPI_SETTOUCHPADPARAMETERS, ...
+		TOUCHPAD_SENSITIVITY_LEVEL_HIGH_SENSITIVITY = 0x00000001,
+		TOUCHPAD_SENSITIVITY_LEVEL_MEDIUM_SENSITIVITY = 0x00000002,
+		TOUCHPAD_SENSITIVITY_LEVEL_LOW_SENSITIVITY = 0x00000003,
+		TOUCHPAD_SENSITIVITY_LEVEL_LEAST_SENSITIVE = 0x00000004,
+	};
+	APPLICATIONCORE_API EForceMaxTouchpadSensitivityRestorationValues GetForceMaxTouchpadSensitivityRestorationValue();
+
 protected:
 	friend LRESULT WindowsApplication_WndProc(HWND hwnd, uint32 msg, WPARAM wParam, LPARAM lParam);
 
