Add a new log for any time we allocate a platform user ID to help track down some hard to reproduce behavior which has been reported where we ay be allocating more ID's then we need to

#rnx

diff --git a/Engine/Source/Runtime/ApplicationCore/Private/GenericPlatform/GenericPlatformInputDeviceMapper.cpp b/Engine/Source/Runtime/ApplicationCore/Private/GenericPlatform/GenericPlatformInputDeviceMapper.cpp
--- a/Engine/Source/Runtime/ApplicationCore/Private/GenericPlatform/GenericPlatformInputDeviceMapper.cpp
+++ b/Engine/Source/Runtime/ApplicationCore/Private/GenericPlatform/GenericPlatformInputDeviceMapper.cpp
@@ -734,6 +734,8 @@ FPlatformUserId FGenericPlatformInputDeviceMapper::AllocateNewUserId()
 
 	AllocatedPlatformUserIds.AddUnique(LastPlatformUserId);
 
+	UE_LOG(LogInputDeviceMapper, Log, TEXT("[%hs] Allocating a new PlatformUserId %d"), __func__, LastPlatformUserId.GetInternalId());
+
 	// We want to warn about this state happening, but not crash the application or anything. 
 	// Each platform has a specific number of platform users which can be signed on at any given time
 	// so going over that amount would cause undefined behavior such as your input no being routed correctly
