Fix UIView contentScaleFactor not being set on UI thread and only set it if it's changed

#rnx

diff --git a/Engine/Source/Runtime/ApplicationCore/Private/IOS/IOSView.cpp b/Engine/Source/Runtime/ApplicationCore/Private/IOS/IOSView.cpp
--- a/Engine/Source/Runtime/ApplicationCore/Private/IOS/IOSView.cpp
+++ b/Engine/Source/Runtime/ApplicationCore/Private/IOS/IOSView.cpp
@@ -273,8 +273,9 @@ MTL::Device* GMetalDevice = nullptr;
 	
 	int32 Width = ScreenWidth;
 	int32 Height = ScreenHeight;
-	
-	self.contentScaleFactor = mNativeScale;
+
+	CGFloat currentScaleFactor = self.contentScaleFactor;
+	CGFloat newScaleFactor = mNativeScale;
 	
 	// 0 means to use native size
 	if (RequestedContentScaleFactor != 0.0f || RequestedResX > 0 || RequestedResY > 0)
@@ -310,10 +311,21 @@ MTL::Device* GMetalDevice = nullptr;
 		}
 		else
 		{
-			self.contentScaleFactor = RequestedContentScaleFactor;
+			newScaleFactor = RequestedContentScaleFactor;
 		}
 	}
 
+	// As this can be called every tick(), and setting it requires being on the UI thread, only
+	// update the value if it's actually changed.
+	if (currentScaleFactor != newScaleFactor)
+	{
+		FIOSView* __block blockSelf = self;
+		dispatch_async(dispatch_get_main_queue(), ^{
+			blockSelf.contentScaleFactor = newScaleFactor;
+			blockSelf = nil;
+		});
+	}
+
 	_ViewSize.width = Width;
 	_ViewSize.height = Height;
 }
