CocoaThread - Simplify (Main|Game)ThreadCall use when call args == default args and change order of argument to make more sense
We don't need to specify a WaitMode when Wait = false.

#tests Local testing, functionally identical

diff --git a/Engine/Plugins/Experimental/TextToSpeech/Source/TextToSpeech/Private/Mac/MacTextToSpeech.cpp b/Engine/Plugins/Experimental/TextToSpeech/Source/TextToSpeech/Private/Mac/MacTextToSpeech.cpp
--- a/Engine/Plugins/Experimental/TextToSpeech/Source/TextToSpeech/Private/Mac/MacTextToSpeech.cpp
+++ b/Engine/Plugins/Experimental/TextToSpeech/Source/TextToSpeech/Private/Mac/MacTextToSpeech.cpp
@@ -51,7 +51,7 @@ if (finishedSpeaking && _OwningTextToSpeechId != FTextToSpeechBase::InvalidTextT
 		{
 			TTSPtr->Pin()->OnTextToSpeechFinishSpeaking_GameThread();
 		}
-	}, @[ NSDefaultRunLoopMode ], false);
+	}, false);
 }
 }
 
@@ -181,7 +181,7 @@ void FMacTextToSpeech::OnActivated()
 		SpeechSynthesizerDelegate = [[FSpeechSynthesizerDelegate alloc] initWithTextToSpeechId:IdCopy];
 		checkf(SpeechSynthesizerDelegate, TEXT("Failed to create Mac speech synthesizer delegate"));
 		SpeechSynthesizer.delegate = SpeechSynthesizerDelegate;
-	}, NSDefaultRunLoopMode, true);
+	});
 
 	Volume = SpeechSynthesizer.volume;
 	Rate = RateWPMToRate(SpeechSynthesizer.rate);
@@ -197,7 +197,7 @@ void FMacTextToSpeech::OnDeactivated()
 		[SpeechSynthesizerDelegate release];
 		checkf(SpeechSynthesizer, TEXT("Deactivating Mac TTS with null speech synthesizer. Speech synthesizer must be valid throughout the lifetime of the object."));
 		[SpeechSynthesizer release];
-	}, NSDefaultRunLoopMode, true);
+	});
 	UE_LOG(LogTextToSpeech, Verbose, TEXT("Mac TTS deactivated."));
 }
 
diff --git a/Engine/Plugins/Media/AvfMedia/Source/AvfMedia/Private/AvfMediaPrivate.h b/Engine/Plugins/Media/AvfMedia/Source/AvfMedia/Private/AvfMediaPrivate.h
--- a/Engine/Plugins/Media/AvfMedia/Source/AvfMedia/Private/AvfMediaPrivate.h
+++ b/Engine/Plugins/Media/AvfMedia/Source/AvfMedia/Private/AvfMediaPrivate.h
@@ -12,7 +12,7 @@ DECLARE_LOG_CATEGORY_EXTERN(LogAvfMedia, Log, All);
 #if PLATFORM_MAC
 	#define AVF_GAME_THREAD_BLOCK 
 	#define AVF_GAME_THREAD_RETURN 
-	#define AVF_GAME_THREAD_CALL(Block)	GameThreadCall(Block, @[ NSDefaultRunLoopMode ], false)
+	#define AVF_GAME_THREAD_CALL(Block)	GameThreadCall(Block, false)
 #else
 	#define AVF_GAME_THREAD_BLOCK bool(void)
 	#define AVF_GAME_THREAD_RETURN return true
diff --git a/Engine/Source/Developer/DesktopPlatform/Private/Mac/MacNativeFeedbackContext.cpp b/Engine/Source/Developer/DesktopPlatform/Private/Mac/MacNativeFeedbackContext.cpp
--- a/Engine/Source/Developer/DesktopPlatform/Private/Mac/MacNativeFeedbackContext.cpp
+++ b/Engine/Source/Developer/DesktopPlatform/Private/Mac/MacNativeFeedbackContext.cpp
@@ -268,7 +268,7 @@ FMacNativeFeedbackContext::FMacNativeFeedbackContext()
 {
 	MainThreadCall(^{
 		WindowController = [[FMacNativeFeedbackContextWindowController alloc] init];
-	}, UnrealNilEventMode, true);
+	}, true, UnrealNilEventMode);
 	SetDefaultTextColor();
 }
 
@@ -281,7 +281,7 @@ FMacNativeFeedbackContext::~FMacNativeFeedbackContext()
 	
 	MainThreadCall(^{
 		[WindowController release];
-	}, UnrealNilEventMode, true);
+	}, true, UnrealNilEventMode);
 }
 
 void FMacNativeFeedbackContext::Serialize(const TCHAR* Data, ELogVerbosity::Type Verbosity, const FName& Category)
@@ -362,7 +362,7 @@ void FMacNativeFeedbackContext::SerializeToWindow(const TCHAR* Data, ELogVerbosi
 						[Colors release];
 						[AttributeKeys release];
 						OutstandingTasks--;
-					}, NSDefaultRunLoopMode, false);
+					}, false);
 				}
 			}
 			else
@@ -382,7 +382,7 @@ void FMacNativeFeedbackContext::SerializeToWindow(const TCHAR* Data, ELogVerbosi
 					[AttributedString release];
 					CFRelease(CocoaText);
 					OutstandingTasks--;
-				}, NSDefaultRunLoopMode, false);
+				}, false);
 
 				if(!MacApplication)
 				{
@@ -505,5 +505,5 @@ void FMacNativeFeedbackContext::SetDefaultTextColor()
 		[Colors release];
 		[AttributeKeys release];
 		OutstandingTasks--;
-	}, NSDefaultRunLoopMode, false);
+	}, false);
 }
diff --git a/Engine/Source/Developer/StandaloneRenderer/Private/Mac/OpenGL/SlateOpenGLContext.cpp b/Engine/Source/Developer/StandaloneRenderer/Private/Mac/OpenGL/SlateOpenGLContext.cpp
--- a/Engine/Source/Developer/StandaloneRenderer/Private/Mac/OpenGL/SlateOpenGLContext.cpp
+++ b/Engine/Source/Developer/StandaloneRenderer/Private/Mac/OpenGL/SlateOpenGLContext.cpp
@@ -188,7 +188,7 @@ void FSlateOpenGLContext::Initialize(void* InWindow, const FSlateOpenGLContext*
 			[Window setContentView:View];
 
 			[[Window standardWindowButton:NSWindowCloseButton] setAction:@selector(performClose:)];
-		}, NSDefaultRunLoopMode, true);
+		});
 
 		View.layer.magnificationFilter = kCAFilterNearest;
 		View.layer.minificationFilter = kCAFilterNearest;
diff --git a/Engine/Source/Developer/StandaloneRenderer/Private/Mac/OpenGL/SlateOpenGLViewport.cpp b/Engine/Source/Developer/StandaloneRenderer/Private/Mac/OpenGL/SlateOpenGLViewport.cpp
--- a/Engine/Source/Developer/StandaloneRenderer/Private/Mac/OpenGL/SlateOpenGLViewport.cpp
+++ b/Engine/Source/Developer/StandaloneRenderer/Private/Mac/OpenGL/SlateOpenGLViewport.cpp
@@ -44,7 +44,7 @@ void FSlateOpenGLViewport::SwapBuffers()
 	[(FCocoaWindow*)[RenderingContext.View window] startRendering];
 	MainThreadCall(^{
 		[RenderingContext.View setNeedsDisplay:YES];
-	}, NSDefaultRunLoopMode, false);
+	}, false);
 	glBindFramebuffer(GL_FRAMEBUFFER, ((FSlateCocoaView*)RenderingContext.View)->Framebuffer);
 	UnlockGLContext(RenderingContext.Context);
 }
diff --git a/Engine/Source/Runtime/Apple/MetalRHI/Private/MetalViewport.cpp b/Engine/Source/Runtime/Apple/MetalRHI/Private/MetalViewport.cpp
--- a/Engine/Source/Runtime/Apple/MetalRHI/Private/MetalViewport.cpp
+++ b/Engine/Source/Runtime/Apple/MetalRHI/Private/MetalViewport.cpp
@@ -134,7 +134,7 @@ FMetalViewport::FMetalViewport(FMetalDevice& InDevice, void* WindowHandle, uint3
 
 		[Window setContentView:View];
 		[[Window standardWindowButton:NSWindowCloseButton] setAction:@selector(performClose:)];
-	}, NSDefaultRunLoopMode, true);
+	});
 #endif
 	Resize(InSizeX, InSizeY, bInIsFullscreen, Format);
 	
@@ -237,7 +237,7 @@ void FMetalViewport::Resize(uint32 InSizeX, uint32 InSizeY, bool bInIsFullscreen
 			MetalLayer.wantsExtendedDynamicRangeContent = bUseHDR;
 		}
 		
-	}, NSDefaultRunLoopMode, true);
+	});
 #else
 	// A note on HDR in iOS
 	// Setting the pixel format to one of the Apple XR formats is all you need.
@@ -554,7 +554,7 @@ void FMetalViewport::Present(FMetalCommandQueue& CommandQueue, bool bLockToVsync
 							MainThreadCall(^{
 								FCocoaWindow* Window = (FCocoaWindow*)[theView window];
 								[Window startRendering];
-							}, NSDefaultRunLoopMode, false);
+							}, false);
 #endif
 						};
 						
diff --git a/Engine/Source/Runtime/ApplicationCore/Private/Mac/Accessibility/MacAccessibilityElement.cpp b/Engine/Source/Runtime/ApplicationCore/Private/Mac/Accessibility/MacAccessibilityElement.cpp
--- a/Engine/Source/Runtime/ApplicationCore/Private/Mac/Accessibility/MacAccessibilityElement.cpp
+++ b/Engine/Source/Runtime/ApplicationCore/Private/Mac/Accessibility/MacAccessibilityElement.cpp
@@ -155,9 +155,9 @@
 					}
 				} // switch WidgetType
 				self.accessibilityRole = _AccessibilityRole;
-			}, NSDefaultRunLoopMode, false);
+			}, false);
 		} // if widget valid
-	}, @[NSDefaultRunLoopMode], false);
+	}, false);
 	
 	return self;
 }
@@ -207,9 +207,9 @@
 				Self.Help = TempHelp;
 				Self.Value = TempValue;
 				Self.LastCachedStringTime = FPlatformTime::Seconds();
-			}, NSDefaultRunLoopMode, false);
+			}, false);
 		} // if widget valid
-	}, @[ NSDefaultRunLoopMode ], false);
+	}, false);
 }
 
 - (void)ExposeToVoiceOver:(bool)bShouldExpose;
@@ -269,7 +269,7 @@
 				// Voiceover only recognizes 1 user, the primary accessible user.
 				Widget->SetUserFocus(FGenericAccessibleUserRegistry::GetPrimaryUserIndex());
 			}
-		}, @[ NSDefaultRunLoopMode, UnrealIMEEventMode, UnrealShowEventMode, UnrealResizeEventMode, UnrealFullscreenEventMode], false); // GameThreadCall
+		}, false, @[ NSDefaultRunLoopMode, UnrealIMEEventMode, UnrealShowEventMode, UnrealResizeEventMode, UnrealFullscreenEventMode]); // GameThreadCall
 	}
 	[super setAccessibilityFocused:InIsFocused];
 }
@@ -432,11 +432,11 @@
 					MainThreadCall(^{
 						FMacAccessibilityElement* Self = [[FMacAccessibilityManager AccessibilityManager] GetAccessibilityElement:TempId];
 						NSAccessibilityPostNotification(Self, NSAccessibilityValueChangedNotification);
-					}, NSDefaultRunLoopMode, false);
+					}, false);
 				}
 			}
 		} // if widget valid and enabled
-	}, @[ NSDefaultRunLoopMode, UnrealIMEEventMode, UnrealShowEventMode, UnrealResizeEventMode, UnrealFullscreenEventMode], false); // GameThreadCall
+	}, false, @[ NSDefaultRunLoopMode, UnrealIMEEventMode, UnrealShowEventMode, UnrealResizeEventMode, UnrealFullscreenEventMode]); // GameThreadCall
 	return YES;
 }
 
@@ -461,10 +461,10 @@
 				MainThreadCall(^{
 					FMacAccessibilityElement* Self = [[FMacAccessibilityManager AccessibilityManager] GetAccessibilityElement:TempId];
 					NSAccessibilityPostNotification(Self, NSAccessibilityValueChangedNotification);
-				}, NSDefaultRunLoopMode, false);
+				}, false);
 			}
 		} // widget valid
-	}, @[ NSDefaultRunLoopMode ], false);
+	}, false);
 	return YES;
 }
 
@@ -489,10 +489,10 @@
 				MainThreadCall(^{
 					FMacAccessibilityElement* Self = [[FMacAccessibilityManager AccessibilityManager] GetAccessibilityElement:TempId];
 					NSAccessibilityPostNotification(Self, NSAccessibilityValueChangedNotification);
-				}, NSDefaultRunLoopMode, false);
+				}, false);
 			}
 		} // widget valid
-	}, @[ NSDefaultRunLoopMode ], false);
+	}, false);
 	return YES;
 }
 
diff --git a/Engine/Source/Runtime/ApplicationCore/Private/Mac/Accessibility/MacAccessibilityManager.cpp b/Engine/Source/Runtime/ApplicationCore/Private/Mac/Accessibility/MacAccessibilityManager.cpp
--- a/Engine/Source/Runtime/ApplicationCore/Private/Mac/Accessibility/MacAccessibilityManager.cpp
+++ b/Engine/Source/Runtime/ApplicationCore/Private/Mac/Accessibility/MacAccessibilityManager.cpp
@@ -224,10 +224,10 @@ static void DumpAccessibilityStatsForwarder()
 					//@TODO: Updating the label, help and value of ALL elements is suboptimal...but it works
 					//Tested with EngineTest in Debug Editor in the Editor with ~650 elements and still fairly responsive
 					[Element CacheStrings];
-				}, NSDefaultRunLoopMode, true);
+				});
 			} // if Widget valid
 		} // for Ids loop
-	}, @[ NSDefaultRunLoopMode, UnrealIMEEventMode, UnrealShowEventMode, UnrealResizeEventMode, UnrealFullscreenEventMode, UnrealNilEventMode ], true);
+	}, true, @[ NSDefaultRunLoopMode, UnrealIMEEventMode, UnrealShowEventMode, UnrealResizeEventMode, UnrealFullscreenEventMode, UnrealNilEventMode ]);
 }
 
 #if !UE_BUILD_SHIPPING
diff --git a/Engine/Source/Runtime/ApplicationCore/Private/Mac/CocoaTextView.cpp b/Engine/Source/Runtime/ApplicationCore/Private/Mac/CocoaTextView.cpp
--- a/Engine/Source/Runtime/ApplicationCore/Private/Mac/CocoaTextView.cpp
+++ b/Engine/Source/Runtime/ApplicationCore/Private/Mac/CocoaTextView.cpp
@@ -202,7 +202,7 @@
 						ITextInputMethodContext::ECaretPosition CaretPosition;
 						IMMContext->GetSelectionRange(SelectionLocation, SelectionLength, CaretPosition);
 					}
-				}, @[ NSDefaultRunLoopMode, UnrealIMEEventMode ]);
+				}, true, @[ NSDefaultRunLoopMode, UnrealIMEEventMode ]);
 			}
 		}
 		else
@@ -229,7 +229,7 @@
 				IMMContext->SetTextInRange(SelectionLocation, SelectionLength, TheFString);
 				IMMContext->SetSelectionRange(SelectionLocation+TheFString.Len(), 0, ITextInputMethodContext::ECaretPosition::Ending);
 			}
-		}, @[ NSDefaultRunLoopMode, UnrealIMEEventMode ]);
+		}, true, @[ NSDefaultRunLoopMode, UnrealIMEEventMode ]);
 
 		[self unmarkText];
 		[[self inputContext] invalidateCharacterCoordinates]; // recentering
@@ -271,7 +271,7 @@
 						ITextInputMethodContext::ECaretPosition CaretPosition;
 						IMMContext->GetSelectionRange(SelectionLocation, SelectionLength, CaretPosition);
 					}
-				}, @[ NSDefaultRunLoopMode, UnrealIMEEventMode ]);
+				}, true, @[ NSDefaultRunLoopMode, UnrealIMEEventMode ]);
 			}
 		}
 		else
@@ -287,7 +287,7 @@
 				{
 					IMMContext->SetTextInRange(SelectionLocation, SelectionLength, FString());
 				}
-			}, @[ NSDefaultRunLoopMode, UnrealIMEEventMode ]);
+			}, true, @[ NSDefaultRunLoopMode, UnrealIMEEventMode ]);
 			[self unmarkText];
 		}
 		else
@@ -299,7 +299,7 @@
 					{
 						IMMContext->BeginComposition();
 					}
-				}, @[ NSDefaultRunLoopMode, UnrealIMEEventMode ]);
+				}, true, @[ NSDefaultRunLoopMode, UnrealIMEEventMode ]);
 			}
 			markedRange = NSMakeRange(SelectionLocation, [aString length]);
 
@@ -344,7 +344,7 @@
 					IMMContext->UpdateCompositionRange(CompositionRange.location, CompositionRange.length);
 					IMMContext->SetSelectionRange(markedRange.location + selectedRange.location, 0, ITextInputMethodContext::ECaretPosition::Ending);
 				}
-			}, @[ NSDefaultRunLoopMode, UnrealIMEEventMode ]);
+			}, true, @[ NSDefaultRunLoopMode, UnrealIMEEventMode ]);
 		}
 		[[self inputContext] invalidateCharacterCoordinates]; // recentering
 	}
@@ -369,7 +369,7 @@
 					IMMContext->UpdateCompositionRange(0, 0);
 					IMMContext->EndComposition();
 				}
-			}, @[ NSDefaultRunLoopMode, UnrealIMEEventMode ]);
+			}, true, @[ NSDefaultRunLoopMode, UnrealIMEEventMode ]);
 		}
 	}
 }
@@ -394,7 +394,7 @@
 			{
 				IMMContext->GetSelectionRange(SelectionLocation, SelectionLength, CaretPosition);
 			}
-		}, @[ NSDefaultRunLoopMode, UnrealIMEEventMode ]);
+		}, true, @[ NSDefaultRunLoopMode, UnrealIMEEventMode ]);
 		return {SelectionLocation, SelectionLength};
 	}
 	else
@@ -438,7 +438,7 @@
 			{
 				IMMContext->GetTextInRange(aRange.location, aRange.length, String);
 			}
-		}, @[ NSDefaultRunLoopMode, UnrealIMEEventMode ]);
+		}, true, @[ NSDefaultRunLoopMode, UnrealIMEEventMode ]);
 		CFStringRef CFString = FPlatformString::TCHARToCFString(*String);
 		if(CFString)
 		{
@@ -479,7 +479,7 @@
 			{
 				IMMContext->GetTextBounds(aRange.location, aRange.length, Position, Size);
 			}
-		}, @[ NSDefaultRunLoopMode, UnrealIMEEventMode ]);
+		}, true, @[ NSDefaultRunLoopMode, UnrealIMEEventMode ]);
 
 		if(actualRange)
 		{
diff --git a/Engine/Source/Runtime/ApplicationCore/Private/Mac/CocoaThreadTesting.cpp b/Engine/Source/Runtime/ApplicationCore/Private/Mac/CocoaThreadTesting.cpp
--- a/Engine/Source/Runtime/ApplicationCore/Private/Mac/CocoaThreadTesting.cpp
+++ b/Engine/Source/Runtime/ApplicationCore/Private/Mac/CocoaThreadTesting.cpp
@@ -15,7 +15,7 @@ void RecursiveThreadCall(int Depth, bool bFromGameThread, int MaxDepth)
 			// Test other thread call
 			MainThreadCall(^{
 				RecursiveThreadCall(Depth + 1, false, MaxDepth);
-			}, UnrealNilEventMode);
+			}, true, UnrealNilEventMode);
 			
 			// Test current thread call
 			GameThreadCall(^{});
diff --git a/Engine/Source/Runtime/ApplicationCore/Private/Mac/CocoaWindow.cpp b/Engine/Source/Runtime/ApplicationCore/Private/Mac/CocoaWindow.cpp
--- a/Engine/Source/Runtime/ApplicationCore/Private/Mac/CocoaWindow.cpp
+++ b/Engine/Source/Runtime/ApplicationCore/Private/Mac/CocoaWindow.cpp
@@ -137,7 +137,7 @@ NSString* NSPerformDragOperation = @"NSPerformDragOperation";
 	{
 		MainThreadCall(^{
 			[super setAlphaValue:WindowAlpha];
-		}, UnrealShowEventMode, false);
+		}, false);
 	}
 }
 
@@ -148,7 +148,7 @@ NSString* NSPerformDragOperation = @"NSPerformDragOperation";
 		bRenderInitialized = true;
 		MainThreadCall(^{
 			[super setAlphaValue:Opacity];
-		}, UnrealShowEventMode, false);
+		}, false);
 	}
 }
 
@@ -168,7 +168,7 @@ NSString* NSPerformDragOperation = @"NSPerformDragOperation";
 				MacApplication->CloseWindow(Window.ToSharedRef());
 			}
 		}
-	}, @[ NSDefaultRunLoopMode ], false);
+	}, false);
 }
 
 - (void)performZoom:(id)Sender
@@ -295,7 +295,7 @@ NSString* NSPerformDragOperation = @"NSPerformDragOperation";
 					MacApplication->OnWindowActivationChanged(Window.ToSharedRef(), EWindowActivation::Activate);
 				}
 			}
-		}, @[ NSDefaultRunLoopMode, UnrealResizeEventMode, UnrealShowEventMode, UnrealFullscreenEventMode, UnrealCloseEventMode ], true);
+		}, true, @[ NSDefaultRunLoopMode, UnrealResizeEventMode, UnrealShowEventMode, UnrealFullscreenEventMode, UnrealCloseEventMode ]);
 	}
 }
 
@@ -316,7 +316,7 @@ NSString* NSPerformDragOperation = @"NSPerformDragOperation";
 					MacApplication->OnWindowActivationChanged(Window.ToSharedRef(), EWindowActivation::Deactivate);
 				}
 			}
-		}, @[ NSDefaultRunLoopMode, UnrealResizeEventMode, UnrealShowEventMode, UnrealFullscreenEventMode, UnrealCloseEventMode ], true);
+		}, true, @[ NSDefaultRunLoopMode, UnrealResizeEventMode, UnrealShowEventMode, UnrealFullscreenEventMode, UnrealCloseEventMode ]);
 	}
 }
 
@@ -429,7 +429,7 @@ NSString* NSPerformDragOperation = @"NSPerformDragOperation";
 					MacApplication->OnWindowDidResize(Window.ToSharedRef(), false);
 				}
 			}
-		}, @[ NSDefaultRunLoopMode ], false);
+		}, false);
 	}
 
 	GameThreadCall(^{
@@ -441,7 +441,7 @@ NSString* NSPerformDragOperation = @"NSPerformDragOperation";
 				MacApplication->OnWindowChangedScreen(Window.ToSharedRef());
 			}
 		}
-	}, @[ NSDefaultRunLoopMode ], false);
+	}, false);
 }
 
 - (void)windowDidChangeBackingProperties:(NSNotification *)notification
@@ -481,7 +481,7 @@ NSString* NSPerformDragOperation = @"NSPerformDragOperation";
 					MacApplication->OnWindowWillResize(Window.ToSharedRef());
 				}
 			}
-		}, @[ NSDefaultRunLoopMode, UnrealResizeEventMode, UnrealShowEventMode, UnrealFullscreenEventMode ], true);
+		}, true, @[ NSDefaultRunLoopMode, UnrealResizeEventMode, UnrealShowEventMode, UnrealFullscreenEventMode ]);
 	}
 	return frameSize;
 }
diff --git a/Engine/Source/Runtime/ApplicationCore/Private/Mac/MacApplication.cpp b/Engine/Source/Runtime/ApplicationCore/Private/Mac/MacApplication.cpp
--- a/Engine/Source/Runtime/ApplicationCore/Private/Mac/MacApplication.cpp
+++ b/Engine/Source/Runtime/ApplicationCore/Private/Mac/MacApplication.cpp
@@ -170,7 +170,7 @@ FMacApplication::FMacApplication()
 		CacheKeyboardInputSource();
 
 		WindowUnderCursor = [FindSlateWindowUnderCursor() retain];
-	}, NSDefaultRunLoopMode, true);
+	});
 
 #if WITH_EDITOR
 	NSMutableArray* MultiTouchDevices = (__bridge NSMutableArray*)MTDeviceCreateList();
@@ -260,7 +260,7 @@ FMacApplication::~FMacApplication()
 		}
 		
 		[WindowUnderCursor release];
-	}, NSDefaultRunLoopMode, true);
+	});
 
 	if (TextInputMethodSystem.IsValid())
 	{
@@ -298,7 +298,7 @@ void FMacApplication::SetAccessibleMessageHandler(const TSharedRef<FGenericAcces
 		{
 			OnVoiceoverEnabled();
 		}
-	}, NSDefaultRunLoopMode, false);
+	}, false);
 }
 #endif
 
@@ -600,7 +600,7 @@ void FMacApplication::DeferEvent(NSObject* Object)
 						{
 							OnWindowDidResize(Window.ToSharedRef());
 						}
-					}, @[ NSDefaultRunLoopMode, UnrealResizeEventMode, UnrealShowEventMode, UnrealFullscreenEventMode ], true);
+					}, true, @[ NSDefaultRunLoopMode, UnrealResizeEventMode, UnrealShowEventMode, UnrealFullscreenEventMode ]);
 				}
 				return;
 			}
@@ -894,7 +894,7 @@ void FMacApplication::ResendEvent(NSEvent* Event)
 	MainThreadCall(^{
 		NSEvent* Wrapper = [NSEvent otherEventWithType:NSEventTypeApplicationDefined location:[Event locationInWindow] modifierFlags:[Event modifierFlags] timestamp:[Event timestamp] windowNumber:[Event windowNumber] context:nil subtype:RESET_EVENT_SUBTYPE data1:(NSInteger)Event data2:0];
 		[NSApp sendEvent:Wrapper];
-	}, NSDefaultRunLoopMode, true);
+	});
 }
 
 void FMacApplication::ProcessMouseMovedEvent(const FDeferredMacEvent& Event, TSharedPtr<FMacWindow> EventWindow)
@@ -909,7 +909,7 @@ void FMacApplication::ProcessMouseMovedEvent(const FDeferredMacEvent& Event, TSh
 		MainThreadCall(^{
 			[WindowHandle setMovable:IsMovable];
 			[WindowHandle setMovableByWindowBackground:IsMouseOverTitleBar];
-		}, NSDefaultRunLoopMode, false);
+		}, false);
 	}
 
 	FMacCursor* MacCursor = (FMacCursor*)Cursor.Get();
@@ -1037,11 +1037,11 @@ void FMacApplication::ProcessMouseDownEvent(const FDeferredMacEvent& Event, TSha
 				FCocoaWindow* WindowHandle = EventWindow->GetWindowHandle();
 				if (bShouldMinimize)
 				{
-					MainThreadCall(^{ [WindowHandle performMiniaturize:nil]; }, NSDefaultRunLoopMode, true);
+					MainThreadCall(^{ [WindowHandle performMiniaturize:nil]; });
 				}
 				else
 				{
-					MainThreadCall(^{ [WindowHandle zoom:nil]; }, NSDefaultRunLoopMode, true);
+					MainThreadCall(^{ [WindowHandle zoom:nil]; });
 				}
 			}
 			else
@@ -1172,7 +1172,7 @@ void FMacApplication::ProcessKeyDownEvent(const FDeferredMacEvent& Event, TShare
 		if (MainMenu)
 		{
 			NSEvent* NativeEvent = Event.Event;
-			MainThreadCall(^{ [MainMenu highlightKeyEquivalent:NativeEvent]; }, NSDefaultRunLoopMode, false);
+			MainThreadCall(^{ [MainMenu highlightKeyEquivalent:NativeEvent]; }, false);
 		}
 	}
 	else
@@ -1297,7 +1297,7 @@ bool FMacApplication::OnWindowDestroyed(TSharedRef<FMacWindow> DestroyedWindow)
 		//This is to clear the FMacAccessibilityElement* that represents the cocoa window in the FCocoaAccessibilityView.
 		// Before the FCocoaAccessibilityView is deallocated, it could be a danling pointer and accessibility calls from VoiceOVer will cause a crash
 		[WindowHandle ClearAccessibilityView];
-	}, NSDefaultRunLoopMode, false);
+	}, false);
 #endif
 	
 	{
@@ -1366,7 +1366,7 @@ bool FMacApplication::OnWindowDestroyed(TSharedRef<FMacWindow> DestroyedWindow)
 			{
 				GameThreadCall(^{
 					OnWindowActivationChanged(WindowToActivate.ToSharedRef(), EWindowActivation::Activate);
-				}, @[ NSDefaultRunLoopMode, UnrealResizeEventMode, UnrealShowEventMode, UnrealFullscreenEventMode, UnrealCloseEventMode ], true);
+				}, true, @[ NSDefaultRunLoopMode, UnrealResizeEventMode, UnrealShowEventMode, UnrealFullscreenEventMode, UnrealCloseEventMode ]);
 			}
 		}
 	}
@@ -1420,7 +1420,7 @@ void FMacApplication::OnWindowOrderedFront(TSharedRef<FMacWindow> Window)
 			{
 				[CocoaWindow UpdateAccessibilityView:WindowId];
 			}
-		}, NSDefaultRunLoopMode, false);
+		}, false);
 	}
 #endif
 
@@ -1515,7 +1515,7 @@ void FMacApplication::OnApplicationDidBecomeActive()
 		// For non-editor clients, record if the active window is in focus
 		FGenericCrashContext::SetEngineData(TEXT("Platform.AppHasFocus"), TEXT("true"));
 #endif
-	}, @[ NSDefaultRunLoopMode ], false);
+	}, false);
 	
 	// Call out to update isOnActiveSpace per FCocoaWindow when the app becomes active.  Works round unrespostive crash reporter on Catalina
 	// First call to isOnActiveSpace in OnActiveSpaceDidChange() during crash reporter startup returns false on Catalina but true on BigSur.
@@ -1575,7 +1575,7 @@ void FMacApplication::OnApplicationWillResignActive()
 		// For non-editor clients, record if the active window is in focus
 		FGenericCrashContext::SetEngineData(TEXT("Platform.AppHasFocus"), TEXT("false"));
 #endif
-	}, @[ NSDefaultRunLoopMode ], false);
+	}, false);
 }
 
 void FMacApplication::OnWindowsReordered()
@@ -1646,7 +1646,7 @@ void FMacApplication::OnCursorLock()
 					SCOPED_AUTORELEASE_POOL;
 					[NativeWindow setMinSize:NSMakeSize(NativeWindow.frame.size.width, NativeWindow.frame.size.height)];
 					[NativeWindow setMaxSize:NSMakeSize(NativeWindow.frame.size.width, NativeWindow.frame.size.height)];
-				}, NSDefaultRunLoopMode, false);
+				}, false);
 			}
 			else
 			{
@@ -1660,7 +1660,7 @@ void FMacApplication::OnCursorLock()
 						SCOPED_AUTORELEASE_POOL;
 						[NativeWindow setMinSize:MinSize];
 						[NativeWindow setMaxSize:MaxSize];
-					}, NSDefaultRunLoopMode, false);
+					}, false);
 				}
 			}
 		}
@@ -1808,7 +1808,7 @@ void FMacApplication::UpdateScreensArray()
 		{
 			GAllScreens.Add(MakeShareable(new FMacScreen(Screen)));
 		}
-	}, NSDefaultRunLoopMode, true);
+	});
 
 	FScopeLock Lock(&GAllScreensMutex);
 
@@ -2205,7 +2205,7 @@ void FMacApplication::CloseQueuedWindows()
 				[Window close];
 				[Window release];
 			}
-		}, UnrealCloseEventMode, true);
+		}, true, UnrealCloseEventMode);
 	}
 }
 
@@ -2225,7 +2225,7 @@ void FMacApplication::InvalidateTextLayouts()
 				}
 			}
 
-		}, UnrealIMEEventMode, true);
+		}, true, UnrealIMEEventMode);
 
 		WindowsRequiringTextInvalidation.Empty();
 	}
@@ -2394,21 +2394,19 @@ void FMacApplication::OnAccessibleEventRaised(const FAccessibleEventArgs& Args)
 					{
 						[[FMacAccessibilityManager AccessibilityManager] RemoveAccessibilitySubtree:Id];
 					}
-					
-
 				}
 				// LayoutChanged is to indicate things like "a widget became visible or hidden" while
 				// ScreenChanged is for large-scale UI changes. It can potentially take an NSString to read
 				// to the user when this happens, if we choose to support that.
 				//@TODO: Posting a notification with user info to the NSApp doesn't seem to do anything. Find out wy.
-			}, NSDefaultRunLoopMode, false);
+			}, false);
 			break;
 		}
 		case EAccessibleEvent::WidgetRemoved:
 		{
 			MainThreadCall(^{
 				[[FMacAccessibilityManager AccessibilityManager] RemoveAccessibilityElement:Id];
-			}, NSDefaultRunLoopMode, false);
+			}, false);
 			break;
 		}
 		case EAccessibleEvent::Notification:
@@ -2423,7 +2421,7 @@ void FMacApplication::OnAccessibleEventRaised(const FAccessibleEventArgs& Args)
 					NSAccessibilityPostNotificationWithUserInfo(NSApp.mainWindow, NSAccessibilityAnnouncementRequestedNotification, AnnouncementInfo);
 					AccessibilityAnnouncementDelayTimer = Nil;
 				}];
-			}, NSDefaultRunLoopMode, false);
+			}, false);
 			break;
 		}
 	}
@@ -2465,8 +2463,8 @@ void FMacApplication::OnVoiceoverEnabled()
 			{
 				AccessibilityCacheTimer = [NSTimer scheduledTimerWithTimeInterval:0.25f target:[FMacAccessibilityManager AccessibilityManager] selector:@selector(UpdateAllCachedProperties) userInfo:nil repeats:YES];
 			}
-		}, NSDefaultRunLoopMode, false);
-	}, @[ NSDefaultRunLoopMode ], false);
+		}, false);
+	}, false);
 }
 
 void FMacApplication::OnVoiceoverDisabled()
@@ -2495,7 +2493,7 @@ void FMacApplication::OnVoiceoverDisabled()
 		[[FMacAccessibilityManager AccessibilityManager] Clear];
 		//This releases the accessibility element referenced by the app, deallocating it
 		NSApp.accessibilityApplicationFocusedUIElement = Nil;
-	}, NSDefaultRunLoopMode, false);
+	}, false);
 }
 
 #endif
diff --git a/Engine/Source/Runtime/ApplicationCore/Private/Mac/MacConsoleOutputDevice.cpp b/Engine/Source/Runtime/ApplicationCore/Private/Mac/MacConsoleOutputDevice.cpp
--- a/Engine/Source/Runtime/ApplicationCore/Private/Mac/MacConsoleOutputDevice.cpp
+++ b/Engine/Source/Runtime/ApplicationCore/Private/Mac/MacConsoleOutputDevice.cpp
@@ -110,7 +110,7 @@ void FMacConsoleOutputDevice::CreateConsole()
 		}
 		
 		SetDefaultTextColor();
-	}, UnrealNilEventMode, true);
+	}, true, UnrealNilEventMode);
 }
 
 void FMacConsoleOutputDevice::DestroyConsole()
@@ -137,7 +137,7 @@ void FMacConsoleOutputDevice::DestroyConsole()
 			
 			[ConsoleWindow close];
 			TextViewTextColor = NULL;
-		}, UnrealNilEventMode, true);
+		}, true, UnrealNilEventMode);
 	}
 }
 
@@ -216,7 +216,7 @@ void FMacConsoleOutputDevice::Serialize( const TCHAR* Data, ELogVerbosity::Type
 						[Colors release];
 						[AttributeKeys release];
 						OutstandingTasks--;
-					}, NSDefaultRunLoopMode, false);
+					}, false);
 				}
 			}
 			else
@@ -232,7 +232,7 @@ void FMacConsoleOutputDevice::Serialize( const TCHAR* Data, ELogVerbosity::Type
 					[TextView scrollRangeToVisible:NSMakeRange([[TextView string] length], 0)];
 					[AttributedString release];
 					OutstandingTasks--;
-				}, NSDefaultRunLoopMode, false);
+				}, false);
 				
 				if(!MacApplication)
 				{
@@ -279,7 +279,7 @@ void FMacConsoleOutputDevice::SetDefaultTextColor()
 		[Colors release];
 		[AttributeKeys release];
 		OutstandingTasks--;
-	}, NSDefaultRunLoopMode, false);
+	}, false);
 }
 
 @implementation FMacConsoleWindow
diff --git a/Engine/Source/Runtime/ApplicationCore/Private/Mac/MacPlatformApplicationMisc.cpp b/Engine/Source/Runtime/ApplicationCore/Private/Mac/MacPlatformApplicationMisc.cpp
--- a/Engine/Source/Runtime/ApplicationCore/Private/Mac/MacPlatformApplicationMisc.cpp
+++ b/Engine/Source/Runtime/ApplicationCore/Private/Mac/MacPlatformApplicationMisc.cpp
@@ -270,7 +270,7 @@ void FMacPlatformApplicationMisc::PreInit()
 					[AlertPanel runModal];
 					[AlertPanel release];
 					exit(1);
-				}, NSDefaultRunLoopMode, true);
+				});
 			}
 		}
 	}
@@ -322,7 +322,7 @@ void FMacPlatformApplicationMisc::TearDown()
 		MainThreadCall(^{
 			[[NSProcessInfo processInfo] endActivity:CommandletActivity];
 			[CommandletActivity release];
-		}, NSDefaultRunLoopMode, false);
+		}, false);
 		CommandletActivity = nil;
 	}
 }
@@ -499,7 +499,7 @@ void FMacPlatformApplicationMisc::PumpMessages(bool bFromMainLoop)
 					UpdateApplicationMenu();
 					UpdateWindowMenu();
 					UpdateCocoaWindows();
-				}, NSDefaultRunLoopMode, false);
+				}, false);
 
 				UpdateCachedMacMenuStateCb();
 				
@@ -554,7 +554,7 @@ void FMacPlatformApplicationMisc::ActivateApplication()
 {
 	MainThreadCall(^{
 		[NSApp activateIgnoringOtherApps:YES];
-	}, NSDefaultRunLoopMode, false);
+	}, false);
 }
 
 void FMacPlatformApplicationMisc::UpdateApplicationMenu()
diff --git a/Engine/Source/Runtime/ApplicationCore/Private/Mac/MacPlatformSplash.cpp b/Engine/Source/Runtime/ApplicationCore/Private/Mac/MacPlatformSplash.cpp
--- a/Engine/Source/Runtime/ApplicationCore/Private/Mac/MacPlatformSplash.cpp
+++ b/Engine/Source/Runtime/ApplicationCore/Private/Mac/MacPlatformSplash.cpp
@@ -306,7 +306,7 @@ void FMacPlatformSplash::Show()
 				[GSplashScreenImage release];
 				GSplashScreenImage = NULL;
 			}
-		}, NSDefaultRunLoopMode, true);
+		});
 
 		FMacPlatformApplicationMisc::PumpMessages(true);
 	}
@@ -325,7 +325,7 @@ void FMacPlatformSplash::Hide()
 
 			[GSplashScreenImage release];
 			GSplashScreenImage = NULL;
-		}, NSDefaultRunLoopMode, true);
+		});
 
 		FMacPlatformApplicationMisc::PumpMessages(true);
 	}
diff --git a/Engine/Source/Runtime/ApplicationCore/Private/Mac/MacTextInputMethodSystem.cpp b/Engine/Source/Runtime/ApplicationCore/Private/Mac/MacTextInputMethodSystem.cpp
--- a/Engine/Source/Runtime/ApplicationCore/Private/Mac/MacTextInputMethodSystem.cpp
+++ b/Engine/Source/Runtime/ApplicationCore/Private/Mac/MacTextInputMethodSystem.cpp
@@ -69,7 +69,7 @@ namespace
 					FCocoaTextView* TextView = (FCocoaTextView*)[CocoaWindow openGLView];
 					[[TextView inputContext] invalidateCharacterCoordinates];
 				}
-			}, UnrealIMEEventMode, true);
+			}, true, UnrealIMEEventMode);
 		}
 	}
 	
@@ -85,7 +85,7 @@ namespace
 					FCocoaTextView* TextView = (FCocoaTextView*)[CocoaWindow openGLView];
 					[[TextView inputContext] invalidateCharacterCoordinates];
 				}
-			}, UnrealIMEEventMode, true);
+			}, true, UnrealIMEEventMode);
 		}
 	}
 	
@@ -102,7 +102,7 @@ namespace
 					[[TextView inputContext] discardMarkedText];
 					[TextView unmarkText];
 				}
-			}, UnrealIMEEventMode, true);
+			}, true, UnrealIMEEventMode);
 		}
 	}
 }
diff --git a/Engine/Source/Runtime/ApplicationCore/Private/Mac/MacWindow.cpp b/Engine/Source/Runtime/ApplicationCore/Private/Mac/MacWindow.cpp
--- a/Engine/Source/Runtime/ApplicationCore/Private/Mac/MacWindow.cpp
+++ b/Engine/Source/Runtime/ApplicationCore/Private/Mac/MacWindow.cpp
@@ -215,7 +215,7 @@ PRAGMA_DISABLE_DEPRECATION_WARNINGS
 PRAGMA_ENABLE_DEPRECATION_WARNINGS
 			check(0);
 		}
-	}, UnrealShowEventMode, true);
+	}, true, UnrealShowEventMode);
 }
 
 FCocoaWindow* FMacWindow::GetWindowHandle() const
@@ -251,7 +251,7 @@ void FMacWindow::MoveWindowTo( int32 X, int32 Y )
 		SCOPED_AUTORELEASE_POOL;
 		const FVector2D Point = FMacApplication::ConvertSlatePositionToCocoa(X, Y);
 		[WindowHandle setFrameOrigin:NSMakePoint(Point.X, Point.Y - [WindowHandle openGLFrame].size.height + 1)];
-	}, UnrealResizeEventMode, true);
+	}, true, UnrealResizeEventMode);
 }
 
 void FMacWindow::BringToFront( bool bForce )
@@ -268,7 +268,7 @@ void FMacWindow::BringToFront( bool bForce )
 		MainThreadCall(^{
 			SCOPED_AUTORELEASE_POOL;
 			[WindowHandleCopy orderFrontAndMakeMain:bOrderAndKey andKey:bOrderAndKey];
-		}, UnrealShowEventMode, true);
+		}, true, UnrealShowEventMode);
 
 		MacApplication->OnWindowOrderedFront(SharedThis(this));
 	}
@@ -286,7 +286,7 @@ void FMacWindow::Destroy()
 			SCOPED_AUTORELEASE_POOL;
 			[WindowHandleCopy setAlphaValue:0.0f];
 			[WindowHandleCopy setBackgroundColor:[NSColor clearColor]];
-		}, UnrealShowEventMode, false);
+		}, false);
 
 		MacApplication->OnWindowDestroyed(SharedThis(this));
 		WindowHandle = nullptr;
@@ -298,7 +298,7 @@ void FMacWindow::Minimize()
 	MainThreadCall(^{
 		SCOPED_AUTORELEASE_POOL;
 		[WindowHandle miniaturize:nil];
-	}, UnrealResizeEventMode, true);
+	}, true, UnrealResizeEventMode);
 }
 
 void FMacWindow::Maximize()
@@ -310,7 +310,7 @@ void FMacWindow::Maximize()
 			WindowHandle->bZoomed = true;
 			[WindowHandle zoom:nil];
 		}
-	}, UnrealResizeEventMode, true);
+	}, true, UnrealResizeEventMode);
 }
 
 void FMacWindow::Restore()
@@ -325,7 +325,7 @@ void FMacWindow::Restore()
 		{
 			[WindowHandle zoom:nil];
 		}
-	}, UnrealResizeEventMode, true);
+	}, true, UnrealResizeEventMode);
 
 	WindowHandle->bZoomed = WindowHandle.zoomed;
 }
@@ -362,7 +362,7 @@ void FMacWindow::Show()
 		MainThreadCall(^{
 			SCOPED_AUTORELEASE_POOL;
 			[WindowHandleCopy orderFrontAndMakeMain:bShouldActivate andKey:bShouldActivate];
-		}, UnrealShowEventMode, true);
+		}, true, UnrealShowEventMode);
 
 		bIsVisible = true;
 	}
@@ -377,7 +377,7 @@ void FMacWindow::Hide()
 		MainThreadCall(^{
 			SCOPED_AUTORELEASE_POOL;
 			[WindowHandle orderOut:nil];
-		}, UnrealCloseEventMode, true);
+		}, true, UnrealCloseEventMode);
 	}
 }
 
@@ -450,7 +450,7 @@ void FMacWindow::SetOpacity( const float InOpacity )
 		SCOPED_AUTORELEASE_POOL;
         CachedOpacity = InOpacity;
 		[WindowHandle setAlphaValue:InOpacity];
-	}, UnrealNilEventMode, true);
+	}, true, UnrealNilEventMode);
 }
 
 bool FMacWindow::IsPointInWindow( int32 X, int32 Y ) const
@@ -528,7 +528,7 @@ void FMacWindow::SetText(const TCHAR* const Text)
 			}
 			CFRelease( CFName );
 		}
-	}, UnrealNilEventMode, true);
+	}, true, UnrealNilEventMode);
 }
 
 bool FMacWindow::IsRegularWindow() const
@@ -570,7 +570,7 @@ void FMacWindow::SetNativeWindowButtonsVisibility(bool bVisible)
 		{
 			MaximizeButton.hidden = bHidden;
 		}
-	}, NSDefaultRunLoopMode, false);
+	}, false);
 }
 
 void FMacWindow::OnDisplayReconfiguration(CGDirectDisplayID Display, CGDisplayChangeSummaryFlags Flags)
@@ -621,7 +621,7 @@ void FMacWindow::ApplySizeAndModeChanges(int32 X, int32 Y, int32 Width, int32 He
 			{
 				CGDisplayFade(FadeReservationToken, 0.3f, kCGDisplayBlendNormal, kCGDisplayBlendSolidColor, 0, 0, 0, true);
 			}
-		}, UnrealResizeEventMode, true);
+		}, true, UnrealResizeEventMode);
 	}
 
 	if (WindowMode == EWindowMode::Windowed || WindowMode == EWindowMode::WindowedFullscreen)
@@ -630,7 +630,7 @@ void FMacWindow::ApplySizeAndModeChanges(int32 X, int32 Y, int32 Width, int32 He
 		{
 			MainThreadCall(^{
 				CGDisplaySetDisplayMode(WindowedModeSavedState.CapturedDisplayID, WindowedModeSavedState.DesktopDisplayMode, nullptr);
-			}, UnrealResizeEventMode, true);
+			}, true, UnrealResizeEventMode);
 
 			CGDisplayModeRelease(WindowedModeSavedState.DesktopDisplayMode);
 			WindowedModeSavedState.DesktopDisplayMode = nullptr;
@@ -676,7 +676,7 @@ void FMacWindow::ApplySizeAndModeChanges(int32 X, int32 Y, int32 Width, int32 He
 					{
 						[WindowHandle setContentAspectRatio:NSMakeSize((float)Width / (float)Height, 1.0f)];
 					}
-				}, UnrealResizeEventMode, true);
+				}, true, UnrealResizeEventMode);
 			}
 		}
 	}
@@ -699,7 +699,7 @@ void FMacWindow::ApplySizeAndModeChanges(int32 X, int32 Y, int32 Width, int32 He
 			CGDisplayModeRef DisplayMode = FPlatformApplicationMisc::GetSupportedDisplayMode(WindowedModeSavedState.CapturedDisplayID, Width, Height);
 			MainThreadCall(^{
 				CGDisplaySetDisplayMode(WindowedModeSavedState.CapturedDisplayID, DisplayMode, nullptr);
-			}, UnrealResizeEventMode, true);
+			}, true, UnrealResizeEventMode);
 
 			UpdateFullScreenState(bIsFullScreen != bWantsFullScreen);
 
@@ -717,7 +717,7 @@ void FMacWindow::ApplySizeAndModeChanges(int32 X, int32 Y, int32 Width, int32 He
 		MainThreadCall(^{
 			CGDisplayFade(FadeReservationToken, 0.5, kCGDisplayBlendNormal, kCGDisplayBlendSolidColor, 0, 0, 0, false);
 			CGReleaseDisplayFadeReservation(FadeReservationToken);
-		}, UnrealResizeEventMode, false);
+		}, false);
 	}
 }
 
@@ -750,7 +750,7 @@ void FMacWindow::UpdateFullScreenState(bool bToggleFullScreen)
 			[WindowHandle setLevel:WindowedModeSavedState.WindowLevel];
 			[NSApp setPresentationOptions:NSApplicationPresentationDefault];
 		}
-	}, UnrealFullscreenEventMode, true);
+	}, true, UnrealFullscreenEventMode);
 
 	// If we toggle fullscreen, ensure that the window has transitioned BEFORE leaving this function.
 	// This prevents problems with failure to correctly update mouse locks and rendering contexts due to bad event ordering.
diff --git a/Engine/Source/Runtime/Core/Private/Mac/CocoaThread.cpp b/Engine/Source/Runtime/Core/Private/Mac/CocoaThread.cpp
--- a/Engine/Source/Runtime/Core/Private/Mac/CocoaThread.cpp
+++ b/Engine/Source/Runtime/Core/Private/Mac/CocoaThread.cpp
@@ -382,13 +382,13 @@ FCocoaRunLoopSource* FCocoaRunLoopSource::GameRunLoopSource = nullptr;
 		MainThreadCall(^{
 			[NSApp replyToApplicationShouldTerminate:YES];
 			[[NSProcessInfo processInfo] enableSuddenTermination];
-		}, NSDefaultRunLoopMode, false);
+		}, false);
 	}
 	else
 	{
 		MainThreadCall(^{
 			[[NSProcessInfo processInfo] enableSuddenTermination];
-		}, NSDefaultRunLoopMode, false);
+		}, false);
 	}
 	
 	// And now it is time to die.
@@ -447,7 +447,7 @@ static void PerformBlockOnThread(FCocoaRunLoopSource& ThreadSource, dispatch_blo
 	}
 }
 
-void MainThreadCall(dispatch_block_t Block, NSString* WaitMode, bool const bWait)
+void MainThreadCall(dispatch_block_t Block, bool const bWait, NSString* WaitMode)
 {
 	if ( [NSThread mainThread] != [NSThread currentThread] )
 	{
@@ -460,7 +460,7 @@ void MainThreadCall(dispatch_block_t Block, NSString* WaitMode, bool const bWait
 	}
 }
 
-void GameThreadCall(dispatch_block_t Block, NSArray* SendModes, bool const bWait)
+void GameThreadCall(dispatch_block_t Block, bool const bWait, NSArray* SendModes)
 {
 	if ( [NSThread gameThread] != [NSThread currentThread] )
 	{
diff --git a/Engine/Source/Runtime/Core/Private/Mac/MacPlatformProcess.cpp b/Engine/Source/Runtime/Core/Private/Mac/MacPlatformProcess.cpp
--- a/Engine/Source/Runtime/Core/Private/Mac/MacPlatformProcess.cpp
+++ b/Engine/Source/Runtime/Core/Private/Mac/MacPlatformProcess.cpp
@@ -1218,7 +1218,7 @@ bool FMacPlatformProcess::LaunchFileInDefaultExternalApplication( const TCHAR* F
 
 void FMacPlatformProcess::ExploreFolder( const TCHAR* FilePath )
 {
-    extern void MainThreadCall(dispatch_block_t Block, NSString* WaitMode, bool const bWait);
+    extern void MainThreadCall(dispatch_block_t Block, bool const bWait, NSString* WaitMode);
 
 	SCOPED_AUTORELEASE_POOL;
 	__block CFStringRef CFFilePath = FPlatformString::TCHARToCFString( FilePath );
@@ -1238,7 +1238,7 @@ void FMacPlatformProcess::ExploreFolder( const TCHAR* FilePath )
 		}
 	}
 	CFRelease( CFFilePath );
-    }, NSDefaultRunLoopMode, false);
+    }, false, nullptr);
 }
 
 void FMacPlatformProcess::ClosePipe( void* ReadPipe, void* WritePipe )
diff --git a/Engine/Source/Runtime/Core/Public/Mac/CocoaThread.h b/Engine/Source/Runtime/Core/Public/Mac/CocoaThread.h
--- a/Engine/Source/Runtime/Core/Public/Mac/CocoaThread.h
+++ b/Engine/Source/Runtime/Core/Public/Mac/CocoaThread.h
@@ -24,23 +24,65 @@ CORE_API extern NSString* UnrealIMEEventMode; /* Process only input method event
 - (void)main; // Override that sets the variable backing +[NSRunLoop gameRunLoop], do not override in a subclass.
 @end
 
-CORE_API void MainThreadCall(dispatch_block_t Block, NSString* WaitMode = NSDefaultRunLoopMode, bool const bWait = true);
+/**
+ * Schedule a block to be executed on MainThread.
+ * @param Block - The block to execute on MainThread, example: ^{ UE_LOG(LogMac, Log, TEXT("Hello MainThread"); }
+ * @param bWait - Wether or not to wait until the block is executed.
+ * @param WaitMode - The mode we are allowed to execute while waiting. (Only use when bWait = true)
+ * 
+ * Avoid being too restrictive on the WaitMode or you could cause a deadlock by preventing further progress.
+ */
+CORE_API void MainThreadCall(dispatch_block_t Block, bool const bWait = true, NSString* WaitMode = NSDefaultRunLoopMode);
 
+UE_DEPRECATED(5.6, "Use the alternative with inverted last parameters")
+CORE_API inline void MainThreadCall(dispatch_block_t Block, NSString* WaitMode, bool const bWait = true)
+{
+	MainThreadCall(Block, bWait, WaitMode);
+}
+
+/**
+ * Schedule a block to be executed on MainThread with a return value.
+ * @param Block - The block to execute on MainThread, example: ^{ return 5; }
+ * @param WaitMode - The mode we are allowed to execute while waiting.
+ * 
+ * Avoid being too restrictive on the WaitMode or you could cause a deadlock by preventing further progress.
+ */
 template<typename ReturnType>
 ReturnType MainThreadReturn(ReturnType (^Block)(void), NSString* WaitMode = NSDefaultRunLoopMode)
 {
 	__block ReturnType ReturnValue;
-	MainThreadCall(^{ ReturnValue = Block(); }, WaitMode, true);
+	MainThreadCall(^{ ReturnValue = Block(); }, true, WaitMode);
 	return ReturnValue;
 }
 
-CORE_API void GameThreadCall(dispatch_block_t Block, NSArray* SendModes = @[ NSDefaultRunLoopMode ], bool const bWait = true);
+/**
+ * Schedule a block to be executed on GameThread
+ * @param Block - The block to execute on GameThread, example: ^{ UE_LOG(LogMac, Log, TEXT("Hello GameThread"); }
+ * @param bWait - Wether or not to wait until the block is executed.
+ * @param SendModes - The modes for the block we are scheduling, this is use to specify who can execute the scheduled block.
+ * 
+ * Be careful to make sure that GameThread will be processing at least one of the modes or you could end up causing a deadlock by preventing further progress.
+ */
+CORE_API void GameThreadCall(dispatch_block_t Block, bool const bWait = true, NSArray* SendModes = @[ NSDefaultRunLoopMode ]);
+
+UE_DEPRECATED(5.6, "Use the alternative with inverted last parameters")
+CORE_API inline void GameThreadCall(dispatch_block_t Block, NSArray* SendModes, bool const bWait = true)
+{
+	GameThreadCall(Block, bWait, SendModes);
+}
 
+/**
+ * Schedule a block to be executed on GameThread with a return value.
+ * @param Block - The block to execute on GameThread, example: ^{ return 5; }
+ * @param SendModes - The modes for the block we are scheduling, this is use to specify who can execute the scheduled block.
+ * 
+ * Be careful to make sure that GameThread will be processing at least one of the modes or you could end up causing a deadlock by preventing further progress.
+ */
 template<typename ReturnType>
 ReturnType GameThreadReturn(ReturnType (^Block)(void), NSArray* SendModes = @[ NSDefaultRunLoopMode ])
 {
 	__block ReturnType ReturnValue;
-	GameThreadCall(^{ ReturnValue = Block(); }, SendModes, true);
+	GameThreadCall(^{ ReturnValue = Block(); }, true, SendModes);
 	return ReturnValue;
 }
 
diff --git a/Engine/Source/Runtime/Launch/Private/Mac/LaunchMac.cpp b/Engine/Source/Runtime/Launch/Private/Mac/LaunchMac.cpp
--- a/Engine/Source/Runtime/Launch/Private/Mac/LaunchMac.cpp
+++ b/Engine/Source/Runtime/Launch/Private/Mac/LaunchMac.cpp
@@ -197,7 +197,7 @@ static int32 MacOSVersionCompare(const NSOperatingSystemVersion& VersionA, const
 				GEngine->DeferredCommands.Add(TEXT("EXIT"));
 			}
 		}
-	}, @[ NSDefaultRunLoopMode ], false);
+	}, false);
 }
 
 - (IBAction)showAboutWindow:(id)Sender
@@ -208,7 +208,7 @@ static int32 MacOSVersionCompare(const NSOperatingSystemVersion& VersionA, const
 		{
 			FModuleManager::GetModuleChecked<IMainFrameModule>(TEXT("MainFrame")).ShowAboutWindow();
 		}
-	}, @[ NSDefaultRunLoopMode ], false);
+	}, false);
 #else
 	[NSApp orderFrontStandardAboutPanel:Sender];
 #endif
@@ -222,7 +222,7 @@ static int32 MacOSVersionCompare(const NSOperatingSystemVersion& VersionA, const
 		{
 			FModuleManager::LoadModuleChecked<ISettingsModule>("Settings").ShowViewer(FName("Editor"), FName("General"), FName("Appearance"));
 		}
-	}, @[ NSDefaultRunLoopMode ], false);
+	}, false);
 }
 #endif
 
diff --git a/Engine/Source/Runtime/Slate/Private/Framework/MultiBox/Mac/MacMenu.cpp b/Engine/Source/Runtime/Slate/Private/Framework/MultiBox/Mac/MacMenu.cpp
--- a/Engine/Source/Runtime/Slate/Private/Framework/MultiBox/Mac/MacMenu.cpp
+++ b/Engine/Source/Runtime/Slate/Private/Framework/MultiBox/Mac/MacMenu.cpp
@@ -101,7 +101,7 @@ static FCriticalSection GMacMenuEntryBlockCS;;
 	
 	GameThreadCall(^{
 		FSlateApplication::Get().ClearKeyboardFocus( EFocusCause::WindowActivate );
-	}, @[ NSDefaultRunLoopMode ], false);
+	}, false);
 
 }
 
@@ -341,7 +341,7 @@ void FSlateMacMenu::CleanupOnShutdown()
 				Block.Reset();
 			}
 			GMacMenuEntryBlockDeletionQueue.Empty();
-		}, @[ NSDefaultRunLoopMode ], true);
+		});
 	}
 }
 
@@ -437,7 +437,7 @@ void FSlateMacMenu::PostInitStartup()
             [WindowMenu addItem:[NSMenuItem separatorItem]];
             [WindowMenu addItem:BringAllToFrontItem];
             [WindowMenu addItem:[NSMenuItem separatorItem]];
-        }, NSDefaultRunLoopMode, false);
+        }, false);
 	}
 }
 
@@ -612,7 +612,7 @@ void FSlateMacMenu::UpdateWithMultiBox(const TSharedPtr< FMultiBox > MultiBox)
 		}
 
 		FPlatformApplicationMisc::bChachedMacMenuStateNeedsUpdate = true;
-	}, NSDefaultRunLoopMode, false);
+	}, false);
 	// Now we destroy all the menu blocks in the game thread to ensure ~SWidget() is only called in Game Thread
 	GameThreadCall(^{
 		FScopeLock MacMenuEntryBlockLock(&GMacMenuEntryBlockCS);
@@ -621,7 +621,7 @@ void FSlateMacMenu::UpdateWithMultiBox(const TSharedPtr< FMultiBox > MultiBox)
 			Block.Reset();
 		}
 		GMacMenuEntryBlockDeletionQueue.Empty();
-	}, @[NSDefaultRunLoopMode], false);
+	}, false);
 }
 
 void FSlateMacMenu::UpdateMenu(FMacMenu* Menu)
@@ -914,7 +914,7 @@ void FSlateMacMenu::ExecuteMenuItemAction(const TSharedRef< const class FMenuEnt
 				(*MenuBlock)->GetDirectActions().Execute();
 			}
 			delete MenuBlock;
-		}, @[ NSDefaultRunLoopMode ], false);
+		}, false);
 	}
 }
 
diff --git a/Engine/Source/Runtime/WebBrowser/Private/CEF/CEFWebBrowserWindow.cpp b/Engine/Source/Runtime/WebBrowser/Private/CEF/CEFWebBrowserWindow.cpp
--- a/Engine/Source/Runtime/WebBrowser/Private/CEF/CEFWebBrowserWindow.cpp
+++ b/Engine/Source/Runtime/WebBrowser/Private/CEF/CEFWebBrowserWindow.cpp
@@ -2738,7 +2738,7 @@ void FCEFWebBrowserWindow::ProcessPendingNavigation()
 		{
 			MainThreadCall(^{
 				MainFrame->LoadURL(Url);
-			}, NSDefaultRunLoopMode, true);
+			});
 		}
 #else
 		MainFrame->LoadURL(Url);
