Fix nullptr exception in FMacWindow

diff --git a/Engine/Source/Runtime/ApplicationCore/Private/Mac/MacWindow.cpp b/Engine/Source/Runtime/ApplicationCore/Private/Mac/MacWindow.cpp
--- a/Engine/Source/Runtime/ApplicationCore/Private/Mac/MacWindow.cpp
+++ b/Engine/Source/Runtime/ApplicationCore/Private/Mac/MacWindow.cpp
@@ -510,7 +510,11 @@ bool FMacWindow::IsForegroundWindow() const
 {
 	// @TODO: UE-227397
 	// Currently using cached property to avoid triggering a MainThreadChecker issue.
-	return WindowHandle->bKeyWindow;
+	if (WindowHandle)
+	{
+		return WindowHandle->bKeyWindow;
+	}
+	return false;
 }
 
 void FMacWindow::SetText(const TCHAR* const Text)
