Fix a special case to detect a monitor set in portrait mode while the OS have the draw rectangle set in landscape mode.

diff --git a/Engine/Source/Runtime/ApplicationCore/Private/Windows/WindowsApplication.cpp b/Engine/Source/Runtime/ApplicationCore/Private/Windows/WindowsApplication.cpp
--- a/Engine/Source/Runtime/ApplicationCore/Private/Windows/WindowsApplication.cpp
+++ b/Engine/Source/Runtime/ApplicationCore/Private/Windows/WindowsApplication.cpp
@@ -1312,6 +1312,14 @@ static void GetMonitorsInfo(TArray<FMonitorInfo>& OutMonitorInfo)
 
 					if (GetSizeForDevID(Info.Name, Info.NativeWidth, Info.NativeHeight, Info.DPI))
 					{
+						// usually expected: 4:3=1.333, 16:9=1.777, 16:10=1.6, multi-monitor-wide: >2
+						bool bIsPortrait = Info.NativeWidth < Info.NativeHeight;
+						// special case to detect specific platform display possibly rotated by the OS but the rotation is not reported in the EDID data
+						if (bIsPortrait && (Info.DisplayRect.Bottom == Info.NativeWidth) && (Info.DisplayRect.Right == Info.NativeHeight))
+						{
+							Swap(Info.NativeWidth, Info.NativeHeight);
+						}
+
 						Info.ID = Monitor.DeviceID;
 						Info.bIsPrimary = (DisplayDevice.StateFlags & DISPLAY_DEVICE_PRIMARY_DEVICE) > 0;
 
