FGenericApplication - Fix ensure triggered on Mac because some code is executed on MainThread != GameThread

#tests Local repro when exiting sleep mode and unlocking mac while the editor is running

diff --git a/Engine/Source/Runtime/ApplicationCore/Private/GenericPlatform/GenericApplication.cpp b/Engine/Source/Runtime/ApplicationCore/Private/GenericPlatform/GenericApplication.cpp
--- a/Engine/Source/Runtime/ApplicationCore/Private/GenericPlatform/GenericApplication.cpp
+++ b/Engine/Source/Runtime/ApplicationCore/Private/GenericPlatform/GenericApplication.cpp
@@ -143,7 +143,11 @@ float FDisplayMetrics::GetDebugTitleSafeZoneRatio()
 
 bool FDisplayMetrics::TryGetTitleSafeZoneOverwrite(FVector4& InOutTitleSafePaddingSize)
 {
+#if PLATFORM_MAC
+	if (CVarEnableSafeZoneOverrides.GetValueOnAnyThread())
+#else
 	if (CVarEnableSafeZoneOverrides.GetValueOnGameThread())
+#endif
 	{
 		InOutTitleSafePaddingSize.X = (float)PrimaryDisplayWidth * (1.0f - CVarSafeZoneRatioLeft.GetValueOnGameThread()) * 0.5f;
 		InOutTitleSafePaddingSize.Y = (float)PrimaryDisplayHeight * (1.0f - CVarSafeZoneRatioTop.GetValueOnGameThread()) * 0.5f;
