MacSplashScreen - Fix unreadable text when using non-retina resolution & make the text look more like Windows

#tests local testing with and without retina resolution on external display

diff --git a/Engine/Source/Runtime/ApplicationCore/Private/Mac/MacPlatformSplash.cpp b/Engine/Source/Runtime/ApplicationCore/Private/Mac/MacPlatformSplash.cpp
--- a/Engine/Source/Runtime/ApplicationCore/Private/Mac/MacPlatformSplash.cpp
+++ b/Engine/Source/Runtime/ApplicationCore/Private/Mac/MacPlatformSplash.cpp
@@ -97,10 +97,8 @@ static void StartSetSplashText( const SplashTextType::Type InType, const FText&
 			switch ( CurTypeIndex )
 			{
 				case SplashTextType::StartupProgress:
-					Brightness = 180.0f / 255.0f;
-					break;
 				case SplashTextType::VersionInfo1:
-					Brightness = 240.0f / 255.0f;
+					Brightness = 180.0f / 255.0f;
 					break;
 				case SplashTextType::GameName:
 					Brightness = 240.0f / 255.0f;
@@ -109,11 +107,11 @@ static void StartSetSplashText( const SplashTextType::Type InType, const FText&
 
 			NSColor *TextColor = [NSColor colorWithDeviceRed: Brightness green: Brightness blue: Brightness alpha: 1.0f];
 
-			NSString* FontName = @"Helvetica-Bold";
+			NSString* FontName = @"Helvetica";
 			switch ( CurTypeIndex )
 			{
 			case SplashTextType::GameName:
-				FontName = @"Verdana-Bold";
+				FontName = @"Helvetica-Bold";
 				break;
 			}
 
@@ -161,7 +159,6 @@ static void StartSetSplashText( const SplashTextType::Type InType, const FText&
 				Color,
 				[NSFont fontWithName: FontName size: FontSize],
 				[NSColor colorWithDeviceRed: 0.0 green: 0.0 blue: 0.0 alpha: 1.0],
-				[NSNumber numberWithFloat:-4.0],
 				style,
 				nil ]
 		forKeys:
@@ -169,7 +166,6 @@ static void StartSetSplashText( const SplashTextType::Type InType, const FText&
 				NSForegroundColorAttributeName,
 				NSFontAttributeName,
 				NSStrokeColorAttributeName,
-				NSStrokeWidthAttributeName,
 				NSParagraphStyleAttributeName,
 				nil]
 		];
