FMacWindow - Fix nullptr exceptions when deferred event process a window already gone.

#tests Local testing

diff --git a/Engine/Source/Runtime/ApplicationCore/Private/Mac/MacWindow.cpp b/Engine/Source/Runtime/ApplicationCore/Private/Mac/MacWindow.cpp
--- a/Engine/Source/Runtime/ApplicationCore/Private/Mac/MacWindow.cpp
+++ b/Engine/Source/Runtime/ApplicationCore/Private/Mac/MacWindow.cpp
@@ -392,7 +392,11 @@ void FMacWindow::SetWindowMode(EWindowMode::Type NewWindowMode)
 
 EWindowMode::Type FMacWindow::GetWindowMode() const
 {
-	return [WindowHandle windowMode];
+	if (WindowHandle)
+	{
+		return [WindowHandle windowMode];
+	}
+	return EWindowMode::Windowed;
 }
 
 bool FMacWindow::IsMaximized() const
@@ -405,7 +409,7 @@ bool FMacWindow::IsMinimized() const
 	SCOPED_AUTORELEASE_POOL;
 	// @TODO: UE-227397
 	// Currently using cached property to avoid triggering a MainThreadChecker issue.
-	return WindowHandle->bMinimized;
+	return WindowHandle == nullptr || WindowHandle->bMinimized;
 }
 
 bool FMacWindow::IsVisible() const
