Use expected result size to verify downloaded files on iOS

diff --git a/Engine/Source/Runtime/ApplicationCore/Private/IOS/IOSBackgroundURLSessionHandler.cpp b/Engine/Source/Runtime/ApplicationCore/Private/IOS/IOSBackgroundURLSessionHandler.cpp
--- a/Engine/Source/Runtime/ApplicationCore/Private/IOS/IOSBackgroundURLSessionHandler.cpp
+++ b/Engine/Source/Runtime/ApplicationCore/Private/IOS/IOSBackgroundURLSessionHandler.cpp
@@ -53,9 +53,10 @@ static void LogIOSBackgroundDownloadMessage(NSString* Message)
 
 @property (nonatomic, retain) NSMutableArray<__kindof NSURL*>* URLs;
 @property (nonatomic, retain) NSMutableArray<__kindof NSNumber*>* RetryCountPerURL;
+@property (nonatomic) uint64 ExpectedResultSize;
 
 // assumes all URL's have same content path but different domain
-+ (FBackgroundNSURLSessionDownloadTaskData* _Nonnull)TaskDataWithURLs:(NSArray<__kindof NSURL*>*)URLs WithRetryCount:(NSInteger)RetryCount;
++ (FBackgroundNSURLSessionDownloadTaskData* _Nonnull)TaskDataWithURLs:(NSArray<__kindof NSURL*>*)URLs WithRetryCount:(NSInteger)RetryCount WithExpectedResultSize:(uint64)ExpectedResultSize;
 + (FBackgroundNSURLSessionDownloadTaskData* _Nullable)TaskDataFromSerializedString:(NSString*)SerializedData;
 
 - (NSString*)ToSerializedString;
@@ -75,8 +76,9 @@ NSString* const SerializationKeyProtocolVersion = @"v";
 NSString* const SerializationKeyCDNs = @"c";
 NSString* const SerializationKeyPath = @"p";
 NSString* const SerializationKeyRetryCountPerURL = @"r";
+NSString* const SerializationKeyExpectedResultSize = @"s";
 
-+ (FBackgroundNSURLSessionDownloadTaskData* _Nonnull)TaskDataWithURLs:(NSArray<__kindof NSURL*>*)URLs WithRetryCount:(NSInteger)RetryCount
++ (FBackgroundNSURLSessionDownloadTaskData* _Nonnull)TaskDataWithURLs:(NSArray<__kindof NSURL*>*)URLs WithRetryCount:(NSInteger)RetryCount WithExpectedResultSize:(uint64)ExpectedResultSize
 {
 	FBackgroundNSURLSessionDownloadTaskData* Data = [[FBackgroundNSURLSessionDownloadTaskData alloc] init];
 
@@ -89,6 +91,7 @@ NSString* const SerializationKeyRetryCountPerURL = @"r";
 	}
 	NSAssert(Data.URLs.count > 0, @"URLs should be non empty");
 	NSAssert(Data.URLs.count == Data.RetryCountPerURL.count, @"URLs and RetryCountPerURL arrays should have same size");
+	Data.ExpectedResultSize = ExpectedResultSize;
 
 	return [Data autorelease];
 }
@@ -115,17 +118,19 @@ NSString* const SerializationKeyRetryCountPerURL = @"r";
 	NSMutableArray<__kindof NSString*>* CDNs = [Dict valueForKey:SerializationKeyCDNs];
 	NSString* Path = [Dict valueForKey:SerializationKeyPath];
 	NSArray<__kindof NSNumber*>* RetryCountPerURL = [Dict valueForKey:SerializationKeyRetryCountPerURL];
+	NSNumber* ExpectedResultSize = [Dict valueForKey:SerializationKeyExpectedResultSize]; // can be nil
 
 	if (Error != nil || Version == nil || Version.intValue != 1 || CDNs == nil || Path == nil || RetryCountPerURL == nil)
 	{
-		UE_DNLD_LOG(@"Failed to deserialize task state '%@' due to '%@', %u, %u, %u, %u, %u",
+		UE_DNLD_LOG(@"Failed to deserialize task state '%@' due to '%@', %u, %u, %u, %u, %u, %u",
 			SerializedData,
 			Error != nil ? Error.localizedDescription : @"",
 			Version == nil ? 1 : 0,
 			Version.intValue,
 			CDNs == nil ? 1 : 0,
 			Path == nil ? 1 : 0,
-			RetryCountPerURL == nil ? 1 : 0
+			RetryCountPerURL == nil ? 1 : 0,
+			ExpectedResultSize == nil ? 1 : 0
 		);
 		[Data release];
 		return nil;
@@ -139,6 +144,7 @@ NSString* const SerializationKeyRetryCountPerURL = @"r";
 	}
 	NSAssert(Data.URLs.count > 0, @"URLs should be non empty");
 	NSAssert(Data.URLs.count == Data.RetryCountPerURL.count, @"URLs and RetryCountPerURL arrays should have same size");
+	Data.ExpectedResultSize = ExpectedResultSize != nil ? ExpectedResultSize.unsignedLongLongValue : 0;
 
 	return [Data autorelease];
 }
@@ -173,6 +179,10 @@ NSString* const SerializationKeyRetryCountPerURL = @"r";
 	[Dict setValue:CDNs forKey:SerializationKeyCDNs];
 	[Dict setValue:Path forKey:SerializationKeyPath];
 	[Dict setValue:self.RetryCountPerURL forKey:SerializationKeyRetryCountPerURL];
+	if (self.ExpectedResultSize != 0)
+	{
+		[Dict setValue:[NSNumber numberWithUnsignedLongLong:self.ExpectedResultSize] forKey:SerializationKeyExpectedResultSize];
+	}
 
 	NSError* Error = nil;
 	NSData* Data = [NSJSONSerialization dataWithJSONObject:Dict options:NSJSONWritingSortedKeys error:&Error];
@@ -314,7 +324,7 @@ enum class EBackgroundNSURLCDNInfoResponse : uint32 // Beware these constants de
 - (NSURLSessionDownloadTask*)CreateDownloadForURL:(NSURL* _Nonnull)URL WithPriority:(float)Priority WithTaskData:(FBackgroundNSURLSessionDownloadTaskData* _Nonnull)TaskData;
 - (NSURLSessionDownloadTask*)CreateDownloadForResumeData:(NSData* _Nonnull)ResumeData WithPriority:(float)Priority WithTaskData:(FBackgroundNSURLSessionDownloadTaskData* _Nonnull)TaskData;
 
-- (NSUInteger)CreateOrFindDownloadForURLs:(NSArray<__kindof NSString*>*)URLStrings WithPriority:(float)Priority;
+- (NSUInteger)CreateOrFindDownloadForURLs:(NSArray<__kindof NSString*>*)URLStrings WithPriority:(float)Priority WithExpectedResultSize:(uint64)ExpectedResultSize;
 - (void)PauseDownload:(NSUInteger)DownloadId;
 - (void)ResumeDownload:(NSUInteger)DownloadId;
 - (void)CancelDownload:(NSUInteger)DownloadId;
@@ -838,7 +848,7 @@ static constexpr NSInteger HTTPStatusCodeErrorServer = 500;
 	return Task;
 }
 
-- (NSUInteger)CreateOrFindDownloadForURLs:(NSArray<__kindof NSString*>*)URLStrings WithPriority:(float)Priority
+- (NSUInteger)CreateOrFindDownloadForURLs:(NSArray<__kindof NSString*>*)URLStrings WithPriority:(float)Priority WithExpectedResultSize:(uint64)ExpectedResultSize
 {
 	if (_AllDownloadsFuture.wait_for(std::chrono::seconds(0)) != std::future_status::ready)
 	{
@@ -869,7 +879,7 @@ static constexpr NSInteger HTTPStatusCodeErrorServer = 500;
 	URLs = [self ReorderCDNsByReachability:URLs];
 
 	// Serialize current settings
-	FBackgroundNSURLSessionDownloadTaskData* TaskData = [FBackgroundNSURLSessionDownloadTaskData TaskDataWithURLs:URLs WithRetryCount:_RetryResumeDataLimit];
+	FBackgroundNSURLSessionDownloadTaskData* TaskData = [FBackgroundNSURLSessionDownloadTaskData TaskDataWithURLs:URLs WithRetryCount:_RetryResumeDataLimit WithExpectedResultSize:ExpectedResultSize];
 
 	// Check for existing download task, could be from previous app session.
 	UE_DNLD_LOG(@"Trying to find existing download for asset path '%@'", AssetPath);
@@ -1462,6 +1472,12 @@ static constexpr NSInteger HTTPStatusCodeErrorServer = 500;
 	NSURLSessionDownloadTask* Task = (NSURLSessionDownloadTask*)GenericTask;
 	NSString* LocalizedDescription = Error.localizedDescription;
 
+	NSNumber* CancelReason = [Error.userInfo objectForKey:NSURLErrorBackgroundTaskCancelledReasonKey];
+	if (CancelReason != nil)
+	{
+		LocalizedDescription = [NSString stringWithFormat:@"%@ (BackgroundTaskCancelledReason=%i)", LocalizedDescription, CancelReason.intValue];
+	}
+
 	const NSUInteger DownloadId = [self FindDownloadIdForTask:Task];
 	const bool bIsTrackedTask = DownloadId != InvalidDownloadId;
 
@@ -1557,6 +1573,39 @@ static constexpr NSInteger HTTPStatusCodeErrorServer = 500;
 		[[NSFileManager defaultManager] removeItemAtPath:DestinationPath error:nil];
 	}
 
+	// Check file size before attempting to move
+	bool bResultSizeIsCorrect = false;
+	FBackgroundNSURLSessionDownloadTaskData* TaskData = [FBackgroundNSURLSessionDownloadTaskData TaskDataFromSerializedString:Task.taskDescription];
+	if (TaskData != nil && TaskData.ExpectedResultSize > 0)
+	{
+		NSString* LocationPath = Location.path;
+
+		NSError* AttributesError = nil;
+		NSDictionary* Attributes = [[NSFileManager defaultManager] attributesOfItemAtPath:LocationPath error:&AttributesError];
+
+		if (Attributes != nil && AttributesError == nil)
+		{
+			const UInt64 CurrentFileSize = [Attributes fileSize];
+			bResultSizeIsCorrect = CurrentFileSize == TaskData.ExpectedResultSize;
+
+			// Fail request if current size doesn't match expected size. 
+			if (!bResultSizeIsCorrect)
+			{
+				[self SetDownloadResult:HTTPStatusCodeErrorServer WithTempFile:nil ForDownload:Task];
+
+				UE_DNLD_LOG(@"didFinishDownloadingToURL task '%@' with taskIdentifier %lu file '%@' size %llu doesn't match expected file size of %llu", Task.taskDescription, Task.taskIdentifier, LocationPath, CurrentFileSize, TaskData.ExpectedResultSize);
+				return;
+			}
+		}
+		else
+		{
+			[self SetDownloadResult:HTTPStatusCodeErrorServer WithTempFile:nil ForDownload:Task];
+
+			UE_DNLD_LOG(@"didFinishDownloadingToURL task '%@' with taskIdentifier %lu can't access file attributes of '%@' due to '%@'", Task.taskDescription, Task.taskIdentifier, LocationPath, (AttributesError != nil ? AttributesError.localizedDescription : @"error is nil"));
+			return;
+		}
+	}
+
 	NSError* Error = nil;
 	[[NSFileManager defaultManager] moveItemAtURL:Location toURL:[NSURL fileURLWithPath:DestinationPath] error:&Error];
 
@@ -1568,13 +1617,13 @@ static constexpr NSInteger HTTPStatusCodeErrorServer = 500;
 	{
 		[self SetDownloadResult:HTTPStatusCodeErrorServer WithTempFile:nil ForDownload:Task];
 
-		UE_DNLD_LOG(@"didFinishDownloadingToURL task '%@' with taskIdentifier %lu failed to move file to '%@' due to '%@'", Task.taskDescription, Task.taskIdentifier, DestinationPath, Error.localizedDescription);
+		UE_DNLD_LOG(@"didFinishDownloadingToURL task '%@' with taskIdentifier %lu failed to move file to '%@' due to '%@', result size was correct %u", Task.taskDescription, Task.taskIdentifier, DestinationPath, Error.localizedDescription, bResultSizeIsCorrect ? 1 : 0);
 	}
 	else
 	{
 		[self SetDownloadResult:HTTPStatusCodeSuccessCreated WithTempFile:DestinationPath ForDownload:Task];
 
-		UE_DNLD_LOG(@"didFinishDownloadingToURL task '%@' with taskIdentifier %lu move file to '%@', download finished", Task.taskDescription, Task.taskIdentifier, DestinationPath);
+		UE_DNLD_LOG(@"didFinishDownloadingToURL task '%@' with taskIdentifier %lu move file to '%@', result size was correct %u, download finished", Task.taskDescription, Task.taskIdentifier, DestinationPath, bResultSizeIsCorrect ? 1 : 0);
 	}
 }
 
@@ -1614,7 +1663,7 @@ void FBackgroundURLSessionHandler::AllowCellular(bool bAllow)
 	}
 }
 
-uint64 FBackgroundURLSessionHandler::CreateOrFindDownload(const TArray<FString>& URLs, const float Priority, BackgroundHttpFileHashHelperRef HelperRef)
+uint64 FBackgroundURLSessionHandler::CreateOrFindDownload(const TArray<FString>& URLs, const float Priority, BackgroundHttpFileHashHelperRef HelperRef, const uint64 ExpectedResultSize)
 {
 	@autoreleasepool
 	{
@@ -1625,7 +1674,7 @@ uint64 FBackgroundURLSessionHandler::CreateOrFindDownload(const TArray<FString>&
 		}
 
 		[[FBackgroundNSURLSession Shared] SetFileHashHelper:HelperRef];
-		return [[FBackgroundNSURLSession Shared] CreateOrFindDownloadForURLs:URLArray WithPriority:Priority];
+		return [[FBackgroundNSURLSession Shared] CreateOrFindDownloadForURLs:URLArray WithPriority:Priority WithExpectedResultSize:ExpectedResultSize];
 	}
 }
 
diff --git a/Engine/Source/Runtime/ApplicationCore/Public/IOS/IOSBackgroundURLSessionHandler.h b/Engine/Source/Runtime/ApplicationCore/Public/IOS/IOSBackgroundURLSessionHandler.h
--- a/Engine/Source/Runtime/ApplicationCore/Public/IOS/IOSBackgroundURLSessionHandler.h
+++ b/Engine/Source/Runtime/ApplicationCore/Public/IOS/IOSBackgroundURLSessionHandler.h
@@ -79,7 +79,7 @@ public:
 	// Priority is a value between 0.0 to 1.0, see NSURLSessionTaskPriorityDefault.
 	// HelperRef is optional shared reference to BackgroundHttpFileHashHelperRef.
 	// In case of HandleEventsForBackgroundURLSession this subsystem will create it's own reference.
-	static uint64 CreateOrFindDownload(const TArray<FString>& URLs, const float Priority, BackgroundHttpFileHashHelperRef HelperRef);
+	static uint64 CreateOrFindDownload(const TArray<FString>& URLs, const float Priority, BackgroundHttpFileHashHelperRef HelperRef, const uint64 ExpectedResultSize = 0);
 
 	static void PauseDownload(const uint64 DownloadId);
 
diff --git a/Engine/Source/Runtime/Online/BackgroundHTTP/Private/BackgroundHttpRequestImpl.cpp b/Engine/Source/Runtime/Online/BackgroundHTTP/Private/BackgroundHttpRequestImpl.cpp
--- a/Engine/Source/Runtime/Online/BackgroundHTTP/Private/BackgroundHttpRequestImpl.cpp
+++ b/Engine/Source/Runtime/Online/BackgroundHTTP/Private/BackgroundHttpRequestImpl.cpp
@@ -117,6 +117,16 @@ const TArray<FString>& FBackgroundHttpRequestImpl::GetURLList() const
 	return URLList;
 }
 
+void FBackgroundHttpRequestImpl::SetExpectedResultSize(const uint64 ExpectedSize)
+{
+	ExpectedResultSize = ExpectedSize;
+}
+
+uint64 FBackgroundHttpRequestImpl::GetExpectedResultSize() const
+{
+	return ExpectedResultSize;
+}
+
 void FBackgroundHttpRequestImpl::SetCompleteNotification(FBackgroundHttpNotificationObjectPtr DownloadCompleteNotificationObjectIn)
 {
 	DownloadCompleteNotificationObject = DownloadCompleteNotificationObjectIn;
diff --git a/Engine/Source/Runtime/Online/BackgroundHTTP/Private/IOS/ApplePlatformBackgroundHttpManager.cpp b/Engine/Source/Runtime/Online/BackgroundHTTP/Private/IOS/ApplePlatformBackgroundHttpManager.cpp
--- a/Engine/Source/Runtime/Online/BackgroundHTTP/Private/IOS/ApplePlatformBackgroundHttpManager.cpp
+++ b/Engine/Source/Runtime/Online/BackgroundHTTP/Private/IOS/ApplePlatformBackgroundHttpManager.cpp
@@ -46,7 +46,7 @@ void FApplePlatformBackgroundHttpManager::AddRequest(const FBackgroundHttpReques
 		return;
 	}
 
-	const uint64 DownloadId = FBackgroundURLSessionHandler::CreateOrFindDownload(Request->GetURLList(), Request->GetNSURLSessionPriority(), GetFileHashHelper());
+	const uint64 DownloadId = FBackgroundURLSessionHandler::CreateOrFindDownload(Request->GetURLList(), Request->GetNSURLSessionPriority(), GetFileHashHelper(), Request->GetExpectedResultSize());
 	Request->SetInternalDownloadId(DownloadId);
 
 	// Should never happen in practice
diff --git a/Engine/Source/Runtime/Online/BackgroundHTTP/Public/BackgroundHttpRequestImpl.h b/Engine/Source/Runtime/Online/BackgroundHTTP/Public/BackgroundHttpRequestImpl.h
--- a/Engine/Source/Runtime/Online/BackgroundHTTP/Public/BackgroundHttpRequestImpl.h
+++ b/Engine/Source/Runtime/Online/BackgroundHTTP/Public/BackgroundHttpRequestImpl.h
@@ -32,6 +32,8 @@ public:
     BACKGROUNDHTTP_API virtual void ResumeRequest() override;
     BACKGROUNDHTTP_API virtual void SetURLAsList(const TArray<FString>& URLs, int NumRetriesToAttempt) override;
 	BACKGROUNDHTTP_API virtual const TArray<FString>& GetURLList() const override;
+	BACKGROUNDHTTP_API virtual void SetExpectedResultSize(const uint64 ExpectedSize) override;
+	BACKGROUNDHTTP_API virtual uint64 GetExpectedResultSize() const override;
 	BACKGROUNDHTTP_API virtual void SetCompleteNotification(FBackgroundHttpNotificationObjectPtr DownloadCompleteNotificationObjectIn) override;
 	BACKGROUNDHTTP_API virtual void CompleteWithExistingResponseData(FBackgroundHttpResponsePtr BackgroundResponse) override;
 	BACKGROUNDHTTP_API virtual FBackgroundHttpRequestCompleteDelegate& OnProcessRequestComplete() override;
@@ -55,6 +57,7 @@ protected:
 	FString RequestID;
 	int NumberOfTotalRetries;
 	EBackgroundHTTPPriority RequestPriority;
+	uint64 ExpectedResultSize;
 	
 	struct FDownloadMetricsInfo
 	{
diff --git a/Engine/Source/Runtime/Online/BackgroundHTTP/Public/Interfaces/IBackgroundHttpRequest.h b/Engine/Source/Runtime/Online/BackgroundHTTP/Public/Interfaces/IBackgroundHttpRequest.h
--- a/Engine/Source/Runtime/Online/BackgroundHTTP/Public/Interfaces/IBackgroundHttpRequest.h
+++ b/Engine/Source/Runtime/Online/BackgroundHTTP/Public/Interfaces/IBackgroundHttpRequest.h
@@ -114,6 +114,24 @@ public:
 	*/
 	virtual const TArray<FString>& GetURLList() const = 0;
 
+	/**
+	* Sets expected result size for the request, might be used by the implementation
+	* to check file size before invoking completion notification. 
+	* Must be called before ProcessRequest.
+	*
+	* @param ExpectedSize expected file size in bytes, 0 to disable the check. Default value is 0.
+	*
+	*/
+	virtual void SetExpectedResultSize(const uint64 ExpectedSize) {}
+
+	/**
+	* Returns expected result size for the request, set by SetExpectedResultSize.
+	*
+	* @return Expected result size in bytes.
+	*
+	*/
+	virtual uint64 GetExpectedResultSize() const {return 0;}
+
 	/**
 	* Sets an FHTTPRequestDownloadNotificationObject to be mapped to this HTTPRequest.
 	* This request will keep the particular FHttpRequestDownloadNotificationObject reference until it completes.
