PR #11974: iPad physical keyboard improvements

diff --git a/Engine/Source/Runtime/ApplicationCore/Private/IOS/IOSInputInterface.cpp b/Engine/Source/Runtime/ApplicationCore/Private/IOS/IOSInputInterface.cpp
--- a/Engine/Source/Runtime/ApplicationCore/Private/IOS/IOSInputInterface.cpp
+++ b/Engine/Source/Runtime/ApplicationCore/Private/IOS/IOSInputInterface.cpp
@@ -25,7 +25,7 @@ uint32 TranslateGCKeyCodeToASCII(GCKeyCode KeyCode)
 {
     uint32 c = '?';
 
-    if ( KeyCode < GCKeyCodeSlash ) // Only implemented up to '/'
+    if ( KeyCode <= GCKeyCodeNonUSBackslash ) // Only implemented up to '§' GCKeyCode 100
     {
         if (KeyCode == GCKeyCodeKeyA) { c = 'A';}
         else if (KeyCode == GCKeyCodeKeyB) { c = 'B';}
@@ -65,26 +65,76 @@ uint32 TranslateGCKeyCodeToASCII(GCKeyCode KeyCode)
         else if (KeyCode == GCKeyCodeZero) { c = '0';}
         else if (KeyCode == GCKeyCodeReturnOrEnter) { c = 10;}
         else if (KeyCode == GCKeyCodeEscape) { c = 27;}
-        else if (KeyCode == GCKeyCodeBackslash) { c = 8;}
+        else if (KeyCode == GCKeyCodeDeleteOrBackspace) { c = 8;}
         else if (KeyCode == GCKeyCodeTab) { c = '\t';}
         else if (KeyCode == GCKeyCodeSpacebar) { c = ' ';}
         else if (KeyCode == GCKeyCodeHyphen) { c = '-';}
         else if (KeyCode == GCKeyCodeEqualSign) { c = '=';}
-        else if (KeyCode == GCKeyCodeOpenBracket) { c = '{';}
-        else if (KeyCode == GCKeyCodeCloseBracket) { c = '}';}
+        else if (KeyCode == GCKeyCodeOpenBracket) { c = '[';}
+        else if (KeyCode == GCKeyCodeCloseBracket) { c = ']';}
         else if (KeyCode == GCKeyCodeBackslash) { c = '\\';}
         else if (KeyCode == GCKeyCodeSemicolon) { c = ';';}
-        else if (KeyCode == GCKeyCodeQuote) { c = '\"';}
-        else if (KeyCode == GCKeyCodeGraveAccentAndTilde) { c = '~';}
+        else if (KeyCode == GCKeyCodeQuote) { c = '\'';}
+        else if (KeyCode == GCKeyCodeGraveAccentAndTilde) { c = '`';}
         else if (KeyCode == GCKeyCodeComma) { c = ',';}
         else if (KeyCode == GCKeyCodePeriod) { c = '.';}
         else if (KeyCode == GCKeyCodeSlash) { c = '/';}
+        else if (KeyCode == GCKeyCodeNonUSBackslash) { c = 167;} // Section (§) in FGenericPlatformInput::GetStandardPrintableKeyMap
         
-        UE_LOG(LogIOS, Log, TEXT("char: %c"), TCHAR(c));
-    }
+        UE_LOG(LogIOS, Log, TEXT("char: %c - GCKeycode: %d"), TCHAR(c), KeyCode);
+	}
+
     return c;
 }
 
+uint32 TranslateNonCharKeys(GCKeyCode KeyCode)
+{
+	uint32 c = '?';
+	if (KeyCode == GCKeyCodeReturnOrEnter) { c = KEYCODE_ENTER;}
+	else if (KeyCode == GCKeyCodeEscape) { c = KEYCODE_ESCAPE;}
+	else if (KeyCode == GCKeyCodeDeleteOrBackspace) { c = KEYCODE_BACKSPACE;}
+	else if (KeyCode == GCKeyCodeTab) { c = KEYCODE_TAB;}
+	else if (KeyCode == GCKeyCodeLeftArrow) { c = KEYCODE_LEFT;}
+	else if (KeyCode == GCKeyCodeRightArrow) { c = KEYCODE_RIGHT;}
+	else if (KeyCode == GCKeyCodeDownArrow) { c = KEYCODE_DOWN;}
+	else if (KeyCode == GCKeyCodeUpArrow) { c = KEYCODE_UP;}
+	else if (KeyCode == GCKeyCodeLeftControl) { c = KEYCODE_LEFT_CONTROL;}
+	else if (KeyCode == GCKeyCodeLeftShift) { c = KEYCODE_LEFT_SHIFT;}
+	else if (KeyCode == GCKeyCodeLeftAlt) { c = KEYCODE_LEFT_ALT;}
+	else if (KeyCode == GCKeyCodeCapsLock) {c = KEYCODE_CAPS_LOCK;}
+	else if (KeyCode == GCKeyCodeLeftGUI) { c = KEYCODE_LEFT_COMMAND;}
+	else if (KeyCode == GCKeyCodeRightControl) { c = KEYCODE_RIGHT_CONTROL;}
+	else if (KeyCode == GCKeyCodeRightShift) { c = KEYCODE_RIGHT_SHIFT;}
+	else if (KeyCode == GCKeyCodeRightAlt) { c = KEYCODE_RIGHT_ALT;}
+	else if (KeyCode == GCKeyCodeRightGUI) { c = KEYCODE_RIGHT_COMMAND;}
+	else if (KeyCode == GCKeyCodeF1) { c = KEYCODE_F1;}
+	else if (KeyCode == GCKeyCodeF2) { c = KEYCODE_F2;}
+	else if (KeyCode == GCKeyCodeF3) { c = KEYCODE_F3;}
+	else if (KeyCode == GCKeyCodeF4) { c = KEYCODE_F4;}
+	else if (KeyCode == GCKeyCodeF5) { c = KEYCODE_F5;}
+	else if (KeyCode == GCKeyCodeF6) { c = KEYCODE_F6;}
+	else if (KeyCode == GCKeyCodeF7) { c = KEYCODE_F7;}
+	else if (KeyCode == GCKeyCodeF8) { c = KEYCODE_F8;}
+	else if (KeyCode == GCKeyCodeF9) { c = KEYCODE_F9;}
+	else if (KeyCode == GCKeyCodeF10) { c = KEYCODE_F10;}
+	else if (KeyCode == GCKeyCodeF11) { c = KEYCODE_F11;}
+	else if (KeyCode == GCKeyCodeF12) { c = KEYCODE_F12;}
+	else if (KeyCode == GCKeyCodeF13) { c = KEYCODE_F13;}
+	else if (KeyCode == GCKeyCodeF14) { c = KEYCODE_F14;}
+	else if (KeyCode == GCKeyCodeF15) { c = KEYCODE_F15;}
+	else if (KeyCode == GCKeyCodeF16) { c = KEYCODE_F16;}
+	else if (KeyCode == GCKeyCodeF17) { c = KEYCODE_F17;}
+	else if (KeyCode == GCKeyCodeF18) { c = KEYCODE_F18;}
+	else if (KeyCode == GCKeyCodeF19) { c = KEYCODE_F19;}
+	else if (KeyCode == GCKeyCodeF20) { c = KEYCODE_F20;}
+	else { c = TranslateGCKeyCodeToASCII(KeyCode); }
+
+	UE_LOG(LogIOS, Log, TEXT("KeyCode: %d - GCKeycode: %d"), c, KeyCode);
+
+	return c;
+}
+
+
 constexpr EIOSEventType operator+(EIOSEventType type, int Index) { return (EIOSEventType)((int)type + Index); }
 // protects the input stack used on 2 threads
 static FCriticalSection CriticalSection;
@@ -266,7 +316,9 @@ void FIOSInputInterface::HandleKeyboardConnection(GCKeyboard* Keyboard)
         {
             FDeferredIOSEvent DeferredEvent;
             DeferredEvent.type = pressed ? EIOSEventType::KeyDown : EIOSEventType::KeyUp;
-            DeferredEvent.keycode = TranslateGCKeyCodeToASCII(keyCode);
+            DeferredEvent.keycode = TranslateNonCharKeys(keyCode);
+            DeferredEvent.charcode = TranslateGCKeyCodeToASCII(keyCode);
+
             
             FScopeLock Lock(&EventsMutex);
             DeferredEvents.Add(DeferredEvent);
@@ -403,12 +455,12 @@ void FIOSInputInterface::ProcessEvent(const FDeferredIOSEvent& Event)
         {
             case EIOSEventType::KeyDown:
             {
-                MessageHandler->OnKeyDown(Event.keycode, Event.keycode, false);
+                MessageHandler->OnKeyDown(Event.keycode, Event.charcode, false);
                 break;
             }
             case EIOSEventType::KeyUp:
             {
-                MessageHandler->OnKeyUp(Event.keycode, Event.keycode, false);
+                MessageHandler->OnKeyUp(Event.keycode, Event.charcode, false);
                 break;
             }
             case EIOSEventType::LeftMouseDown:
diff --git a/Engine/Source/Runtime/ApplicationCore/Public/IOS/IOSInputInterface.h b/Engine/Source/Runtime/ApplicationCore/Public/IOS/IOSInputInterface.h
--- a/Engine/Source/Runtime/ApplicationCore/Public/IOS/IOSInputInterface.h
+++ b/Engine/Source/Runtime/ApplicationCore/Public/IOS/IOSInputInterface.h
@@ -19,7 +19,40 @@
 #define KEYCODE_ENTER 1000
 #define KEYCODE_BACKSPACE 1001
 #define KEYCODE_ESCAPE 1002
-
+#define KEYCODE_TAB 1003
+#define KEYCODE_LEFT 1004
+#define KEYCODE_RIGHT 1005
+#define KEYCODE_DOWN 1006
+#define KEYCODE_UP 1007
+#define KEYCODE_LEFT_CONTROL 1008
+#define KEYCODE_LEFT_SHIFT 1009
+#define KEYCODE_LEFT_ALT 1010
+#define KEYCODE_LEFT_COMMAND 1011
+#define KEYCODE_CAPS_LOCK 1012
+#define KEYCODE_RIGHT_CONTROL 1013
+#define KEYCODE_RIGHT_SHIFT 1014
+#define KEYCODE_RIGHT_ALT 1015
+#define KEYCODE_RIGHT_COMMAND 1016
+#define KEYCODE_F1 1017
+#define KEYCODE_F2 1018
+#define KEYCODE_F3 1019
+#define KEYCODE_F4 1020
+#define KEYCODE_F5 1021
+#define KEYCODE_F6 1022
+#define KEYCODE_F7 1023
+#define KEYCODE_F8 1024
+#define KEYCODE_F9 1025
+#define KEYCODE_F10 1026
+#define KEYCODE_F11 1027
+#define KEYCODE_F12 1028
+#define KEYCODE_F13 1029
+#define KEYCODE_F14 1030
+#define KEYCODE_F15 1031
+#define KEYCODE_F16 1032
+#define KEYCODE_F17 1033
+#define KEYCODE_F18 1034
+#define KEYCODE_F19 1035
+#define KEYCODE_F20 1036
 
 enum TouchType
 {
@@ -58,6 +91,7 @@ struct FDeferredIOSEvent
 {
     EIOSEventType type;
     uint32 keycode;
+    uint32 charcode;
 };
 
 /**
diff --git a/Engine/Source/Runtime/InputCore/Private/IOS/IOSPlatformInput.cpp b/Engine/Source/Runtime/InputCore/Private/IOS/IOSPlatformInput.cpp
--- a/Engine/Source/Runtime/InputCore/Private/IOS/IOSPlatformInput.cpp
+++ b/Engine/Source/Runtime/InputCore/Private/IOS/IOSPlatformInput.cpp
@@ -15,6 +15,35 @@ uint32 FIOSPlatformInput::GetKeyMap( uint32* KeyCodes, FString* KeyNames, uint32
 		ADDKEYMAP(KEYCODE_ENTER, TEXT("Enter"));
 		ADDKEYMAP(KEYCODE_BACKSPACE, TEXT("BackSpace"));
 		ADDKEYMAP(KEYCODE_ESCAPE, TEXT("Escape"));
+		ADDKEYMAP(KEYCODE_TAB, TEXT("Tab"));
+		
+		ADDKEYMAP(KEYCODE_LEFT, TEXT("Left"));
+		ADDKEYMAP(KEYCODE_RIGHT, TEXT("Right"));
+		ADDKEYMAP(KEYCODE_DOWN, TEXT("Down"));
+		ADDKEYMAP(KEYCODE_UP, TEXT("Up"));
+		
+		ADDKEYMAP(KEYCODE_LEFT_CONTROL, TEXT("LeftControl"));
+		ADDKEYMAP(KEYCODE_LEFT_SHIFT, TEXT("LeftShift"));
+		ADDKEYMAP(KEYCODE_LEFT_ALT, TEXT("LeftAlt"));
+		ADDKEYMAP(KEYCODE_LEFT_COMMAND, TEXT("LeftCommand"));
+		ADDKEYMAP(KEYCODE_CAPS_LOCK, TEXT("CapsLock"));
+		ADDKEYMAP(KEYCODE_RIGHT_CONTROL, TEXT("RightControl"));
+		ADDKEYMAP(KEYCODE_RIGHT_SHIFT, TEXT("RightShift"));
+		ADDKEYMAP(KEYCODE_RIGHT_ALT, TEXT("RightAlt"));
+		ADDKEYMAP(KEYCODE_RIGHT_COMMAND, TEXT("RightCommand"));
+		
+		ADDKEYMAP(KEYCODE_F1, TEXT("F1"));
+		ADDKEYMAP(KEYCODE_F2, TEXT("F2"));
+		ADDKEYMAP(KEYCODE_F3, TEXT("F3"));
+		ADDKEYMAP(KEYCODE_F4, TEXT("F4"));
+		ADDKEYMAP(KEYCODE_F5, TEXT("F5"));
+		ADDKEYMAP(KEYCODE_F6, TEXT("F6"));
+		ADDKEYMAP(KEYCODE_F7, TEXT("F7"));
+		ADDKEYMAP(KEYCODE_F8, TEXT("F8"));
+		ADDKEYMAP(KEYCODE_F9, TEXT("F9"));
+		ADDKEYMAP(KEYCODE_F10, TEXT("F10"));
+		ADDKEYMAP(KEYCODE_F11, TEXT("F11"));
+		ADDKEYMAP(KEYCODE_F12, TEXT("F12"));
 	}
 	return NumMappings;
 
