Fixed a duplicated global name

diff --git a/Engine/Source/Runtime/ApplicationCore/Public/Android/AndroidInputInterface.h b/Engine/Source/Runtime/ApplicationCore/Public/Android/AndroidInputInterface.h
--- a/Engine/Source/Runtime/ApplicationCore/Public/Android/AndroidInputInterface.h
+++ b/Engine/Source/Runtime/ApplicationCore/Public/Android/AndroidInputInterface.h
@@ -76,7 +76,7 @@ enum
 
 enum InputDeviceType
 {
-	Unknown,
+	UnknownInputDeviceType,
 	TouchScreen,
 	GameController,
 };
@@ -137,7 +137,7 @@ struct FAndroidInputDeviceInfo
 	int32 FeedbackMotorCount = 0;
 	bool IsExternal = false;
 
-	InputDeviceType DeviceType = InputDeviceType::Unknown;
+	InputDeviceType DeviceType = InputDeviceType::UnknownInputDeviceType;
 	MappingState DeviceState = MappingState::Unassigned;
 	int32 DataIndex = -1; // -1 = invalid
 };
