[ASIS] In case of ASIS usage ignore portrait mode height and width adjustment

#rnx

diff --git a/Engine/Source/Runtime/ApplicationCore/Private/Android/AndroidWindow.cpp b/Engine/Source/Runtime/ApplicationCore/Private/Android/AndroidWindow.cpp
--- a/Engine/Source/Runtime/ApplicationCore/Private/Android/AndroidWindow.cpp
+++ b/Engine/Source/Runtime/ApplicationCore/Private/Android/AndroidWindow.cpp
@@ -692,12 +692,14 @@ void FAndroidWindow::CalculateSurfaceSize(int32_t& SurfaceWidth, int32_t& Surfac
 		SurfaceHeight = (GSurfaceViewHeight > 0) ? GSurfaceViewHeight : CachedNativeWindowHeight;
 	}
 
+#if !USE_ANDROID_STANDALONE
 	// some phones gave it the other way (so, if swap if the app is landscape, but width < height)
 	if ((GAndroidIsPortrait && SurfaceWidth > SurfaceHeight) ||
 		(!GAndroidIsPortrait && SurfaceWidth < SurfaceHeight))
 	{
 		Swap(SurfaceWidth, SurfaceHeight);
 	}
+#endif
 
 	// ensure the size is divisible by a specified amount
 	// do not convert to a surface size that is larger than native resolution
