Use an inline allocator for the input device scope stack.  In all cases, the depth of the stack should be very small and no need for dynamic allocation.

diff --git a/Engine/Source/Runtime/ApplicationCore/Private/GenericPlatform/GenericApplication.cpp b/Engine/Source/Runtime/ApplicationCore/Private/GenericPlatform/GenericApplication.cpp
--- a/Engine/Source/Runtime/ApplicationCore/Private/GenericPlatform/GenericApplication.cpp
+++ b/Engine/Source/Runtime/ApplicationCore/Private/GenericPlatform/GenericApplication.cpp
@@ -48,9 +48,10 @@ const FGamepadKeyNames::Type FGamepadKeyNames::RightStickLeft("Gamepad_RightStic
 
 namespace UE::InputDeviceScope::Private
 {
-TArray<FInputDeviceScope*>& GetScopeStack()
+using TArrayInputDeviceScope = TArray<FInputDeviceScope*, TInlineAllocator<16>>; 
+TArrayInputDeviceScope& GetScopeStack()
 {
-	static thread_local TArray<FInputDeviceScope*> ScopeStack;
+	static thread_local TArrayInputDeviceScope ScopeStack;
 	return ScopeStack;
 }
 }
@@ -67,7 +68,7 @@ FInputDeviceScope::FInputDeviceScope(IInputDevice* InInputDevice, FName InInputD
 
 FInputDeviceScope::~FInputDeviceScope()
 {
-	TArray<FInputDeviceScope*>& ScopeStack = UE::InputDeviceScope::Private::GetScopeStack();
+	UE::InputDeviceScope::Private::TArrayInputDeviceScope& ScopeStack = UE::InputDeviceScope::Private::GetScopeStack();
 
 	// This should always be the top of the stack
 	ensureMsgf((ScopeStack.Num() > 0 && ScopeStack.Last() == this), TEXT("FInputDeviceScope was not destroyed in correct order!"));
@@ -76,7 +77,7 @@ FInputDeviceScope::~FInputDeviceScope()
 
 const FInputDeviceScope* FInputDeviceScope::GetCurrent()
 {
-	TArray<FInputDeviceScope*>& ScopeStack = UE::InputDeviceScope::Private::GetScopeStack();
+	UE::InputDeviceScope::Private::TArrayInputDeviceScope& ScopeStack = UE::InputDeviceScope::Private::GetScopeStack();
 	if (ScopeStack.Num() > 0)
 	{
 		return ScopeStack.Last();
