MacConsoleOutputDevice - Allow the engine to exit properly when a commandlet is run and finish without issues and without needed to use -silent or UnrealEditor-Cmd

#tests Local testing

diff --git a/Engine/Source/Runtime/ApplicationCore/Private/Mac/MacConsoleOutputDevice.cpp b/Engine/Source/Runtime/ApplicationCore/Private/Mac/MacConsoleOutputDevice.cpp
--- a/Engine/Source/Runtime/ApplicationCore/Private/Mac/MacConsoleOutputDevice.cpp
+++ b/Engine/Source/Runtime/ApplicationCore/Private/Mac/MacConsoleOutputDevice.cpp
@@ -285,7 +285,8 @@ void FMacConsoleOutputDevice::SetDefaultTextColor()
 @implementation FMacConsoleWindow
 - (void)windowWillClose:(NSNotification*)Notification
 {
-	if (!MacApplication && [[NSApp orderedWindows] count] == 1)
+	// Only force exit when we are not already exiting
+	if (!IsEngineExitRequested() && !MacApplication && [[NSApp orderedWindows] count] == 1)
 	{
 		UE_LOG(LogCore, Warning, TEXT("*** INTERRUPTED *** : Console Window Closed, Shutting Down"));
 		if (GLog)
