Fix for RHI dependancies missing when compiling without engine

diff --git a/Engine/Source/Runtime/ApplicationCore/Private/IOS/IOSAppDelegate.cpp b/Engine/Source/Runtime/ApplicationCore/Private/IOS/IOSAppDelegate.cpp
--- a/Engine/Source/Runtime/ApplicationCore/Private/IOS/IOSAppDelegate.cpp
+++ b/Engine/Source/Runtime/ApplicationCore/Private/IOS/IOSAppDelegate.cpp
@@ -37,7 +37,9 @@
 #include <AVFoundation/AVAudioSession.h>
 #import <AVFoundation/AVFoundation.h>
 #include "HAL/IConsoleManager.h"
+#if WITH_ENGINE
 #include "RHICommandList.h"
+#endif
 
 #if WITH_ACCESSIBILITY
 #include "IOS/Accessibility/IOSAccessibilityCache.h"
@@ -1533,6 +1535,7 @@ FCriticalSection RenderSuspend;
     [self ToggleAudioSession:false];
 #endif
     
+#if WITH_ENGINE
 	FGraphEventRef RTBlockTask = FFunctionGraphTask::CreateAndDispatchWhenReady([]()
 	{
 		FRHICommandListImmediate& RHICmdList = FRHICommandListImmediate::Get();
@@ -1543,6 +1546,7 @@ FCriticalSection RenderSuspend;
 	}, TStatId(), NULL, ENamedThreads::GetRenderThread());
 
 	RTBlockTask->Wait();
+#endif
 	
     RenderSuspend.TryLock();
     if (FTaskGraphInterface::IsRunning())
