Fix -newconsole not saving DebugConsole.ini anymore
#rnx

diff --git a/Engine/Source/Runtime/ApplicationCore/Private/Windows/WindowsConsoleOutputDevice2.cpp b/Engine/Source/Runtime/ApplicationCore/Private/Windows/WindowsConsoleOutputDevice2.cpp
--- a/Engine/Source/Runtime/ApplicationCore/Private/Windows/WindowsConsoleOutputDevice2.cpp
+++ b/Engine/Source/Runtime/ApplicationCore/Private/Windows/WindowsConsoleOutputDevice2.cpp
@@ -8,6 +8,7 @@
 #include "GenericPlatform/GenericApplication.h"
 #include "HAL/ConsoleManager.h"
 #include "HAL/PlatformApplicationMisc.h"
+#include "HAL/PlatformFileManager.h"
 #include "HAL/PlatformTime.h"
 #include "HAL/UnrealMemory.h"
 #include "Logging/LogMacros.h"
@@ -17,6 +18,7 @@
 #include "Misc/ConfigCacheIni.h"
 #include "Misc/CoreDelegates.h"
 #include "Misc/CString.h"
+#include "Misc/FileHelper.h"
 #include "Misc/OutputDevice.h"
 #include "Misc/OutputDeviceHelper.h"
 #include "Misc/Parse.h"
@@ -2724,6 +2726,13 @@ void FWindowsConsoleOutputDevice2::SaveToINI()
 		Selection = TEXT("GameConsoleWindows");
 
 	FConfigCacheIni Config(EConfigCacheType::DiskBacked);
+
+	// Workaround to file not saving. Now requires file to exist before calling LoadFile to save
+	if (!IFileManager::Get().FileExists(*Filename))
+	{
+		FFileHelper::SaveStringToFile(FStringView(), *Filename);
+	}
+
 	Config.LoadFile(Filename);
 
 	Config.SetInt(Selection, TEXT("ConsoleX"), ConsolePosX, Filename);
