Fix for iOS UI Thread checker callout.
[UIScene interfaceOrientation] being called off UI thread.

#rnx

diff --git a/Engine/Source/Runtime/ApplicationCore/Private/IOS/IOSApplication.cpp b/Engine/Source/Runtime/ApplicationCore/Private/IOS/IOSApplication.cpp
--- a/Engine/Source/Runtime/ApplicationCore/Private/IOS/IOSApplication.cpp
+++ b/Engine/Source/Runtime/ApplicationCore/Private/IOS/IOSApplication.cpp
@@ -228,19 +228,23 @@ void FIOSApplication::OrientationChanged(UIInterfaceOrientation orientation)
 {
 	// queue up the size as we see it now all the way to the RHI 
 	[FIOSAsyncTask CreateTaskWithBlock : ^ bool(void)
-	 {
+	{
 		__block FPlatformRect WindowRect;
 		dispatch_sync(dispatch_get_main_queue(), ^{
-			WindowRect = FIOSWindow::GetScreenRect();
+			WindowRect = FIOSWindow::GetScreenRect();	// Requires UI Thread
 		});
+
 		int32 WindowWidth = WindowRect.Right - WindowRect.Left;
 		int32 WindowHeight = WindowRect.Bottom - WindowRect.Top;
 		
 	 	FIOSApplication* App = [IOSAppDelegate GetDelegate].IOSApplication;
-	 
-		App->GetMessageHandler()->OnSizeChanged(App->Windows[0],WindowWidth,WindowHeight, false);
+		App->GetMessageHandler()->OnSizeChanged(App->Windows[0], WindowWidth, WindowHeight, false);
 		App->GetMessageHandler()->OnResizingWindow(App->Windows[0]);
-		App->CacheDisplayMetrics();
+		
+		dispatch_sync(dispatch_get_main_queue(), ^{
+			App->CacheDisplayMetrics();	// Requires UI Thread
+		});
+	 	
 		FDisplayMetrics DisplayMetrics;
 		FDisplayMetrics::RebuildDisplayMetrics(DisplayMetrics);
 		App->BroadcastDisplayMetricsChanged(DisplayMetrics);
