Associated the value of FeedbackMotorCount to the "New input device" log. It will help to verify the feedback funtion.

diff --git a/Engine/Source/Runtime/ApplicationCore/Private/Android/AndroidInputInterface.cpp b/Engine/Source/Runtime/ApplicationCore/Private/Android/AndroidInputInterface.cpp
--- a/Engine/Source/Runtime/ApplicationCore/Private/Android/AndroidInputInterface.cpp
+++ b/Engine/Source/Runtime/ApplicationCore/Private/Android/AndroidInputInterface.cpp
@@ -1272,8 +1272,10 @@ void FAndroidInputInterface::SendControllerEvents()
 						}
 					}
 
-					UE_LOG(LogAndroid, Log, TEXT("New input device recorded: ControllerId = %d, DeviceId=%d, ControllerNumber=%d, DeviceName=%s, Descriptor=%s"),
-						FindControllerId(CurrentDevice.DeviceId), CurrentDevice.DeviceId, CurrentDevice.ControllerId, *CurrentDevice.Name, *CurrentDevice.Descriptor);
+					UE_LOG(LogAndroid, Log, 
+						TEXT("New input device recorded: ControllerId = %d, DeviceId=%d, ControllerNumber=%d, DeviceName=%s, Descriptor=%s, FeedbackMotorCount=%d"),
+						FindControllerId(CurrentDevice.DeviceId), CurrentDevice.DeviceId, CurrentDevice.ControllerId, 
+						*CurrentDevice.Name, *CurrentDevice.Descriptor, CurrentDevice.FeedbackMotorCount);
 				}
 				else
 				{
