Added functionality to query activation protocol that is still waiting to be sent on iOS and Android platforms

diff --git a/Engine/Source/Runtime/ApplicationCore/Private/IOS/IOSAppDelegate.cpp b/Engine/Source/Runtime/ApplicationCore/Private/IOS/IOSAppDelegate.cpp
--- a/Engine/Source/Runtime/ApplicationCore/Private/IOS/IOSAppDelegate.cpp
+++ b/Engine/Source/Runtime/ApplicationCore/Private/IOS/IOSAppDelegate.cpp
@@ -937,6 +937,10 @@ static IOSAppDelegate* CachedDelegate = nil;
 	GOverrideThreadWaitTime = BlockTime;
 }
 
+-(NSString*)GetPendingActivationProtocol
+{
+	return self.PendingActivationProtocol;
+}
 
 
 - (NSProcessInfoThermalState)GetThermalState
@@ -1143,47 +1147,41 @@ static FAutoConsoleVariableRef CVarGEnableThermalsReport(
 }
 
 #if WITH_IOS_UNIVERSAL_LINKS
-static void HandleUniversalLink(const FString& UniversalLink)
+-(BOOL)application:(UIApplication*)application continueUserActivity:(NSUserActivity*)userActivity restorationHandler:(void(^)(NSArray<id<UIUserActivityRestoring>>* _Nullable)) restorationHandler 
 {
-	static FString ProtocolActivationUri;
 	static FTSTicker::FDelegateHandle ProtocolActivationTickerHandle;
 	
-	ProtocolActivationUri = UniversalLink;
-
-	if (!ProtocolActivationUri.IsEmpty())
+	if ([userActivity.activityType isEqualToString:NSUserActivityTypeBrowsingWeb])
 	{
-		if (!ProtocolActivationTickerHandle.IsValid())
+		self.PendingActivationProtocol = userActivity.webpageURL.absoluteString;
+
+		if ([self.PendingActivationProtocol length] != 0)
 		{
-			ProtocolActivationTickerHandle = FTSTicker::GetCoreTicker().AddTicker(FTickerDelegate::CreateLambda(
-				[](float)
-				{
-					if (FCoreDelegates::OnActivatedByProtocol.IsBound())
+			if (!ProtocolActivationTickerHandle.IsValid())
+			{
+				ProtocolActivationTickerHandle = FTSTicker::GetCoreTicker().AddTicker(FTickerDelegate::CreateLambda(
+					[self](float)
 					{
-						FCoreDelegates::OnActivatedByProtocol.Broadcast(ProtocolActivationUri, PLATFORMUSERID_NONE);
+						if (FCoreDelegates::OnActivatedByProtocol.IsBound())
+						{
+							FCoreDelegates::OnActivatedByProtocol.Broadcast(FString([[IOSAppDelegate GetDelegate] PendingActivationProtocol]), PLATFORMUSERID_NONE);
 
-						ProtocolActivationUri.Reset();
-						ProtocolActivationTickerHandle.Reset();
-						return false; // remove from ticker
-					}
+							self.PendingActivationProtocol = @"";
+							ProtocolActivationTickerHandle.Reset();
+							return false; // remove from ticker
+						}
 
-					return true; // try again next tick
-				}
-			));
+						return true; // try again next tick
+					}
+				));
+			}
 		}
-	}
-	else if (ProtocolActivationTickerHandle.IsValid())
-	{
-		FTSTicker::GetCoreTicker().RemoveTicker(ProtocolActivationTickerHandle);
-		ProtocolActivationTickerHandle.Reset();
-	}
-}
-
--(BOOL)application:(UIApplication*)application continueUserActivity:(NSUserActivity*)userActivity restorationHandler:(void(^)(NSArray<id<UIUserActivityRestoring>>* _Nullable)) restorationHandler 
-{
-	if ([userActivity.activityType isEqualToString:NSUserActivityTypeBrowsingWeb])
-	{
-		NSString* UniversalLink = userActivity.webpageURL.absoluteString;
-		HandleUniversalLink(FString(UniversalLink));
+		else if (ProtocolActivationTickerHandle.IsValid())
+		{
+			FTSTicker::GetCoreTicker().RemoveTicker(ProtocolActivationTickerHandle);
+			ProtocolActivationTickerHandle.Reset();
+		}
+		
 		return YES;
 	}
 	return NO;
diff --git a/Engine/Source/Runtime/ApplicationCore/Public/IOS/IOSAppDelegate.h b/Engine/Source/Runtime/ApplicationCore/Public/IOS/IOSAppDelegate.h
--- a/Engine/Source/Runtime/ApplicationCore/Public/IOS/IOSAppDelegate.h
+++ b/Engine/Source/Runtime/ApplicationCore/Public/IOS/IOSAppDelegate.h
@@ -216,7 +216,7 @@ APPLICATIONCORE_API
 @property (assign) int RequestedResX;
 @property (assign) int RequestedResY;
 @property (assign) bool bUpdateAvailable;
-
+@property (retain) NSString* PendingActivationProtocol;
 /**
  * @return the single app delegate object
  */
@@ -241,6 +241,7 @@ APPLICATIONCORE_API
 -(void)CheckForZoomAccessibility;
 -(float)GetBackgroundingMainThreadBlockTime;
 -(void)OverrideBackgroundingMainThreadBlockTime:(float)BlockTime;
+-(NSString*)GetPendingActivationProtocol;
 
 #if !PLATFORM_TVOS && !PLATFORM_VISIONOS && !defined(SWIFT_IMPORT)
   +(EDeviceScreenOrientation) ConvertFromUIInterfaceOrientation:(UIInterfaceOrientation)Orientation;
diff --git a/Engine/Source/Runtime/Core/Private/Android/AndroidPlatformMisc.cpp b/Engine/Source/Runtime/Core/Private/Android/AndroidPlatformMisc.cpp
--- a/Engine/Source/Runtime/Core/Private/Android/AndroidPlatformMisc.cpp
+++ b/Engine/Source/Runtime/Core/Private/Android/AndroidPlatformMisc.cpp
@@ -164,6 +164,7 @@ extern void AndroidThunkCpp_SetSharedPreferenceString(const FString& Group, cons
 extern bool AndroidThunkCpp_GetSharedPreferenceStringTypeSafe(const FString& Group, const FString& Key, FString& OutValue);
 extern void AndroidThunkCpp_DeleteSharedPreference(const FString& Group, const FString& Key);
 extern void AndroidThunkCpp_DeleteSharedPreferenceGroup(const FString& Group);
+extern FString AndroidThunkCpp_GetIntentDataAsString();
 #endif
 
 // From AndroidFile.cpp
@@ -353,6 +354,10 @@ void FAndroidMisc::PlatformPreInit()
 	FGenericPlatformMisc::PlatformPreInit();
 	EstablishVulkanDeviceSupport();
 	FAndroidAppEntry::PlatformInit();
+#if USE_ANDROID_JNI
+	// Handle launch with intent
+	HandleNewIntentUri(AndroidThunkCpp_GetIntentDataAsString());
+#endif
 }
 
 static volatile bool HeadPhonesArePluggedIn = false;
@@ -1664,14 +1669,12 @@ bool FAndroidMisc::IsAllowedRemoteNotifications()
 #endif
 }
 
-static void HandleNewIntentUri_GameThread(const FString& IntentUri)
+static FString PendingProtocolActivationUri;
+static void HandleNewIntentUri_GameThread()
 {
-	static FString ProtocolActivationUri;
 	static FTSTicker::FDelegateHandle ProtocolActivationTickerHandle;
 
-	ProtocolActivationUri = IntentUri;
-
-	if (GAndroidBroadcastIntentData && !ProtocolActivationUri.IsEmpty())
+	if (GAndroidBroadcastIntentData && !PendingProtocolActivationUri.IsEmpty())
 	{
 		if (!ProtocolActivationTickerHandle.IsValid())
 		{
@@ -1680,9 +1683,9 @@ static void HandleNewIntentUri_GameThread(const FString& IntentUri)
 				{
 					if (FCoreDelegates::OnActivatedByProtocol.IsBound())
 					{
-						FCoreDelegates::OnActivatedByProtocol.Broadcast(ProtocolActivationUri, PLATFORMUSERID_NONE);
+						FCoreDelegates::OnActivatedByProtocol.Broadcast(PendingProtocolActivationUri, PLATFORMUSERID_NONE);
 
-						ProtocolActivationUri.Reset();
+						PendingProtocolActivationUri.Reset();
 						ProtocolActivationTickerHandle.Reset();
 						return false; // remove from ticker
 					}
@@ -1699,14 +1702,21 @@ static void HandleNewIntentUri_GameThread(const FString& IntentUri)
 	}
 }
 
+FString FAndroidMisc::GetPendingActivationProtocol()
+{
+	return PendingProtocolActivationUri;
+}
+
 void FAndroidMisc::HandleNewIntentUri(const FString& IntentUri)
 {
+	PendingProtocolActivationUri = IntentUri;
+	
 	// TaskGraph may be null if app was resumed before it's fully initialized, in such cases regular start will pick up the intent and send it again
-	if (!IntentUri.IsEmpty() && FTaskGraphInterface::IsRunning())
+	if (!PendingProtocolActivationUri.IsEmpty() && FTaskGraphInterface::IsRunning())
 	{
-		FFunctionGraphTask::CreateAndDispatchWhenReady([IntentUri]()
+		FFunctionGraphTask::CreateAndDispatchWhenReady([]()
 			{
-				HandleNewIntentUri_GameThread(IntentUri);
+				HandleNewIntentUri_GameThread();
 			}, TStatId(), NULL, ENamedThreads::GameThread);
 	}
 }
diff --git a/Engine/Source/Runtime/Core/Private/GenericPlatform/GenericPlatformMisc.cpp b/Engine/Source/Runtime/Core/Private/GenericPlatform/GenericPlatformMisc.cpp
--- a/Engine/Source/Runtime/Core/Private/GenericPlatform/GenericPlatformMisc.cpp
+++ b/Engine/Source/Runtime/Core/Private/GenericPlatform/GenericPlatformMisc.cpp
@@ -1926,6 +1926,12 @@ void FGenericPlatformMisc::UnregisterForRemoteNotifications()
 	// not implemented by default
 }
 
+FString FGenericPlatformMisc::GetPendingActivationProtocol()
+{
+	// not implemented by default
+	return FString();
+}
+
 bool FGenericPlatformMisc::RequestDeviceCheckToken(TFunction<void(const TArray<uint8>&)> QuerySucceededFunc, TFunction<void(const FString&, const FString&)> QueryFailedFunc)
 {
 	// not implemented by default
diff --git a/Engine/Source/Runtime/Core/Private/IOS/IOSPlatformMisc.cpp b/Engine/Source/Runtime/Core/Private/IOS/IOSPlatformMisc.cpp
--- a/Engine/Source/Runtime/Core/Private/IOS/IOSPlatformMisc.cpp
+++ b/Engine/Source/Runtime/Core/Private/IOS/IOSPlatformMisc.cpp
@@ -1342,6 +1342,12 @@ FIOSPlatformMisc::EIOSAuthNotificationStatus FIOSPlatformMisc::GetNotificationAu
 	return CurrentAuthStatus;
 }
 
+FString FIOSPlatformMisc::GetPendingActivationProtocol()
+{
+	NSString* ActivationProtocol = [[IOSAppDelegate GetDelegate] GetPendingActivationProtocol];
+	return FString(ActivationProtocol);
+}
+
 // See for more information about the Blobs
 // https://opensource.apple.com/source/xnu/xnu-4570.61.1/osfmk/kern/cs_blobs.h.auto.html
 
diff --git a/Engine/Source/Runtime/Core/Public/Android/AndroidPlatformMisc.h b/Engine/Source/Runtime/Core/Public/Android/AndroidPlatformMisc.h
--- a/Engine/Source/Runtime/Core/Public/Android/AndroidPlatformMisc.h
+++ b/Engine/Source/Runtime/Core/Public/Android/AndroidPlatformMisc.h
@@ -161,6 +161,7 @@ public:
 	static CORE_API bool IsAllowedRemoteNotifications();
 
 	static CORE_API void HandleNewIntentUri(const FString& IntentUri);
+	static CORE_API FString GetPendingActivationProtocol();
 
 	/** @return Memory representing a true type or open type font provided by the platform as a default font for unreal to consume; empty array if the default font failed to load. */
 	static CORE_API TArray<uint8> GetSystemFontBytes();
diff --git a/Engine/Source/Runtime/Core/Public/GenericPlatform/GenericPlatformMisc.h b/Engine/Source/Runtime/Core/Public/GenericPlatform/GenericPlatformMisc.h
--- a/Engine/Source/Runtime/Core/Public/GenericPlatform/GenericPlatformMisc.h
+++ b/Engine/Source/Runtime/Core/Public/GenericPlatform/GenericPlatformMisc.h
@@ -1730,6 +1730,11 @@ public:
 	*/
 	static CORE_API void UnregisterForRemoteNotifications();
 
+	/**
+	* Returns activation protocol that's waiting to be broadcast
+	*/
+	static CORE_API FString GetPendingActivationProtocol();
+
 	/**
 	 * Allows platform at runtime to disable unsupported plugins
 	 *  @param	PluginName	Name of enabled plugin to consider
diff --git a/Engine/Source/Runtime/Core/Public/IOS/IOSPlatformMisc.h b/Engine/Source/Runtime/Core/Public/IOS/IOSPlatformMisc.h
--- a/Engine/Source/Runtime/Core/Public/IOS/IOSPlatformMisc.h
+++ b/Engine/Source/Runtime/Core/Public/IOS/IOSPlatformMisc.h
@@ -85,6 +85,8 @@ struct CORE_API FIOSPlatformMisc : public FApplePlatformMisc
 	};
 	static EIOSAuthNotificationStatus GetNotificationAuthorizationStatus();
     
+	static FString GetPendingActivationProtocol();
+	
     static bool IsEntitlementEnabled(const char *EntitlementToCheck);
 	
 	static class IPlatformChunkInstall* GetPlatformChunkInstall();
diff --git a/Engine/Source/Runtime/Launch/Private/Android/LaunchAndroid.cpp b/Engine/Source/Runtime/Launch/Private/Android/LaunchAndroid.cpp
--- a/Engine/Source/Runtime/Launch/Private/Android/LaunchAndroid.cpp
+++ b/Engine/Source/Runtime/Launch/Private/Android/LaunchAndroid.cpp
@@ -179,7 +179,6 @@ extern bool AndroidThunkCpp_IsVirtualKeyboardShown();
 extern bool AndroidThunkCpp_IsWebViewShown();
 extern void AndroidThunkCpp_RestartApplication(const FString& IntentString);
 extern FString AndroidThunkCpp_GetIntentExtrasString(const FString& Key);
-extern FString AndroidThunkCpp_GetIntentDataAsString();
 
 // Base path for file accesses
 extern FString GFilePathBase;
@@ -733,8 +732,6 @@ int32 AndroidMain(struct android_app* state)
 
 	FAndroidStats::Init(FParse::Param(FCommandLine::Get(), TEXT("hwcpipe")));
 
-	FAndroidMisc::HandleNewIntentUri(AndroidThunkCpp_GetIntentDataAsString());
-
 	BootTimingPoint("Tick loop starting");
 	DumpBootTiming();
 
