Fixed an issue of the game controller reconnection state not being updated.
Added condition check for occasionally a controller's DeviceData not found for unknown reason.

diff --git a/Engine/Source/Runtime/ApplicationCore/Private/Android/AndroidInputInterface.cpp b/Engine/Source/Runtime/ApplicationCore/Private/Android/AndroidInputInterface.cpp
--- a/Engine/Source/Runtime/ApplicationCore/Private/Android/AndroidInputInterface.cpp
+++ b/Engine/Source/Runtime/ApplicationCore/Private/Android/AndroidInputInterface.cpp
@@ -449,11 +449,10 @@ void FAndroidInputInterface::MapControllerToPlayer(const FString& ControllerDesc
 	if (State == EInputDeviceConnectionState::Connected)
 	{
 		PlatformUserId = DeviceMapper.GetUserForInputDevice(DeviceId);
-		if (PlatformUserId != PLATFORMUSERID_NONE)
+		if (PlatformUserId == PLATFORMUSERID_NONE)
 		{
-			return;
+			PlatformUserId = DeviceMapper.GetPlatformUserForNewlyConnectedDevice();
 		}
-		PlatformUserId = DeviceMapper.GetPlatformUserForNewlyConnectedDevice();
 	}
 	else if (State == EInputDeviceConnectionState::Disconnected)
 	{
@@ -1273,8 +1272,8 @@ void FAndroidInputInterface::SendControllerEvents()
 					}
 
 					UE_LOG(LogAndroid, Log, 
-						TEXT("New input device recorded: ControllerId = %d, DeviceId=%d, ControllerNumber=%d, DeviceName=%s, Descriptor=%s, FeedbackMotorCount=%d"),
-						FindControllerId(CurrentDevice.DeviceId), CurrentDevice.DeviceId, CurrentDevice.ControllerId, 
+						TEXT("New input device recorded: ControllerId = %d, DeviceId=%d, ControllerType=%d, DeviceName=%s, Descriptor=%s, FeedbackMotorCount=%d"),
+						FindControllerId(CurrentDevice.DeviceId), CurrentDevice.DeviceId, CurrentDevice.DeviceType, 
 						*CurrentDevice.Name, *CurrentDevice.Descriptor, CurrentDevice.FeedbackMotorCount);
 				}
 				else
@@ -1370,6 +1369,14 @@ void FAndroidInputInterface::SendControllerEvents()
 			FInputDeviceScope InputScope(nullptr, *(DeviceInfo.Name), InputDeviceId.GetId(), DeviceInfo.Descriptor);
 
 			FAndroidGamepadDeviceMapping* DeviceData = GameControllerDataMap.Find(InputDeviceId);
+			if (DeviceData == nullptr)
+			{
+				// The protection here is for some very rare cases, with only 1 PS5 controller connected, the OS will report
+				// two game controllers and one of them is a false game controller without a name and descriptor. It will generate
+				// touch inputs and thus be validated as a touch screen and not allocating corresponding device data. 
+				// I only met the case when the game was running at a extremely low FPS with a debugger attached.
+				continue; 
+			}
 			FAndroidControllerData& OldControllerState = DeviceData->OldControllerData;
 			FAndroidControllerData& NewControllerState = DeviceData->NewControllerData;
 			
@@ -1697,11 +1704,15 @@ void FAndroidInputInterface::RemoveInputDevice(int32 DeviceId)
 	if (DeviceInfo->DeviceType == InputDeviceType::GameController)
 	{
 		FInputDeviceId InputDeviceId = InternalDeviceIdMappings.FindDeviceId(DeviceInfo->Descriptor);
-		check(InputDeviceId != INPUTDEVICEID_NONE);
-		// keep the game controller data for future reconnections since the FInputDeviceId won't change
-		FAndroidGamepadDeviceMapping* GameControllerData = GameControllerDataMap.Find(InputDeviceId);
-		check(GameControllerData != nullptr);
-		GameControllerData->ResetRuntimeData();
+		if (InputDeviceId != INPUTDEVICEID_NONE)
+		{
+			// keep the game controller data for future reconnections since the FInputDeviceId won't change
+			FAndroidGamepadDeviceMapping* GameControllerData = GameControllerDataMap.Find(InputDeviceId);
+			if (GameControllerData != nullptr)
+			{
+				GameControllerData->ResetRuntimeData();
+			}
+		}
 	}
 
 	InputDeviceInfoMap.Remove(DeviceId);
@@ -1770,8 +1781,8 @@ void FAndroidInputInterface::JoystickAxisEvent(int32 DeviceId, int32 AxisId, flo
 {
 	FScopeLock Lock(&TouchInputCriticalSection);
 
-	FAndroidInputDeviceInfo* DeviceInfo;
-	FAndroidGamepadDeviceMapping* DeviceData;
+	FAndroidInputDeviceInfo* DeviceInfo = nullptr;
+	FAndroidGamepadDeviceMapping* DeviceData = nullptr;
 	if (!GetInputDeviceByDeviceId(DeviceId, &DeviceInfo, &DeviceData))
 	{
 		AddPendingInputDevice(DeviceId, InputDeviceType::GameController);
@@ -1785,6 +1796,11 @@ void FAndroidInputInterface::JoystickAxisEvent(int32 DeviceId, int32 AxisId, flo
 		UE_LOG(LogAndroid, Log, TEXT("Gamepad state changed to Valid, DeviceId = %d"), DeviceId);
 	}
 
+	if (DeviceData == nullptr)
+	{
+		return; // bypass inputs from the device not identified as a game controller
+	}
+
 	auto RemapTriggerFunction = [](const float Minimum, const float Value)
 	{
 		if(Minimum != 0.0f)
@@ -1910,8 +1926,8 @@ void FAndroidInputInterface::JoystickButtonEvent(int32 deviceId, int32 buttonId,
 {
 	FScopeLock Lock(&TouchInputCriticalSection);
 
-	FAndroidInputDeviceInfo* DeviceInfo;
-	FAndroidGamepadDeviceMapping* DeviceData;
+	FAndroidInputDeviceInfo* DeviceInfo = nullptr;
+	FAndroidGamepadDeviceMapping* DeviceData = nullptr;
 	if (!GetInputDeviceByDeviceId(deviceId, &DeviceInfo, &DeviceData))
 	{
 		AddPendingInputDevice(deviceId, InputDeviceType::GameController);
@@ -1926,6 +1942,11 @@ void FAndroidInputInterface::JoystickButtonEvent(int32 deviceId, int32 buttonId,
 		DumpInputDevices();
 	}
 
+	if (DeviceData == nullptr)
+	{
+		return; // bypass inputs from the device not identified as a game controller
+	}
+
 	//FPlatformMisc::LowLevelOutputDebugStringf(TEXT("JoystickButtonEvent[%d]: %d"), (int)DeviceData->ButtonRemapping, buttonId);
 
 	if (DeviceData->ControllerClass == ControllerClassType::PlaystationWireless)
