ForceMaxTouchpadSensitivity
-Windows has a system setting which supresses touchpad inputs while the user types on the keyboard.  This setting is bad for playing a game with the touchpad and keyboard, the touchpad is often unresponsive.  The setting defaults to "Medium".  If it is set to "Most Sensitive" this touchpad supression is disabled and many games work much better with the touchpad.
-The underlying system has several modes of operation, depending on the windows version and api availability.
-For windows 11 23H2+ we probe for the new SetMaxTouchpadSensitivity api and if it is available use that.  When done this way the setting is specific to our app and we do not have to worry about focus/exit/crash.
-For windows 10 22H2+ (including all windows 11 that do not have the new API) we use the windows SetTouchpadParameters api to change the system setting.  We undo that change when our app loses focus or exits.  We write the restore setting to the registry and will re-apply that on a subsequent app run in order to hopefully preserve the user's system setting if the app crashes.  We *could* also read that registry value in a separate process to repair after a crash, but this is not implemented here.
-For versions earlier than windows 10 22H2 we do nothing.
-See WindowsApplication.cpp namespace UE::WindowsApplication::TouchpadSensitivity for the CVar that activate the system and several CVars that can disable parts of the system.

diff --git a/Engine/Source/Runtime/ApplicationCore/Private/Windows/WindowsApplication.cpp b/Engine/Source/Runtime/ApplicationCore/Private/Windows/WindowsApplication.cpp
--- a/Engine/Source/Runtime/ApplicationCore/Private/Windows/WindowsApplication.cpp
+++ b/Engine/Source/Runtime/ApplicationCore/Private/Windows/WindowsApplication.cpp
@@ -80,8 +80,41 @@ namespace UE::WindowsApplication::TouchpadSensitivity
 			" The setting defaults to 'Medium'. The 'Most sensitive' setting does not reject any touches, so that is what we want.."
 			" On Windows10 22H2 we force the global system setting to 'Most Sensitive', but do not write that to the registry persistently."
 			" We will set it back when the app loses focus and on exit, but it may not happen in the case of a crash leaving the setting altered until a reboot."
-/*			" On Window11 (24H2? CHECKVERSION NUMBER)+ a new api, ForceMaxTouchpadSensitivityForCurrentProcess, exists and we use that."*/
-		)
+			" On Window11 24H2 (excepting very early installs of 24H2) a new api SetMaxTouchpadSensitivity exists and we use that."
+			" SetMaxTouchpadSensitivity changes the setting only for this app, so we do not need to handle clean-up."
+			" When the cvar is changed a CVarSinkHandler will attempt to apply or remove max sensitivity.")
+	);
+
+	static void CVarSinkHandler();
+	static FAutoConsoleVariableSink CVarSink(FConsoleCommandDelegate::CreateStatic(&CVarSinkHandler));
+
+	static bool bDisableForceMaxTouchpadSensitivity_SetTouchpadParameters = false;
+	FAutoConsoleVariableRef	CVarDisableForceMaxTouchpadSensitivity_SetTouchpadParameters(
+		TEXT("WindowsApplication.DisableForceMaxTouchpadSensitivity_SetTouchpadParameters"),
+		bDisableForceMaxTouchpadSensitivity_SetTouchpadParameters,
+		TEXT("Disable use of the legacy SetTouchpadParameters API for ForceMaxTouchpadSensitivity.")
+	);
+
+	static bool bDisableForceMaxTouchpadSensitivity_SetMaxTouchpadSensitivity = false;
+	FAutoConsoleVariableRef	CVarDisableForceMaxTouchpadSensitivity_SetMaxTouchpadSensitivity(
+		TEXT("WindowsApplication.DisableForceMaxTouchpadSensitivity_SetMaxTouchpadSensitivity"),
+		bDisableForceMaxTouchpadSensitivity_SetMaxTouchpadSensitivity,
+		TEXT("Disable use of the new SetMaxTouchpadSensitivity API for ForceMaxTouchpadSensitivity.")
+	);
+
+	static bool bWriteSavedTouchpadSensitivityToRegistry = false;
+	FAutoConsoleVariableRef	CVarWriteSavedTouchpadSensitivityToRegistry(
+		TEXT("WindowsApplication.WriteSavedTouchpadSensitivityToRegistry"),
+		bWriteSavedTouchpadSensitivityToRegistry,
+		TEXT("If true and we are using the SetTouchpadParameters API for ForceMaxTouchpadSensitivity the saved touchpad sensitivity value will be written to the registry so that we can attempt to restore it even if the application crashes.")
+
+	);
+
+	static bool bRestoreTouchpadSensitivityFromRegistry = false;
+	FAutoConsoleVariableRef	CVarRestoreTouchpadSensitivityFromRegistry(
+		TEXT("WindowsApplication.RestoreTouchpadSensitivityFromRegistry"),
+		bRestoreTouchpadSensitivityFromRegistry,
+		TEXT("If true and we are using the SetTouchpadParameters API for ForceMaxTouchpadSensitivity and we find that the system is currently set to MOST_SENSITIVE and there is saved sensitivity setting previously written to the registry behave as if that is the setting we need to restore.")
 	);
 
 	enum class EForceMaxTouchpadSensitivityMode
@@ -89,11 +122,73 @@ namespace UE::WindowsApplication::TouchpadSensitivity
 		NotCached,
 		DoNothing,
 		SetTouchpadParameters,
-		ForceMaxTouchpadSensitivityForCurrentProcess
+		SetMaxTouchpadSensitivity
 	};
 
 	static EForceMaxTouchpadSensitivityMode CachedForceMaxTouchpadSensitivityMode = EForceMaxTouchpadSensitivityMode::NotCached;
 
+	typedef BOOL (*SetMaxTouchpadSensitivityProc)(_In_ BOOL enable);
+	static SetMaxTouchpadSensitivityProc SetMaxTouchpadSensitivityProcCached = nullptr;
+
+	bool SetMaxTouchpadSensitivity(bool Enable)
+	{
+		check(SetMaxTouchpadSensitivityProcCached);
+		if (SetMaxTouchpadSensitivityProcCached)
+		{
+			return SetMaxTouchpadSensitivityProcCached(Enable);
+		}
+		return false;
+	}
+
+	bool TryCacheSetMaxTouchpadSensitivityProc()
+	{
+		if (void* User32Dll = GetModuleHandle(L"user32.dll"))
+		{
+			// We are using the very new SetMaxTouchpadSensitivity(BOOL) api function which was added to windows 11 24H2, but after its release.  So some early installs of 24H2 will be missing the function.
+			// It will also be added to some version of 23H2.
+			// We will need to check for existance and then check for function to be sure we can use the new API.
+			// Info from Microsoft:
+			//		The function was renamed to SetMaxTouchpadSensitivity(BOOL) in API review. You can check for its existence by probing user32.dll for export #2806
+			//		If the API is present but disabled (due to gradual rollout over months), calling the API will fail with an error of ERROR_INVALID_STATE.
+			// Notes:
+			//		We are looking the function up by export ordinal because it is possible that when the name is officially exposed the signature will be different, but the ordinal function
+			//		version would be maintained.
+			//SetMaxTouchpadSensitivityProcCached = static_cast<SetMaxTouchpadSensitivityProc>(FPlatformProcess::GetDllExport(User32Dll, TEXT("SetMaxTouchpadSensitivity"))); // For reference, see comment above.
+			if (SetMaxTouchpadSensitivityProcCached == nullptr)
+			{
+				// Calling GetProcAddress directly here because MAKEINTRESOURCEW (or MAKEINTRESOURCEA) gives us what is essentially an opaque handle that can be passed into certain winapi
+				// functions.  We can't use it as an array of characters and convert it to ANSI as FPlatformProcess::GetDllExport would do.
+				SetMaxTouchpadSensitivityProcCached = static_cast<SetMaxTouchpadSensitivityProc>((void*)::GetProcAddress((HMODULE)User32Dll, MAKEINTRESOURCEA(2806)));
+			}
+			if (SetMaxTouchpadSensitivityProcCached)
+			{
+				UE_LOG(LogWindowsTouchpad, Log, TEXT("TryCacheSetMaxTouchpadSensitivityProc calling SetMaxTouchpadSensitivityProcCached."));
+				BOOL bSuccess = SetMaxTouchpadSensitivityProcCached(false);
+				if (!bSuccess)
+				{
+					DWORD LastError = GetLastError();
+					if (LastError == ERROR_INVALID_STATE)
+					{
+						UE_LOG(LogWindowsTouchpad, Log, TEXT("TryCacheSetMaxTouchpadSensitivityProc: SetMaxTouchpadSensitivity failed with Code %d ERROR_INVALID_STATE, which means it is not available in this version of windows.  We will try to use the fallback path for touchpad sensitivity, if available."), LastError);
+					}
+					else
+					{
+						UE_LOG(LogWindowsTouchpad, Warning, TEXT("TryCacheSetMaxTouchpadSensitivityProc: SetMaxTouchpadSensitivity failed with Code %d, which is unexpected.  We will try to use the fallback path for touchpad sensitivity, if available."), LastError);
+					}
+
+					SetMaxTouchpadSensitivityProcCached = nullptr;
+					return false;
+				}
+				else
+				{
+					return true;
+				}
+			}
+		}
+
+		return false;
+	}
+
 	void CacheForceMaxTouchpadSensitivityMode()
 	{
 		// Do nothing if we already cached it.
@@ -102,16 +197,19 @@ namespace UE::WindowsApplication::TouchpadSensitivity
 			return;
 		}
 
-		// Need to compile switch this one on the windows sdk that has the api call, I assume.
-#if 0
-		if (FPlatformMisc::VerifyWindowsVersion(10, 0, 26100)) //WINDOWS 11 24H2+ (or whatever version ForceMaxTouchpadSensitivityForCurrentProcess appears in, this might be correct)
+		if (!bDisableForceMaxTouchpadSensitivity_SetMaxTouchpadSensitivity && FPlatformMisc::VerifyWindowsVersion(10, 0, 22631)) //WINDOWS 11 23H2+
 		{
-			CachedForceMaxTouchpadSensitivityMode = EForceMaxTouchpadSensitivityMode::ForceMaxTouchpadSensitivityForCurrentProcess;
-			UE_LOG(LogWindowsTouchpad, Log, TEXT("CacheForceMaxTouchpadSensitivityMode ForceMaxTouchpadSensitivityForCurrentProcess"));
+			// Windows 11 23H2 and early installs of 24H2 might be missing the new api, because it was added after the release
+			// so we also have to test that the function is actually available.
+			if (TryCacheSetMaxTouchpadSensitivityProc())
+			{
+				CachedForceMaxTouchpadSensitivityMode = EForceMaxTouchpadSensitivityMode::SetMaxTouchpadSensitivity;
+				UE_LOG(LogWindowsTouchpad, Log, TEXT("CacheForceMaxTouchpadSensitivityMode SetMaxTouchpadSensitivity"));
+				return;
+			}
 		}
-		else
-#endif
-		if (FPlatformMisc::VerifyWindowsVersion(10, 0, 19045)) // WINDOWS 10 22H2+
+
+		if (!bDisableForceMaxTouchpadSensitivity_SetTouchpadParameters && FPlatformMisc::VerifyWindowsVersion(10, 0, 19045)) // WINDOWS 10 22H2+
 		{
 			CachedForceMaxTouchpadSensitivityMode = EForceMaxTouchpadSensitivityMode::SetTouchpadParameters;
 			UE_LOG(LogWindowsTouchpad, Log, TEXT("CacheForceMaxTouchpadSensitivityMode SetTouchpadParameters"));
@@ -127,10 +225,17 @@ namespace UE::WindowsApplication::TouchpadSensitivity
 	{
 		check(CachedForceMaxTouchpadSensitivityMode != EForceMaxTouchpadSensitivityMode::NotCached);
 
-		UE_LOG(LogWindowsTouchpad, Verbose, TEXT("GetForceMaxTouchpadSensitivityMode %i"), CachedForceMaxTouchpadSensitivityMode);
+		UE_LOG(LogWindowsTouchpad, VeryVerbose, TEXT("GetForceMaxTouchpadSensitivityMode %i"), CachedForceMaxTouchpadSensitivityMode);
 		return CachedForceMaxTouchpadSensitivityMode;
 	}
 
+	bool SupportsForceMaxTouchpadSensitivity()
+	{
+		CacheForceMaxTouchpadSensitivityMode();
+		EForceMaxTouchpadSensitivityMode Mode = GetForceMaxTouchpadSensitivityMode();
+		return Mode != EForceMaxTouchpadSensitivityMode::DoNothing;
+	}
+
 	// Define constants, enums, and structs so that we can use the SetTouchpadParameters api on versions of windows
 	// where it was not yet public.
 
@@ -194,6 +299,114 @@ namespace UE::WindowsApplication::TouchpadSensitivity
 	} TOUCHPAD_PARAMETERS_V0, * PTOUCHPAD_PARAMETERS_V0;
 	static_assert(sizeof(TOUCHPAD_PARAMETERS_V0) == 0x18);
 
+	void GetRegistryStoreIdAndSection(FString& OutStoreId, FString& OutSectionName)
+	{
+		{
+			static FString ProjectName;
+			if (ProjectName.IsEmpty())
+			{
+				// First try to use the ProjectName field.
+				if (!GConfig->GetString(TEXT("/Script/EngineSettings.GeneralProjectSettings"), TEXT("ProjectName"), ProjectName, GGameIni))
+				{
+					// Fallback to FApp and then just the executable name.
+					// These are pretty likely to not work with a separate process crash recovery, but they may work if this executable ends up trying
+					// to restore the system setting after a crash or process kill.
+					ProjectName = FApp::HasProjectName() ? FApp::GetProjectName() : FPlatformProcess::ExecutableName();
+				}
+			}
+			OutSectionName = ProjectName;
+		}
+
+		{
+			static FString CompanyName;
+			if (CompanyName.IsEmpty())
+			{
+				// First try to use the CompanyDistinguishedName CN field.
+				if (GConfig->GetString(TEXT("/Script/EngineSettings.GeneralProjectSettings"), TEXT("CompanyDistinguishedName"), CompanyName, GGameIni))
+				{
+					const int32 CNIndex = CompanyName.Find(TEXT("CN="));
+					if (CNIndex != -1)
+					{
+						const int32 StartIndex = CNIndex + 3;
+						const int32 EndIndex = CompanyName.Find(TEXT(","), ESearchCase::IgnoreCase, ESearchDir::FromStart, StartIndex);
+						const int32 Count = EndIndex - StartIndex;
+						CompanyName.MidInline(StartIndex, Count);
+						CompanyName.TrimEndInline();
+					}
+				}
+
+				// Next try to use the CompanyName field.
+				if (CompanyName.IsEmpty() && !GConfig->GetString(TEXT("/Script/EngineSettings.GeneralProjectSettings"), TEXT("CompanyName"), CompanyName, GGameIni))
+				{
+					CompanyName = OutSectionName;
+				}
+
+				// Finally re-use the project name we found before.
+				if (CompanyName.IsEmpty())
+				{
+					CompanyName = OutSectionName;
+				}
+			}
+			OutStoreId = CompanyName;
+		}
+	}
+
+	void WriteSavedSensitivityToRegistry(TOUCHPAD_SENSITIVITY_LEVEL SavedSensitivity)
+	{
+		if (bWriteSavedTouchpadSensitivityToRegistry)
+		{
+			FString Value = FString::FromInt(SavedSensitivity);
+			FString StoreId;
+			FString SectionName;
+			GetRegistryStoreIdAndSection(StoreId, SectionName);
+			FPlatformMisc::SetStoredValue(
+				StoreId,
+				SectionName,
+				TEXT("SavedTouchpadSensitivitySetting"),
+				Value
+			);
+		}
+	}
+
+	void ReadSavedSensitivityFromRegistry(TOUCHPAD_SENSITIVITY_LEVEL& OutSavedSensitivity, bool& OutbSuccess)
+	{
+		if (bWriteSavedTouchpadSensitivityToRegistry)
+		{
+			FString OutValue;
+			FString StoreId;
+			FString SectionName;
+			GetRegistryStoreIdAndSection(StoreId, SectionName);
+			OutbSuccess = FPlatformMisc::GetStoredValue(
+				StoreId,
+				SectionName,
+				TEXT("SavedTouchpadSensitivitySetting"),
+				OutValue
+			);
+			int32 ReadValue = FCString::Atoi(*OutValue);
+			if (ReadValue < TOUCHPAD_SENSITIVITY_LEVEL_MOST_SENSITIVE || ReadValue > TOUCHPAD_SENSITIVITY_LEVEL_LEAST_SENSITIVE)
+			{
+				OutbSuccess = false;
+				return;
+			}
+			OutSavedSensitivity = static_cast<TOUCHPAD_SENSITIVITY_LEVEL>(ReadValue);
+		}
+	}
+
+	void ClearSavedSensitivityFromRegistry()
+	{
+		if (bWriteSavedTouchpadSensitivityToRegistry)
+		{
+			FString StoreId;
+			FString SectionName;
+			GetRegistryStoreIdAndSection(StoreId, SectionName);
+			FPlatformMisc::DeleteStoredValue(
+				StoreId,
+				SectionName,
+				TEXT("SavedTouchpadSensitivitySetting")
+			);
+		}
+	}
+
 	static TOptional<TOUCHPAD_SENSITIVITY_LEVEL> ForceMaxTouchpadSensitivity_SavedSensitivity;
 
 	FWindowsApplication::EForceMaxTouchpadSensitivityRestorationValues GetForceMaxTouchpadSensitivityRestorationValue()
@@ -283,7 +496,7 @@ namespace UE::WindowsApplication::TouchpadSensitivity
 				UE_LOG(LogWindowsTouchpad, Verbose, TEXT("ApplyForceMaxTouchpadSensitivity sensitivity was already changed, doing nothing."));
 				return;
 			}
-			
+
 			// I think its actually impossible to need to wait here.  We would early out just above.  If there was an interveening remove it would have done the wait.
 			WaitOnPreviousApplyAsync();
 
@@ -293,24 +506,42 @@ namespace UE::WindowsApplication::TouchpadSensitivity
 			// Get the user's touchpad settings.
 			TOUCHPAD_PARAMETERS_V0 params = {};
 			params.versionNumber = 0;
-			BOOL bSuccess = ::SystemParametersInfo(SPI_GETTOUCHPADPARAMETERS, sizeof(params), &params, 0 /*fWinIni*/);
-			if (!bSuccess)
 			{
-				UE_LOG(LogWindowsTouchpad, Warning, TEXT("ApplyForceMaxTouchpadSensitivity SystemParametersInfo SPI_GETTOUCHPADPARAMETERS failed"));
-				return;
+				BOOL bSuccess = ::SystemParametersInfo(SPI_GETTOUCHPADPARAMETERS, sizeof(params), &params, 0 /*fWinIni*/);
+				if (!bSuccess)
+				{
+					UE_LOG(LogWindowsTouchpad, Warning, TEXT("ApplyForceMaxTouchpadSensitivity SystemParametersInfo SPI_GETTOUCHPADPARAMETERS failed"));
+					return;
+				}
 			}
 
 			if (params.sensitivityLevel == TOUCHPAD_SENSITIVITY_LEVEL_MOST_SENSITIVE)
 			{
+				if (bRestoreTouchpadSensitivityFromRegistry)
+				{
+					// If we exited without restoring (crashed, killed in debugger, etc) the saved sensitivity the system's current sensitivity will be MOST, but we will
+					// have saved some other value into the registry.
+					// In that case we want to restore to that registry saved setting.
+					// We do not need to set our sensitivity to MOST, because it already is MOST.
+					TOUCHPAD_SENSITIVITY_LEVEL RegistrySensitivity = TOUCHPAD_SENSITIVITY_LEVEL_MOST_SENSITIVE;
+					bool bSuccess = false;
+					ReadSavedSensitivityFromRegistry(RegistrySensitivity, bSuccess);
+					if (bSuccess && RegistrySensitivity != TOUCHPAD_SENSITIVITY_LEVEL_MOST_SENSITIVE)
+					{
+						UE_LOG(LogWindowsTouchpad, Warning, TEXT("ApplyForceMaxTouchpadSensitivity restoring sensitivity from registry. Value: %i"), (int32)RegistrySensitivity);
+						ForceMaxTouchpadSensitivity_SavedSensitivity = RegistrySensitivity;
+					}
+					return;
+				}
+
 				// If the setting is already MOST_SENSTIVE we will do nothing.
 				// This will also stop us overwriting the SavedSensitivity if we call this function twice.
 				UE_LOG(LogWindowsTouchpad, Verbose, TEXT("ApplyForceMaxTouchpadSensitivity sensitivity was already Most Sensitive"));
 				return;
 			}
-			else
-			{
-				ForceMaxTouchpadSensitivity_SavedSensitivity = params.sensitivityLevel;
-			}
+
+			ForceMaxTouchpadSensitivity_SavedSensitivity = params.sensitivityLevel;
+			WriteSavedSensitivityToRegistry(ForceMaxTouchpadSensitivity_SavedSensitivity.GetValue());
 
 			// Update the touchpad sensitivity level to "most sensitive" so that simultaneous
 			// keyboard input does not interfere with the touchpad's mouse input generation.
@@ -319,6 +550,8 @@ namespace UE::WindowsApplication::TouchpadSensitivity
 
 			if (AsyncBehavior == FWindowsApplication::EForceMaxTouchpadSensitivityAsyncBehavior::AllowAsynchronous)
 			{
+				// The SetTouchpadParameters call can take up to about 1 second, this will happen during game startup increasing aparrant load time.  Therefore we have
+				// the option of doing this asynchronously.
 				UE_LOG(LogWindowsTouchpad, Log, TEXT("ApplyForceMaxTouchpadSensitivity Task launching."));
 				ApplyForceMaxTouchpadSensitivityAsyncTask = UE::Tasks::Launch(
 					TEXT("ApplyForceMaxTouchpadSensitivityTask"),
@@ -347,7 +580,7 @@ namespace UE::WindowsApplication::TouchpadSensitivity
 				// Do not specify SPIF_UPDATEINIFILE so that the registry is not updated.
 				// We will revert this update later, and not writing to the registry is a
 				// failsafe in case we fail to do so (e.g. due to a crash).
-				bSuccess = ::SystemParametersInfo(SPI_SETTOUCHPADPARAMETERS, sizeof(params), &params, SPIF_SENDCHANGE);
+				bool bSuccess = ::SystemParametersInfo(SPI_SETTOUCHPADPARAMETERS, sizeof(params), &params, SPIF_SENDCHANGE);
 				if (!bSuccess)
 				{
 					UE_LOG(LogWindowsTouchpad, Warning, TEXT("ApplyForceMaxTouchpadSensitivity SystemParametersInfo SPI_SETTOUCHPADPARAMETERS Failed"));
@@ -357,21 +590,13 @@ namespace UE::WindowsApplication::TouchpadSensitivity
 			}
 		}
 		break;
-		case EForceMaxTouchpadSensitivityMode::ForceMaxTouchpadSensitivityForCurrentProcess:
-		{
-			// Use ForceMaxTouchpadSensitivityForCurrentProcess API
-			// Need to compile switch this one on the windows sdk that has the api call, I assume.
-#if 0
-			//BOOL bSuccess = ForceMaxTouchpadSensitivityForCurrentProcess(TRUE);
-			//if (!bSuccess)
-			//{
-			//	UE_LOG(LogWindowsTouchpad, Warning, TEXT("ApplyForceMaxTouchpadSensitivity ForceMaxTouchpadSensitivityForCurrentProcess Failed"));
-			//	return;
-			//}
-#else
-			UE_LOG(LogWindowsTouchpad, Log, TEXT("ApplyForceMaxTouchpadSensitivity ForceMaxTouchpadSensitivityForCurrentProcess -- Not Yet Implemented!"));
-			check(false);
-#endif
+		case EForceMaxTouchpadSensitivityMode::SetMaxTouchpadSensitivity:
+		{
+			if (!SetMaxTouchpadSensitivity(true))
+			{
+				UE_LOG(LogWindowsTouchpad, Warning, TEXT("ApplyForceMaxTouchpadSensitivity SetMaxTouchpadSensitivity Failed"));
+				return;
+			}
 		}
 		break;
 		default:
@@ -419,27 +644,19 @@ namespace UE::WindowsApplication::TouchpadSensitivity
 			UE_LOG(LogWindowsTouchpad, Log, TEXT("RemoveForceMaxTouchpadSensitivity sucessfully reverted to SavedSensitivity %i"), ForceMaxTouchpadSensitivity_SavedSensitivity.GetValue());
 
 			ForceMaxTouchpadSensitivity_SavedSensitivity.Reset();
+			ClearSavedSensitivityFromRegistry();
 		}
-		else if (GetForceMaxTouchpadSensitivityMode() == EForceMaxTouchpadSensitivityMode::ForceMaxTouchpadSensitivityForCurrentProcess)
+		else if (GetForceMaxTouchpadSensitivityMode() == EForceMaxTouchpadSensitivityMode::SetMaxTouchpadSensitivity)
 		{
-			// Use ForceMaxTouchpadSensitivityForCurrentProcess API
-			// Need to compile switch this one on the windows sdk that has the api call, I assume.
-#if 0
-
-			//BOOL bSuccess = ForceMaxTouchpadSensitivityForCurrentProcess(FALSE);
-			//if (!bSuccess)
-			//{
-			//	UE_LOG(LogWindowsTouchpad, Warning, TEXT("RemoveForceMaxTouchpadSensitivity ForceMaxTouchpadSensitivityForCurrentProcess Failed"));
-			//	return;
-			//}
-#else
-			UE_LOG(LogWindowsTouchpad, Log, TEXT("RemoveForceMaxTouchpadSensitivity ForceMaxTouchpadSensitivityForCurrentProcess -- Not Yet Implemented!"));
-			check(false);
-#endif	
+			if (!!SetMaxTouchpadSensitivity(false))
+			{
+				UE_LOG(LogWindowsTouchpad, Warning, TEXT("RemoveForceMaxTouchpadSensitivity SetMaxTouchpadSensitivity Failed"));
+				return;
+			}
 		}
 		else
 		{
-			UE_LOG(LogWindowsTouchpad, Verbose, TEXT("RemoveForceMaxTouchpadSensitivity Doing nothing because the sensitivity was already maximized"));
+			UE_LOG(LogWindowsTouchpad, Verbose, TEXT("RemoveForceMaxTouchpadSensitivity Doing nothing because the sensitivity is not currently forced."));
 		}
 	}
 
@@ -456,57 +673,60 @@ namespace UE::WindowsApplication::TouchpadSensitivity
 	// To almost perfectly handle error cases we need an external process watching our game and restoring the setting on any abnormal termination.
 	void RemoveForceMaxTouchpadSensitivityOnErrorShutdown()
 	{
-		// Make a local copy of the saved sensitivity to work with.  Now we know it won't change out from under us.  We only have to worry about it being half-set.
-		TOptional<TOUCHPAD_SENSITIVITY_LEVEL> LocalSavedSensitivity = ForceMaxTouchpadSensitivity_SavedSensitivity;
-
-		// Only if we did set the sensitivity using the EForceMaxTouchpadSensitivityMode::TouchpadParameters implementation
-		// will we have cached a value, and need to restore it.
-		if (LocalSavedSensitivity.IsSet())
+		if (bForceMaxTouchpadSensitivity && GetForceMaxTouchpadSensitivityMode() == EForceMaxTouchpadSensitivityMode::SetTouchpadParameters)
 		{
-			// The value should be in the valid range, and not MOST_SENSITIVE, because in that case we would not have needed to change it.
-			// If it is MOST_SENSITIVE or too large the game thread might have been setting or resetting it.
-			if (LocalSavedSensitivity.GetValue() <= TOUCHPAD_SENSITIVITY_LEVEL_MOST_SENSITIVE || LocalSavedSensitivity.GetValue() > TOUCHPAD_SENSITIVITY_LEVEL_LEAST_SENSITIVE)
+			// Make a local copy of the saved sensitivity to work with.  Now we know it won't change out from under us.  We only have to worry about it being half-set.
+			TOptional<TOUCHPAD_SENSITIVITY_LEVEL> LocalSavedSensitivity = ForceMaxTouchpadSensitivity_SavedSensitivity;
+
+			// Only if we did set the sensitivity using the EForceMaxTouchpadSensitivityMode::TouchpadParameters implementation
+			// will we have cached a value, and need to restore it.
+			if (LocalSavedSensitivity.IsSet())
 			{
-				// No restore needed, or invalid value.
-				return;
-			}
+				// The value should be in the valid range, and not MOST_SENSITIVE, because in that case we would not have needed to change it.
+				// If it is MOST_SENSITIVE or too large the game thread might have been setting or resetting it.
+				if (LocalSavedSensitivity.GetValue() <= TOUCHPAD_SENSITIVITY_LEVEL_MOST_SENSITIVE || LocalSavedSensitivity.GetValue() > TOUCHPAD_SENSITIVITY_LEVEL_LEAST_SENSITIVE)
+				{
+					// No restore needed, or invalid value.
+					return;
+				}
 
-			// We are not doing this wait in this error version.
-			//WaitOnPreviousApplyAsync();
+				// We are not doing this wait in this error version.
+				//WaitOnPreviousApplyAsync();
 
-			// Note: we will get the params, change just the one, then write them back so that the reserved param block is preserved.
-			UE_LOG(LogWindowsTouchpad, Verbose, TEXT("RemoveForceMaxTouchpadSensitivityOnErrorShutdown reverting to SavedSensitivity %i"), LocalSavedSensitivity.GetValue());
+				// Note: we will get the params, change just the one, then write them back so that the reserved param block is preserved.
+				UE_LOG(LogWindowsTouchpad, Verbose, TEXT("RemoveForceMaxTouchpadSensitivityOnErrorShutdown reverting to SavedSensitivity %i"), LocalSavedSensitivity.GetValue());
 
-			// Get the user's touchpad settings.
-			TOUCHPAD_PARAMETERS_V0 params = {};
-			params.versionNumber = 0;
-			BOOL bSuccess = ::SystemParametersInfo(SPI_GETTOUCHPADPARAMETERS, sizeof(params), &params, 0 /*fWinIni*/);
-			if (!bSuccess)
-			{
-				UE_LOG(LogWindowsTouchpad, Warning, TEXT("RemoveForceMaxTouchpadSensitivityOnErrorShutdown SystemParametersInfo SPI_GETTOUCHPADPARAMETERS failed"));
-				return;
-			}
+				// Get the user's touchpad settings.
+				TOUCHPAD_PARAMETERS_V0 params = {};
+				params.versionNumber = 0;
+				BOOL bSuccess = ::SystemParametersInfo(SPI_GETTOUCHPADPARAMETERS, sizeof(params), &params, 0 /*fWinIni*/);
+				if (!bSuccess)
+				{
+					UE_LOG(LogWindowsTouchpad, Warning, TEXT("RemoveForceMaxTouchpadSensitivityOnErrorShutdown SystemParametersInfo SPI_GETTOUCHPADPARAMETERS failed"));
+					return;
+				}
 
-			// Restore the saved sensitivity level.
-			// Leave all other bytes in the struct untouched.
-			params.sensitivityLevel = LocalSavedSensitivity.GetValue();
+				// Restore the saved sensitivity level.
+				// Leave all other bytes in the struct untouched.
+				params.sensitivityLevel = LocalSavedSensitivity.GetValue();
 
-			// Do not specify SPIF_UPDATEINIFILE so that the registry is not updated.
-			// We will revert this update later, and not writing to the registry is a
-			// failsafe in case we fail to do so (e.g. due to a crash).
-			bSuccess = ::SystemParametersInfo(SPI_SETTOUCHPADPARAMETERS, sizeof(params), &params, SPIF_SENDCHANGE);
-			if (!bSuccess)
-			{
-				UE_LOG(LogWindowsTouchpad, Warning, TEXT("RemoveForceMaxTouchpadSensitivityOnErrorShutdown SystemParametersInfo SPI_SETTOUCHPADPARAMETERS failed"));
-				return;
+				// Do not specify SPIF_UPDATEINIFILE so that the registry is not updated.
+				// We will revert this update later, and not writing to the registry is a
+				// failsafe in case we fail to do so (e.g. due to a crash).
+				bSuccess = ::SystemParametersInfo(SPI_SETTOUCHPADPARAMETERS, sizeof(params), &params, SPIF_SENDCHANGE);
+				if (!bSuccess)
+				{
+					UE_LOG(LogWindowsTouchpad, Warning, TEXT("RemoveForceMaxTouchpadSensitivityOnErrorShutdown SystemParametersInfo SPI_SETTOUCHPADPARAMETERS failed"));
+					return;
+				}
+				UE_LOG(LogWindowsTouchpad, Log, TEXT("RemoveForceMaxTouchpadSensitivityOnErrorShutdown sucessfully reverted to SavedSensitivity %i"), LocalSavedSensitivity.GetValue());
 			}
-			UE_LOG(LogWindowsTouchpad, Log, TEXT("RemoveForceMaxTouchpadSensitivityOnErrorShutdown sucessfully reverted to SavedSensitivity %i"), LocalSavedSensitivity.GetValue());
 		}
 	}
 
 	void HandleActivateApp( bool bActivate)
 	{
-		if (bForceMaxTouchpadSensitivity)
+		if (bForceMaxTouchpadSensitivity && (GetForceMaxTouchpadSensitivityMode() == EForceMaxTouchpadSensitivityMode::SetTouchpadParameters))
 		{
 			UE_LOG(LogWindowsTouchpad, Verbose, TEXT("HandleActivateApp %i"), bActivate);
 			if (bActivate)
@@ -522,24 +742,56 @@ namespace UE::WindowsApplication::TouchpadSensitivity
 
 	void HandleStartupApp()
 	{
+		UE_LOG(LogWindowsTouchpad, Verbose, TEXT("HandleStartupApp"));
 		CacheForceMaxTouchpadSensitivityMode();
 		if (bForceMaxTouchpadSensitivity)
 		{
-			UE_LOG(LogWindowsTouchpad, Verbose, TEXT("HandleStartupApp"));
 			ApplyForceMaxTouchpadSensitivity(FWindowsApplication::EForceMaxTouchpadSensitivityAsyncBehavior::AllowAsynchronous);
 		}
 	}
 
 	void HandleDestroyApplication()
 	{
-		UE_LOG(LogWindowsTouchpad, Verbose, TEXT("HandleDestroyApplication"));
-		RemoveForceMaxTouchpadSensitivity();
+		if (bForceMaxTouchpadSensitivity && (GetForceMaxTouchpadSensitivityMode() == EForceMaxTouchpadSensitivityMode::SetTouchpadParameters))
+		{
+			UE_LOG(LogWindowsTouchpad, Verbose, TEXT("HandleDestroyApplication"));
+			RemoveForceMaxTouchpadSensitivity();
+		}
 	}
 
 	void HandleShutDownAfterError()
 	{
-		UE_LOG(LogWindowsTouchpad, Verbose, TEXT("HandleShutDownAfterError"));
-		RemoveForceMaxTouchpadSensitivityOnErrorShutdown();
+		if (bForceMaxTouchpadSensitivity && (GetForceMaxTouchpadSensitivityMode() == EForceMaxTouchpadSensitivityMode::SetTouchpadParameters))
+		{
+			UE_LOG(LogWindowsTouchpad, Verbose, TEXT("HandleShutDownAfterError"));
+			RemoveForceMaxTouchpadSensitivityOnErrorShutdown();
+		}
+	}
+
+	void CVarSinkHandler()
+	{
+		check(IsInGameThread());
+
+		static bool CachedbForceMaxTouchpadSensitivity = false;
+		bool ForceMaxTouchpadSensitivity = bForceMaxTouchpadSensitivity;
+
+		if (bForceMaxTouchpadSensitivity != CachedbForceMaxTouchpadSensitivity)
+		{
+			CachedbForceMaxTouchpadSensitivity = bForceMaxTouchpadSensitivity;
+
+			// If the cvar value has been changed we need to put it into effect.
+			if (bForceMaxTouchpadSensitivity)
+			{
+				UE_LOG(LogWindowsTouchpad, Verbose, TEXT("WindowsApplication.ForceMaxTouchpadSensitivity CVarSinkHandler calling ApplyForceMaxTouchpadSensitivity"));
+				CacheForceMaxTouchpadSensitivityMode();
+				ApplyForceMaxTouchpadSensitivity(FWindowsApplication::EForceMaxTouchpadSensitivityAsyncBehavior::AllowAsynchronous);
+			}
+			else
+			{
+				UE_LOG(LogWindowsTouchpad, Verbose, TEXT("WindowsApplication.ForceMaxTouchpadSensitivity CVarSinkHandler calling RemoveForceMaxTouchpadSensitivity"));
+				RemoveForceMaxTouchpadSensitivity();
+			}
+		}
 	}
 }
 
@@ -730,6 +982,11 @@ FWindowsApplication::FWindowsApplication( const HINSTANCE HInstance, const HICON
 	UE::WindowsApplication::TouchpadSensitivity::HandleStartupApp();
 }
 
+bool FWindowsApplication::SupportsForceMaxTouchpadSensitivity()
+{
+	return UE::WindowsApplication::TouchpadSensitivity::SupportsForceMaxTouchpadSensitivity();
+}
+
 void FWindowsApplication::ApplyForceMaxTouchpadSensitivity(FWindowsApplication::EForceMaxTouchpadSensitivityAsyncBehavior AsyncBehavior)
 {
 	UE::WindowsApplication::TouchpadSensitivity::ApplyForceMaxTouchpadSensitivity(AsyncBehavior);
diff --git a/Engine/Source/Runtime/ApplicationCore/Public/Windows/WindowsApplication.h b/Engine/Source/Runtime/ApplicationCore/Public/Windows/WindowsApplication.h
--- a/Engine/Source/Runtime/ApplicationCore/Public/Windows/WindowsApplication.h
+++ b/Engine/Source/Runtime/ApplicationCore/Public/Windows/WindowsApplication.h
@@ -386,16 +386,19 @@ public:
 
 	// Touchpad Sensitivity
 	/** On Windows 10 22H2 and later where we can set the touchpad sensitivity system wide setting force it to most sensitive so that touch swipes are not ignored while keyboard keys are being depressed.
-	 * IF the setting is actually changed the system call that changes it can block for approximately 1 second.  Because that is a pretty long time, potentially delaying game startup
-	 * an argument exists on the Apply function which will allow the application of the setting to be done asynchronously.  The Remove function has no such option. Both functions will block for 
-	 * an outstanding Apply system call (With a timeout).  The end result here is that one can use the parameter to prevent the Apply function from potentially blocking for a long time, but the Remove function can then 
-	 * potentially block for two system calls.  It is expected and reccomended that the real world usage patterns make the double-block rare.
+	 * In Windows 11 24H2 a new api was added which can change this setting only for this app.
+	 * 
+	 * We reccomend that this feature be exposed as a user setting which is enabled by default and hidden/disabled if SupportsForceMaxTouchpadSensitivity() returns false and that the CVar be used to activate it rather than directly calling 
+	 * ApplyForceMaxTouchpadSensitivity/RemoveForceMaxTouchpadSensitivity due to the complexity of handling the focus/exit/crash cases with the 22H2 api.
+	 * 
+	 * See WindowsApplication.cpp namespace UE::WindowsApplication::TouchpadSensitivity for more information about this system.
 	 */
 	enum class EForceMaxTouchpadSensitivityAsyncBehavior : uint8
 	{
 		AllowAsynchronous, // This function may defer the system call that actually changes the setting in order to avoid blocking for a long time.
 		RequireBlocking // This function will not return until the system call returns.
 	};
+	APPLICATIONCORE_API bool SupportsForceMaxTouchpadSensitivity();
 	APPLICATIONCORE_API void ApplyForceMaxTouchpadSensitivity(EForceMaxTouchpadSensitivityAsyncBehavior AsyncBehavior);
 	APPLICATIONCORE_API void RemoveForceMaxTouchpadSensitivity();
 
