iOS functional parity with Android. Ensure that when when supported orientations are set at runtime, that the application autorates to a valid orientation.

diff --git a/Engine/Source/Runtime/ApplicationCore/Private/IOS/IOSAppDelegate.cpp b/Engine/Source/Runtime/ApplicationCore/Private/IOS/IOSAppDelegate.cpp
--- a/Engine/Source/Runtime/ApplicationCore/Private/IOS/IOSAppDelegate.cpp
+++ b/Engine/Source/Runtime/ApplicationCore/Private/IOS/IOSAppDelegate.cpp
@@ -1247,6 +1247,16 @@ static FAutoConsoleVariableRef CVarGEnableThermalsReport(
 		}
 	}
 	
+	const FIOSView *View = [[IOSAppDelegate GetDelegate] IOSView];
+	if (View != nil &&
+		(View->SupportedInterfaceOrientations != UIInterfaceOrientationMaskAll) &&
+		(View->SupportedInterfaceOrientations & Mask) != View->SupportedInterfaceOrientations)
+	{
+		// Return all possible allowed orientations, that includes all orientations from IOSRuntimeSettings plus
+		// the orientations set at runtime via Blueprint. The iOS View Controller will then restrict to the desired
+		// orientation(s) with the "supportedInterfaceOrientations" function.
+		return Mask | View->SupportedInterfaceOrientations;
+	}
 	// If no orientation constraints are set, default to MaskAll
 	return (Mask==0)?UIInterfaceOrientationMaskAll:Mask;
 }
diff --git a/Engine/Source/Runtime/Core/Private/IOS/IOSPlatformMisc.cpp b/Engine/Source/Runtime/Core/Private/IOS/IOSPlatformMisc.cpp
--- a/Engine/Source/Runtime/Core/Private/IOS/IOSPlatformMisc.cpp
+++ b/Engine/Source/Runtime/Core/Private/IOS/IOSPlatformMisc.cpp
@@ -374,10 +374,20 @@ void FIOSPlatformMisc::SetAllowedDeviceOrientation(EDeviceScreenOrientation NewA
 	[IOSAppDelegate GetDelegate].IOSView->SupportedInterfaceOrientations = GetUIInterfaceOrientationMask(NewAllowedDeviceOrientation);
 #endif
 	
-	dispatch_async(dispatch_get_main_queue(), ^
-	{
-		[[IOSAppDelegate GetDelegate].IOSController setNeedsUpdateOfSupportedInterfaceOrientations];
-	});
+	// Notify the view controller about a change in supported interface orientations.
+	
+	dispatch_async(dispatch_get_main_queue(),
+		^{
+			if (@available(iOS 16.0, *))
+			{
+				[[IOSAppDelegate GetDelegate].IOSController setNeedsUpdateOfSupportedInterfaceOrientations];
+			}
+			else
+			{
+				[UIViewController attemptRotationToDeviceOrientation];
+			}
+		}
+	);
 }
 
 bool FIOSPlatformMisc::HasPlatformFeature(const TCHAR* FeatureName)
