Fixed an issue where the mouse inputs won't be acknowledged in special cases like using scrcpy to simulate the input.

diff --git a/Engine/Source/Runtime/ApplicationCore/Private/Android/AndroidInputInterface.cpp b/Engine/Source/Runtime/ApplicationCore/Private/Android/AndroidInputInterface.cpp
--- a/Engine/Source/Runtime/ApplicationCore/Private/Android/AndroidInputInterface.cpp
+++ b/Engine/Source/Runtime/ApplicationCore/Private/Android/AndroidInputInterface.cpp
@@ -1299,7 +1299,28 @@ void FAndroidInputInterface::SendControllerEvents()
 	for(int i = 0; i < FAndroidInputInterface::TouchInputStack.Num(); ++i)
 	{
 		TouchInput Touch = FAndroidInputInterface::TouchInputStack[i];
-		FAndroidInputDeviceInfo* TouchDeviceInfo = InputDeviceInfoMap.Find(Touch.DeviceId);
+
+		FAndroidInputDeviceInfo* TouchDeviceInfo = nullptr;
+		if (Touch.DeviceId == -1) 
+		{
+			// some special inputs has -1 as their device id (e.g. scrcpy inputs)
+			// use default touchscreen instead to handle their inputs
+			FInputDeviceId DefInputDeviceId = IPlatformInputDeviceMapper::Get().GetDefaultInputDevice();
+			const FString* DefInputDeviceDesc = InternalDeviceIdMappings.FindDeviceKey(DefInputDeviceId);
+			for (TPair<int32, FAndroidInputDeviceInfo>& Pair : InputDeviceInfoMap)
+			{
+				if (Pair.Value.Descriptor.Equals(*DefInputDeviceDesc))
+				{
+					TouchDeviceInfo = &(Pair.Value);
+					break;
+				}
+			}
+		}
+		else
+		{
+			TouchDeviceInfo = InputDeviceInfoMap.Find(Touch.DeviceId);
+		}
+
 		if (TouchDeviceInfo == nullptr)
 		{
 			continue;
