Fixed an issue where the mouse inputs won't be acknowledged in special cases like using scrcpy to simulate the input.

diff --git a/Engine/Source/Runtime/ApplicationCore/Private/Android/AndroidInputInterface.cpp b/Engine/Source/Runtime/ApplicationCore/Private/Android/AndroidInputInterface.cpp
--- a/Engine/Source/Runtime/ApplicationCore/Private/Android/AndroidInputInterface.cpp
+++ b/Engine/Source/Runtime/ApplicationCore/Private/Android/AndroidInputInterface.cpp
@@ -1298,62 +1298,49 @@ void FAndroidInputInterface::SendControllerEvents()
 
 	for(int i = 0; i < FAndroidInputInterface::TouchInputStack.Num(); ++i)
 	{
-		TouchInput Touch = FAndroidInputInterface::TouchInputStack[i];
+		// some special inputs has -1 as their device id (e.g. scrcpy inputs)
+		// use default touchscreen instead to handle their inputs
+		FInputDeviceId TouchInputDeviceId = IPlatformInputDeviceMapper::Get().GetDefaultInputDevice();
 
-		FAndroidInputDeviceInfo* TouchDeviceInfo = nullptr;
-		if (Touch.DeviceId == -1) 
+		TouchInput Touch = FAndroidInputInterface::TouchInputStack[i];
+		if (Touch.DeviceId >= 0)
 		{
-			// some special inputs has -1 as their device id (e.g. scrcpy inputs)
-			// use default touchscreen instead to handle their inputs
-			FInputDeviceId DefInputDeviceId = IPlatformInputDeviceMapper::Get().GetDefaultInputDevice();
-			const FString* DefInputDeviceDesc = InternalDeviceIdMappings.FindDeviceKey(DefInputDeviceId);
-			for (TPair<int32, FAndroidInputDeviceInfo>& Pair : InputDeviceInfoMap)
+			FAndroidInputDeviceInfo* TouchDeviceInfo = InputDeviceInfoMap.Find(Touch.DeviceId);
+			if (TouchDeviceInfo == nullptr)
 			{
-				if (Pair.Value.Descriptor.Equals(*DefInputDeviceDesc))
-				{
-					TouchDeviceInfo = &(Pair.Value);
-					break;
-				}
+				continue;
 			}
-		}
-		else
-		{
-			TouchDeviceInfo = InputDeviceInfoMap.Find(Touch.DeviceId);
-		}
 
-		if (TouchDeviceInfo == nullptr)
-		{
-			continue;
-		}
+			if (AndroidEnableInputDeviceListener && TouchDeviceInfo->DeviceState == MappingState::ToValidate)
+			{
+				TouchDeviceInfo->DeviceState = MappingState::Valid;
 
-		if (AndroidEnableInputDeviceListener && TouchDeviceInfo->DeviceState == MappingState::ToValidate)
-		{
-			TouchDeviceInfo->DeviceState = MappingState::Valid;
+				if (TouchDeviceInfo->IsExternal == false)
+				{
+					// It's the built-in touch screen, map it to the default input device id which is already bound to the primary user id.
+					FInputDeviceId InputDeviceId = InternalDeviceIdMappings.MapDefaultInputDevice(TouchDeviceInfo->Descriptor);
+					check(InputDeviceId != INPUTDEVICEID_NONE);
+				}
+				else
+				{
+					MapControllerToPlayer(TouchDeviceInfo->Descriptor, EInputDeviceConnectionState::Connected);
+				}
 
-			if (TouchDeviceInfo->IsExternal == false)
-			{
-				// It's the built-in touch screen, map it to the default input device id which is already bound to the primary user id.
-				FInputDeviceId InputDeviceId = InternalDeviceIdMappings.MapDefaultInputDevice(TouchDeviceInfo->Descriptor);
-				check(InputDeviceId != INPUTDEVICEID_NONE);
-			}
-			else
-			{
-				MapControllerToPlayer(TouchDeviceInfo->Descriptor, EInputDeviceConnectionState::Connected);
+				UE_LOG(LogAndroid, Log, TEXT("Touch Screen state changed to Valid, DeviceId = %d"), Touch.DeviceId);
+
+				DumpInputDevices();
 			}
 
-			UE_LOG(LogAndroid, Log, TEXT("Touch Screen state changed to Valid, DeviceId = %d"), Touch.DeviceId);
+			TouchInputDeviceId = InternalDeviceIdMappings.FindDeviceId(TouchDeviceInfo->Descriptor);
 
-			DumpInputDevices();
-		}
+			if (TouchInputDeviceId == INPUTDEVICEID_NONE)
+			{
+				continue;
+			}
 
-		const FInputDeviceId TouchInputDeviceId = InternalDeviceIdMappings.FindDeviceId(TouchDeviceInfo->Descriptor);
-		if (TouchInputDeviceId == INPUTDEVICEID_NONE)
-		{
-			continue;
+			FInputDeviceScope InputScope(nullptr, *(TouchDeviceInfo->Name), TouchInputDeviceId.GetId(), TouchDeviceInfo->Descriptor);
 		}
 
-		FInputDeviceScope InputScope(nullptr, *(TouchDeviceInfo->Name), TouchInputDeviceId.GetId(), TouchDeviceInfo->Descriptor);
-
 		IPlatformInputDeviceMapper& DeviceMapper = IPlatformInputDeviceMapper::Get();
 		FPlatformUserId PlatformUserId = DeviceMapper.GetUserForInputDevice(TouchInputDeviceId);
 
