ForceMaxTouchpadSensitivity hitches on tab out.
-Made all of the legacy API SetTouchpadParameters settings changes asynchronous (except the error exit case).  Previously it was async only when we were setting the sensitivity to MostSensitive, but not when we were restoring the pre-existing setting.
-The new API SetMaxTouchpadSensitivity call is quick and is still synchronous.

diff --git a/Engine/Source/Runtime/ApplicationCore/Private/Windows/WindowsApplication.cpp b/Engine/Source/Runtime/ApplicationCore/Private/Windows/WindowsApplication.cpp
--- a/Engine/Source/Runtime/ApplicationCore/Private/Windows/WindowsApplication.cpp
+++ b/Engine/Source/Runtime/ApplicationCore/Private/Windows/WindowsApplication.cpp
@@ -21,6 +21,7 @@
 #include "Stats/Stats.h"
 #include "HAL/IConsoleManager.h"
 #include "Tasks/Task.h"
+#include "Tasks/Pipe.h"
 
 #if WITH_EDITOR
 #include "Modules/ModuleManager.h"
@@ -68,7 +69,7 @@ FAutoConsoleVariableRef	CVarPreventDuplicateMouseEventsForTouch(
 
 namespace UE::WindowsApplication::TouchpadSensitivity
 {
-	DECLARE_LOG_CATEGORY_EXTERN(LogWindowsTouchpad, Log, VeryVerbose);
+	DECLARE_LOG_CATEGORY_EXTERN(LogWindowsTouchpad, Display, VeryVerbose);
 	DEFINE_LOG_CATEGORY(LogWindowsTouchpad);
 
 	static bool bForceMaxTouchpadSensitivity = false;
@@ -169,7 +170,7 @@ namespace UE::WindowsApplication::TouchpadSensitivity
 					DWORD LastError = GetLastError();
 					if (LastError == ERROR_INVALID_STATE)
 					{
-						UE_LOG(LogWindowsTouchpad, Log, TEXT("TryCacheSetMaxTouchpadSensitivityProc: SetMaxTouchpadSensitivity failed with Code %d ERROR_INVALID_STATE, which means it is not available in this version of windows.  We will try to use the fallback path for touchpad sensitivity, if available."), LastError);
+						UE_LOG(LogWindowsTouchpad, Display, TEXT("TryCacheSetMaxTouchpadSensitivityProc: SetMaxTouchpadSensitivity failed with Code %d ERROR_INVALID_STATE, which means it is not available in this version of windows.  We will try to use the fallback path for touchpad sensitivity, if available."), LastError);
 					}
 					else
 					{
@@ -204,7 +205,7 @@ namespace UE::WindowsApplication::TouchpadSensitivity
 			if (TryCacheSetMaxTouchpadSensitivityProc())
 			{
 				CachedForceMaxTouchpadSensitivityMode = EForceMaxTouchpadSensitivityMode::SetMaxTouchpadSensitivity;
-				UE_LOG(LogWindowsTouchpad, Log, TEXT("CacheForceMaxTouchpadSensitivityMode SetMaxTouchpadSensitivity"));
+				UE_LOG(LogWindowsTouchpad, Display, TEXT("CacheForceMaxTouchpadSensitivityMode SetMaxTouchpadSensitivity"));
 				return;
 			}
 		}
@@ -212,12 +213,12 @@ namespace UE::WindowsApplication::TouchpadSensitivity
 		if (!bDisableForceMaxTouchpadSensitivity_SetTouchpadParameters && FPlatformMisc::VerifyWindowsVersion(10, 0, 19045)) // WINDOWS 10 22H2+
 		{
 			CachedForceMaxTouchpadSensitivityMode = EForceMaxTouchpadSensitivityMode::SetTouchpadParameters;
-			UE_LOG(LogWindowsTouchpad, Log, TEXT("CacheForceMaxTouchpadSensitivityMode SetTouchpadParameters"));
+			UE_LOG(LogWindowsTouchpad, Display, TEXT("CacheForceMaxTouchpadSensitivityMode SetTouchpadParameters"));
 		}
 		else
 		{
 			CachedForceMaxTouchpadSensitivityMode = EForceMaxTouchpadSensitivityMode::DoNothing;
-			UE_LOG(LogWindowsTouchpad, Log, TEXT("CacheForceMaxTouchpadSensitivityMode DoNothing"));
+			UE_LOG(LogWindowsTouchpad, Display, TEXT("CacheForceMaxTouchpadSensitivityMode DoNothing"));
 		}
 	}
 
@@ -266,6 +267,7 @@ namespace UE::WindowsApplication::TouchpadSensitivity
 		TOUCHPAD_SENSITIVITY_LEVEL_MEDIUM_SENSITIVITY = 0x00000002,
 		TOUCHPAD_SENSITIVITY_LEVEL_LOW_SENSITIVITY = 0x00000003,
 		TOUCHPAD_SENSITIVITY_LEVEL_LEAST_SENSITIVE = 0x00000004,
+		TOUCHPAD_SENSITIVITY_LEVEL_NOT_SET = 0x0000000F, // Special value used to indicate no value has been set.
 	} TOUCHPAD_SENSITIVITY_LEVEL;
 	static_assert((int32)FWindowsApplication::EForceMaxTouchpadSensitivityRestorationValues::TOUCHPAD_SENSITIVITY_LEVEL_MOST_SENSITIVE == (int32)TOUCHPAD_SENSITIVITY_LEVEL::TOUCHPAD_SENSITIVITY_LEVEL_MOST_SENSITIVE);
 	static_assert((int32)FWindowsApplication::EForceMaxTouchpadSensitivityRestorationValues::TOUCHPAD_SENSITIVITY_LEVEL_LEAST_SENSITIVE == (int32)TOUCHPAD_SENSITIVITY_LEVEL::TOUCHPAD_SENSITIVITY_LEVEL_LEAST_SENSITIVE);
@@ -407,7 +409,29 @@ namespace UE::WindowsApplication::TouchpadSensitivity
 		}
 	}
 
-	static TOptional<TOUCHPAD_SENSITIVITY_LEVEL> ForceMaxTouchpadSensitivity_SavedSensitivity;
+	// Async model for SetTouchpadParameters mode.
+	// The ::SystemParametersInfo(SPI_SETTOUCHPADPARAMETERS call can take up to about 1 second to return.  We don't want 1 second hitches on startup, tab-out, tab-in, or normal exit.  We are going to tolerate the hitch on error-out to avoid any kind of locking or waiting there.
+	// Because this call is so slow it is possible to tab in/out/in/out a number of times before even the first setting change completes.  In order to avoid having several seconds of flipping back and forth in such a situation we are going to use an eventual-consistency
+	// design.  We will ensure that the LAST state reached is the LAST state requested.  We will not necessarily do all the requested state transitions.
+	// Therefore setting changes will be requested in two stages. 
+	// 1) Set the atomic ForceMaxTouchpadSensitivity_DesiredSensitivity to the desired sensitivity (either MostSensitive, or the value we want to restore on tab-out/exit).
+	// 2) Enqueue a task to set the sensitivity to the DesiredSensitivity.
+	// The result of this is that if we do multiple requests we will change the atomic DesiredSensitivity multiple times and queue up multiple tasks, but the next task to execute will set the system to the most recently set state.
+	// As long as no additional requests are made the remaining queued tasks will complete very quickly because the current sensitivity setting is already the DesiredSensitivity.  Any time we change the sensitivity we know that at least
+	// one task is in the pipe to execute that change.
+	// 
+	// Flaws:
+	// 1) In a crash we cannot restore the setting from this process.  There is a mechanism to write the saved sensitivity to the registery and restore it from there either on the next app run, or in some external process.  However this could potentially overwrite the user's setting
+	// if they change it to most sensitive.  Still that seems relatively unlikely and the user aparrantly at least knows about the setting.
+	// 2) If, while tabbed out, the setting were to change and multiple tasks were in the pipe I think it is possible we would restore a no-longer current saved state.  This probably can't be done manually, but if one tabbed between two applications using this system I expect they could
+	// end up leaving the setting at MostSensitive.
+
+	/** Async task pipe for processing the incoming profiler service data */
+	static UE::Tasks::FPipe AsyncTaskPipe{ TEXT("ForceMaxTouchpadSensitivityPipe") };
+	static std::atomic<TOUCHPAD_SENSITIVITY_LEVEL> ForceMaxTouchpadSensitivity_DesiredSensitivity = TOUCHPAD_SENSITIVITY_LEVEL_NOT_SET;
+	static std::atomic<TOUCHPAD_SENSITIVITY_LEVEL> ForceMaxTouchpadSensitivity_SavedSensitivity = TOUCHPAD_SENSITIVITY_LEVEL_NOT_SET;
+	static bool bRegistrySavedValueConsumed = false;
+	static std::atomic<bool> bAbortAsyncTasks = false;
 
 	FWindowsApplication::EForceMaxTouchpadSensitivityRestorationValues GetForceMaxTouchpadSensitivityRestorationValue()
 	{
@@ -416,10 +440,11 @@ namespace UE::WindowsApplication::TouchpadSensitivity
 		if (GetForceMaxTouchpadSensitivityMode() == EForceMaxTouchpadSensitivityMode::SetTouchpadParameters)
 		{
 			// If we have a cached value return that.
-			if (ForceMaxTouchpadSensitivity_SavedSensitivity.IsSet())
+			TOUCHPAD_SENSITIVITY_LEVEL LocalSavedSensitivity = ForceMaxTouchpadSensitivity_SavedSensitivity.load();
+			if (LocalSavedSensitivity != TOUCHPAD_SENSITIVITY_LEVEL_NOT_SET)
 			{
-				UE_LOG(LogWindowsTouchpad, Verbose, TEXT("GetForceMaxTouchpadSensitivityRestorationValue returning SavedSensitivity %i"), ForceMaxTouchpadSensitivity_SavedSensitivity.GetValue());
-				return static_cast<FWindowsApplication::EForceMaxTouchpadSensitivityRestorationValues>(ForceMaxTouchpadSensitivity_SavedSensitivity.GetValue());
+				UE_LOG(LogWindowsTouchpad, Verbose, TEXT("GetForceMaxTouchpadSensitivityRestorationValue returning SavedSensitivity %i"), LocalSavedSensitivity);
+				return static_cast<FWindowsApplication::EForceMaxTouchpadSensitivityRestorationValues>(LocalSavedSensitivity);
 			}
 			else
 			{
@@ -453,141 +478,145 @@ namespace UE::WindowsApplication::TouchpadSensitivity
 		}
 	}
 
-	/*
-	* The multithreading model here is very constrained.  Only the system call in Apply-Async is running in a worker thread.  All of the
-	* rest runs in the game thread and all of the state within unreal is updated in the game thread. All of the
-	* apply/remove functions wait on the system call task if it is in progress. This means that we can call Apply-Async and not end up blocking 
-	* for over a second.  But if we make any other call, including another Apply-Async while the task is running we will block for the task and then
-	* block for the system call.  One should keep in mind that the system setting can theoretically change even outside of unreal, though it is probably 
-	* uncommon.
-	* 
-	* The assumption here is that blocking during startup is bad, blocking during shutdown is ok, and these functions should be called infrequently so
-	* blocking on the task will be very rare in reality.
-	*/
-	static UE::Tasks::TTask<void> ApplyForceMaxTouchpadSensitivityAsyncTask;
-	void WaitOnPreviousApplyAsync()
-	{
-		// If we are using the async apply the task could still be outstanding.  Wait until it is done to keep the state sane.
-		// Timeout to make sure we don't block forever if this all goes wrong.
-		if (ApplyForceMaxTouchpadSensitivityAsyncTask.IsValid() && !ApplyForceMaxTouchpadSensitivityAsyncTask.IsCompleted())
-		{
-			UE_LOG(LogWindowsTouchpad, Verbose, TEXT("WaitOnPreviousApplyAsync"));
-			ApplyForceMaxTouchpadSensitivityAsyncTask.Wait(FTimespan::FromSeconds(2.0));
-		}
-	}
 
-	void ApplyForceMaxTouchpadSensitivity(FWindowsApplication::EForceMaxTouchpadSensitivityAsyncBehavior AsyncBehavior)
+	void EnqueueAsyncTouchpadSensitivityChangeTask()
 	{
-		check(IsInGameThread());
-
-		switch (GetForceMaxTouchpadSensitivityMode())
-		{
-		case EForceMaxTouchpadSensitivityMode::NotCached:
-			UE_LOG(LogWindowsTouchpad, Log, TEXT("ApplyForceMaxTouchpadSensitivity NotCached, this is a bug"));
-			check(false);
-			break;
-		case EForceMaxTouchpadSensitivityMode::DoNothing:
-			UE_LOG(LogWindowsTouchpad, Verbose, TEXT("ApplyForceMaxTouchpadSensitivity DoNothing"));
-			break;
-		case EForceMaxTouchpadSensitivityMode::SetTouchpadParameters:
-		{
-			if (ForceMaxTouchpadSensitivity_SavedSensitivity.IsSet())
+		UE_LOG(LogWindowsTouchpad, Verbose, TEXT("TouchpadSensitivityChangeTask Task launching from pipe."));
+		AsyncTaskPipe.Launch(
+			TEXT("TouchpadSensitivityChangeTask"),
+			[]
 			{
-				UE_LOG(LogWindowsTouchpad, Verbose, TEXT("ApplyForceMaxTouchpadSensitivity sensitivity was already changed, doing nothing."));
-				return;
-			}
-
-			// I think its actually impossible to need to wait here.  We would early out just above.  If there was an interveening remove it would have done the wait.
-			WaitOnPreviousApplyAsync();
+				UE_LOG(LogWindowsTouchpad, Verbose, TEXT("TouchpadSensitivityChangeTask Task running."));
 
-			// Set the global system runtime touchpad sensitivity to MostSensitive
-			// Note: we will get the params, change just the one, then write them back so that the reserved param block is preserved.
-			UE_LOG(LogWindowsTouchpad, Log, TEXT("ApplyForceMaxTouchpadSensitivity setting to MostSensitive"));
-			// Get the user's touchpad settings.
-			TOUCHPAD_PARAMETERS_V0 params = {};
-			params.versionNumber = 0;
-			{
-				BOOL bSuccess = ::SystemParametersInfo(SPI_GETTOUCHPADPARAMETERS, sizeof(params), &params, 0 /*fWinIni*/);
-				if (!bSuccess)
+				if (bAbortAsyncTasks.load())
 				{
-					UE_LOG(LogWindowsTouchpad, Warning, TEXT("ApplyForceMaxTouchpadSensitivity SystemParametersInfo SPI_GETTOUCHPADPARAMETERS failed"));
+					UE_LOG(LogWindowsTouchpad, Verbose, TEXT("TouchpadSensitivityChangeTask aborted, doing nothing.  Task Complete."));
 					return;
 				}
-			}
 
-			if (params.sensitivityLevel == TOUCHPAD_SENSITIVITY_LEVEL_MOST_SENSITIVE)
-			{
-				if (bRestoreTouchpadSensitivityFromRegistry)
+				// Get the user's touchpad settings.
+				TOUCHPAD_PARAMETERS_V0 params = {};
+				params.versionNumber = 0;
 				{
-					// If we exited without restoring (crashed, killed in debugger, etc) the saved sensitivity the system's current sensitivity will be MOST, but we will
-					// have saved some other value into the registry.
-					// In that case we want to restore to that registry saved setting.
-					// We do not need to set our sensitivity to MOST, because it already is MOST.
-					TOUCHPAD_SENSITIVITY_LEVEL RegistrySensitivity = TOUCHPAD_SENSITIVITY_LEVEL_MOST_SENSITIVE;
-					bool bSuccess = false;
-					ReadSavedSensitivityFromRegistry(RegistrySensitivity, bSuccess);
-					if (bSuccess && RegistrySensitivity != TOUCHPAD_SENSITIVITY_LEVEL_MOST_SENSITIVE)
+					BOOL bSuccess = ::SystemParametersInfo(SPI_GETTOUCHPADPARAMETERS, sizeof(params), &params, 0 /*fWinIni*/);
+					if (!bSuccess)
 					{
-						UE_LOG(LogWindowsTouchpad, Warning, TEXT("ApplyForceMaxTouchpadSensitivity restoring sensitivity from registry. Value: %i"), (int32)RegistrySensitivity);
-						ForceMaxTouchpadSensitivity_SavedSensitivity = RegistrySensitivity;
+						UE_LOG(LogWindowsTouchpad, Warning, TEXT("TouchpadSensitivityChangeTask SystemParametersInfo SPI_GETTOUCHPADPARAMETERS failed. Task Complete."));
+						return;
 					}
-					return;
 				}
 
-				// If the setting is already MOST_SENSTIVE we will do nothing.
-				// This will also stop us overwriting the SavedSensitivity if we call this function twice.
-				UE_LOG(LogWindowsTouchpad, Verbose, TEXT("ApplyForceMaxTouchpadSensitivity sensitivity was already Most Sensitive"));
-				return;
-			}
+				TOUCHPAD_SENSITIVITY_LEVEL DesiredSensitivity = ForceMaxTouchpadSensitivity_DesiredSensitivity.load();
 
-			ForceMaxTouchpadSensitivity_SavedSensitivity = params.sensitivityLevel;
-			WriteSavedSensitivityToRegistry(ForceMaxTouchpadSensitivity_SavedSensitivity.GetValue());
+				if (DesiredSensitivity == TOUCHPAD_SENSITIVITY_LEVEL_NOT_SET)
+				{
+					UE_LOG(LogWindowsTouchpad, Log, TEXT("TouchpadSensitivityChangeTask DesiredSensitivity is TOUCHPAD_SENSITIVITY_LEVEL_NOT_SET.  This could happen if we tried to restore the setting, but had never changed it.  Setting not changed. Task Complete."));
+					return;
+				}
 
-			// Update the touchpad sensitivity level to "most sensitive" so that simultaneous
-			// keyboard input does not interfere with the touchpad's mouse input generation.
-			// Leave all other bytes in the struct untouched.
-			params.sensitivityLevel = TOUCHPAD_SENSITIVITY_LEVEL_MOST_SENSITIVE;
-
-			if (AsyncBehavior == FWindowsApplication::EForceMaxTouchpadSensitivityAsyncBehavior::AllowAsynchronous)
-			{
-				// The SetTouchpadParameters call can take up to about 1 second, this will happen during game startup increasing aparrant load time.  Therefore we have
-				// the option of doing this asynchronously.
-				UE_LOG(LogWindowsTouchpad, Log, TEXT("ApplyForceMaxTouchpadSensitivity Task launching."));
-				ApplyForceMaxTouchpadSensitivityAsyncTask = UE::Tasks::Launch(
-					TEXT("ApplyForceMaxTouchpadSensitivityTask"),
-					[params] 
-					{ 
-						UE_LOG(LogWindowsTouchpad, Verbose, TEXT("ApplyForceMaxTouchpadSensitivity Task running."));
-
-						// Do not specify SPIF_UPDATEINIFILE so that the registry is not updated.
-						// We will revert this update later, and not writing to the registry is a
-						// failsafe in case we fail to do so (e.g. due to a crash).
-						bool bSuccess = ::SystemParametersInfo(SPI_SETTOUCHPADPARAMETERS, sizeof(params), const_cast<TOUCHPAD_PARAMETERS_V0*>(&params), SPIF_SENDCHANGE);
-						if (!bSuccess)
+				// If we are forcing to max we need to figure out our restore value and save it.
+				if (DesiredSensitivity == TOUCHPAD_SENSITIVITY_LEVEL_MOST_SENSITIVE)
+				{
+					// If we previously exited without restoring (crashed, killed in debugger, etc) the sensitivity the system's current sensitivity will already be MOST, but we may
+					// have saved some other value into the registry so that we can restore it.
+					// We only want to read that registry value once per run of the app.
+					// So first we will mark that we did that, then we might apply it.
+					if (bRestoreTouchpadSensitivityFromRegistry && !bRegistrySavedValueConsumed)
+					{
+						// We only want to attempt to read the registry value once.
+						bRegistrySavedValueConsumed = true;
+
+						// Only if our current sensitivity is already MOST we suspect we might have crashed and want to use the registry value as our saved sensitivity.
+						if (params.sensitivityLevel == DesiredSensitivity)
 						{
-							UE_LOG(LogWindowsTouchpad, Warning, TEXT("ApplyForceMaxTouchpadSensitivityTask SystemParametersInfo SPI_SETTOUCHPADPARAMETERS Failed"));
-							return;
+							TOUCHPAD_SENSITIVITY_LEVEL RegistrySensitivity = TOUCHPAD_SENSITIVITY_LEVEL_MOST_SENSITIVE;
+							bool bSuccess = false;
+							ReadSavedSensitivityFromRegistry(RegistrySensitivity, bSuccess);
+							if (bSuccess && RegistrySensitivity != TOUCHPAD_SENSITIVITY_LEVEL_MOST_SENSITIVE)
+							{
+								UE_LOG(LogWindowsTouchpad, Log, TEXT("TouchpadSensitivityChangeTask restoring saved sensitivity %i from registry, we are already in MostSensitive and do not need to change the setting now. Task Complete."), (int32)RegistrySensitivity);
+								ForceMaxTouchpadSensitivity_SavedSensitivity.store(RegistrySensitivity);
+								// The sensitivity setting is already MOST, so we are done now.
+								return;
+							}
 						}
-						UE_LOG(LogWindowsTouchpad, Log, TEXT("ApplyForceMaxTouchpadSensitivityTask sucessfully changed setting to MostSensitive"));
+						else
+						{
+							// We are not already in MOST, so we can ignore the registry saved value and restore to the current value.
+							// We already marked the registry value consumed, so we can proceed as normal.
+						}
+					}
+				}
+					
+				if (params.sensitivityLevel == DesiredSensitivity)
+				{
+					UE_LOG(LogWindowsTouchpad, Log, TEXT("TouchpadSensitivityChangeTask doing nothing because the current sensitivity already matches DesiredSensitivity %i. Task Complete."), DesiredSensitivity);
+					return;
+				}
 
-					},
-					UE::Tasks::ETaskPriority::Normal);
-			}
-			else
-			{
-				check(AsyncBehavior == FWindowsApplication::EForceMaxTouchpadSensitivityAsyncBehavior::RequireBlocking);
+				if (DesiredSensitivity == TOUCHPAD_SENSITIVITY_LEVEL_MOST_SENSITIVE)
+				{
+					// We are doing the force, save the current value so we can restore it later.
+					UE_LOG(LogWindowsTouchpad, Log, TEXT("TouchpadSensitivityChangeTask setting saved sensitivity %i."), (int32)params.sensitivityLevel);
+					ForceMaxTouchpadSensitivity_SavedSensitivity.store(params.sensitivityLevel);
+					WriteSavedSensitivityToRegistry(params.sensitivityLevel);
+				}
+
+				// And now we will update the sensitivity to the desired value.
+				params.sensitivityLevel = DesiredSensitivity;
+
+				// Last chance to avoid the long system call where we change the setting.
+				if (bAbortAsyncTasks.load())
+				{
+					UE_LOG(LogWindowsTouchpad, Log, TEXT("TouchpadSensitivityChangeTask aborted, doing nothing.  Task Complete."));
+					return;
+				}
 
 				// Do not specify SPIF_UPDATEINIFILE so that the registry is not updated.
 				// We will revert this update later, and not writing to the registry is a
 				// failsafe in case we fail to do so (e.g. due to a crash).
-				bool bSuccess = ::SystemParametersInfo(SPI_SETTOUCHPADPARAMETERS, sizeof(params), &params, SPIF_SENDCHANGE);
+				bool bSuccess = ::SystemParametersInfo(SPI_SETTOUCHPADPARAMETERS, sizeof(params), const_cast<TOUCHPAD_PARAMETERS_V0*>(&params), SPIF_SENDCHANGE);
 				if (!bSuccess)
 				{
-					UE_LOG(LogWindowsTouchpad, Warning, TEXT("ApplyForceMaxTouchpadSensitivity SystemParametersInfo SPI_SETTOUCHPADPARAMETERS Failed"));
+					UE_LOG(LogWindowsTouchpad, Warning, TEXT("TouchpadSensitivityChangeTask SystemParametersInfo SPI_SETTOUCHPADPARAMETERS Failed.  Task Complete."));
 					return;
 				}
-				UE_LOG(LogWindowsTouchpad, Log, TEXT("ApplyForceMaxTouchpadSensitivity sucessfully changed setting to MostSensitive"));
-			}
+				else
+				{
+					UE_LOG(LogWindowsTouchpad, Log, TEXT("TouchpadSensitivityChangeTask sucessfully changed setting to %i."), DesiredSensitivity);
+
+					// If we have restored a value we do not want a value stored in the registry anymore. (if the user changes their setting now we do not want to overwrite it later)
+					if (DesiredSensitivity != TOUCHPAD_SENSITIVITY_LEVEL_MOST_SENSITIVE)
+					{
+						UE_LOG(LogWindowsTouchpad, Log, TEXT("TouchpadSensitivityChangeTask ClearSavedSensitivityFromRegistry."), DesiredSensitivity);
+						ClearSavedSensitivityFromRegistry();
+					}
+
+					UE_LOG(LogWindowsTouchpad, Verbose, TEXT("TouchpadSensitivityChangeTask Task Complete."), DesiredSensitivity);
+					return;
+				}
+			},
+			UE::Tasks::ETaskPriority::Normal);
+	}
+
+	void ApplyForceMaxTouchpadSensitivity()
+	{
+		UE_LOG(LogWindowsTouchpad, Verbose, TEXT("ApplyForceMaxTouchpadSensitivity Starting."));
+		check(IsInGameThread());
+
+		switch (GetForceMaxTouchpadSensitivityMode())
+		{
+		case EForceMaxTouchpadSensitivityMode::NotCached:
+			UE_LOG(LogWindowsTouchpad, Warning, TEXT("ApplyForceMaxTouchpadSensitivity NotCached, this is a bug"));
+			check(false);
+			break;
+		case EForceMaxTouchpadSensitivityMode::DoNothing:
+			UE_LOG(LogWindowsTouchpad, Verbose, TEXT("ApplyForceMaxTouchpadSensitivity DoNothing"));
+			break;
+		case EForceMaxTouchpadSensitivityMode::SetTouchpadParameters:
+		{
+			ForceMaxTouchpadSensitivity_DesiredSensitivity.store(TOUCHPAD_SENSITIVITY_LEVEL_MOST_SENSITIVE);
+			EnqueueAsyncTouchpadSensitivityChangeTask();
 		}
 		break;
 		case EForceMaxTouchpadSensitivityMode::SetMaxTouchpadSensitivity:
@@ -595,7 +624,6 @@ namespace UE::WindowsApplication::TouchpadSensitivity
 			if (!SetMaxTouchpadSensitivity(true))
 			{
 				UE_LOG(LogWindowsTouchpad, Warning, TEXT("ApplyForceMaxTouchpadSensitivity SetMaxTouchpadSensitivity Failed"));
-				return;
 			}
 		}
 		break;
@@ -603,20 +631,70 @@ namespace UE::WindowsApplication::TouchpadSensitivity
 			check(false);
 			break;
 		}
+		UE_LOG(LogWindowsTouchpad, Verbose, TEXT("ApplyForceMaxTouchpadSensitivity Completed."));
 	}
 
 	void RemoveForceMaxTouchpadSensitivity()
 	{
+		UE_LOG(LogWindowsTouchpad, Verbose, TEXT("RemoveForceMaxTouchpadSensitivity Starting."));
 		check(IsInGameThread());
 
-		// Only if we did set the sensitivity using the EForceMaxTouchpadSensitivityMode::TouchpadParameters implementation
-		// will we have cached a value, and need to restore it.
-		if (ForceMaxTouchpadSensitivity_SavedSensitivity.IsSet())
+		switch (GetForceMaxTouchpadSensitivityMode())
+		{
+		case EForceMaxTouchpadSensitivityMode::NotCached:
+			check(false);
+			break;
+		case EForceMaxTouchpadSensitivityMode::DoNothing:
+			UE_LOG(LogWindowsTouchpad, Verbose, TEXT("RemoveForceMaxTouchpadSensitivity DoNothing"));
+			break;
+		case EForceMaxTouchpadSensitivityMode::SetTouchpadParameters:
+		{
+			ForceMaxTouchpadSensitivity_DesiredSensitivity.store(ForceMaxTouchpadSensitivity_SavedSensitivity.load());
+			EnqueueAsyncTouchpadSensitivityChangeTask();
+		}
+		break;
+		case EForceMaxTouchpadSensitivityMode::SetMaxTouchpadSensitivity:
+		{
+			if (!SetMaxTouchpadSensitivity(false))
+			{
+				UE_LOG(LogWindowsTouchpad, Warning, TEXT("RemoveForceMaxTouchpadSensitivity SetMaxTouchpadSensitivity Failed"));
+			}
+		}
+		break;
+		default:
+			check(false);
+			break;
+		}
+		UE_LOG(LogWindowsTouchpad, Verbose, TEXT("RemoveForceMaxTouchpadSensitivity Completed."));
+	}
+
+	// Note: this function could be called from threads other than the main game thread.
+	// If an async change task is already running it is possible for us to end up in a bad state, but I don't want to add locks or waits to this case
+	// should still be available.
+	void RemoveForceMaxTouchpadSensitivityOnErrorShutdown()
+	{
+		if (bForceMaxTouchpadSensitivity)
 		{
-			WaitOnPreviousApplyAsync();
+			UE_LOG(LogWindowsTouchpad, Verbose, TEXT("RemoveForceMaxTouchpadSensitivityOnErrorShutdown Starting."));
+
+			// New Async tasks that start should do nothing.
+			bAbortAsyncTasks.store(true);
+
+			TOUCHPAD_SENSITIVITY_LEVEL LocalSavedSensitivity = ForceMaxTouchpadSensitivity_SavedSensitivity.load();
+
+			if (LocalSavedSensitivity <= TOUCHPAD_SENSITIVITY_LEVEL_MOST_SENSITIVE || LocalSavedSensitivity > TOUCHPAD_SENSITIVITY_LEVEL_LEAST_SENSITIVE)
+			{
+				// No restore needed, NOT_SET, or somehow an invalid value.
+				return;
+			}
+
+			// At this point we have a saved sensitivity value that it makes sense to restore.
+
+			// Set the atomic desired sensitivity.  If there is a task currently executing this might catch it and have it do the right thing.
+			ForceMaxTouchpadSensitivity_DesiredSensitivity.store(LocalSavedSensitivity);
 
 			// Note: we will get the params, change just the one, then write them back so that the reserved param block is preserved.
-			UE_LOG(LogWindowsTouchpad, Log, TEXT("RemoveForceMaxTouchpadSensitivity reverting to SavedSensitivity %i"), ForceMaxTouchpadSensitivity_SavedSensitivity.GetValue());
+			UE_LOG(LogWindowsTouchpad, Verbose, TEXT("RemoveForceMaxTouchpadSensitivityOnErrorShutdown reverting to SavedSensitivity %i"), LocalSavedSensitivity);
 
 			// Get the user's touchpad settings.
 			TOUCHPAD_PARAMETERS_V0 params = {};
@@ -624,13 +702,13 @@ namespace UE::WindowsApplication::TouchpadSensitivity
 			BOOL bSuccess = ::SystemParametersInfo(SPI_GETTOUCHPADPARAMETERS, sizeof(params), &params, 0 /*fWinIni*/);
 			if (!bSuccess)
 			{
-				UE_LOG(LogWindowsTouchpad, Warning, TEXT("RemoveForceMaxTouchpadSensitivity SystemParametersInfo SPI_GETTOUCHPADPARAMETERS failed"));
+				UE_LOG(LogWindowsTouchpad, Warning, TEXT("RemoveForceMaxTouchpadSensitivityOnErrorShutdown SystemParametersInfo SPI_GETTOUCHPADPARAMETERS failed"));
 				return;
 			}
 
 			// Restore the saved sensitivity level.
 			// Leave all other bytes in the struct untouched.
-			params.sensitivityLevel = ForceMaxTouchpadSensitivity_SavedSensitivity.GetValue();
+			params.sensitivityLevel = LocalSavedSensitivity;
 
 			// Do not specify SPIF_UPDATEINIFILE so that the registry is not updated.
 			// We will revert this update later, and not writing to the registry is a
@@ -638,89 +716,22 @@ namespace UE::WindowsApplication::TouchpadSensitivity
 			bSuccess = ::SystemParametersInfo(SPI_SETTOUCHPADPARAMETERS, sizeof(params), &params, SPIF_SENDCHANGE);
 			if (!bSuccess)
 			{
-				UE_LOG(LogWindowsTouchpad, Warning, TEXT("RemoveForceMaxTouchpadSensitivity SystemParametersInfo SPI_SETTOUCHPADPARAMETERS failed"));
+				UE_LOG(LogWindowsTouchpad, Warning, TEXT("RemoveForceMaxTouchpadSensitivityOnErrorShutdown SystemParametersInfo SPI_SETTOUCHPADPARAMETERS failed"));
 				return;
 			}
-			UE_LOG(LogWindowsTouchpad, Log, TEXT("RemoveForceMaxTouchpadSensitivity sucessfully reverted to SavedSensitivity %i"), ForceMaxTouchpadSensitivity_SavedSensitivity.GetValue());
-
-			ForceMaxTouchpadSensitivity_SavedSensitivity.Reset();
-			ClearSavedSensitivityFromRegistry();
-		}
-		else if (GetForceMaxTouchpadSensitivityMode() == EForceMaxTouchpadSensitivityMode::SetMaxTouchpadSensitivity)
-		{
-			if (!!SetMaxTouchpadSensitivity(false))
-			{
-				UE_LOG(LogWindowsTouchpad, Warning, TEXT("RemoveForceMaxTouchpadSensitivity SetMaxTouchpadSensitivity Failed"));
-				return;
-			}
-		}
-		else
-		{
-			UE_LOG(LogWindowsTouchpad, Verbose, TEXT("RemoveForceMaxTouchpadSensitivity Doing nothing because the sensitivity is not currently forced."));
+			UE_LOG(LogWindowsTouchpad, Log, TEXT("RemoveForceMaxTouchpadSensitivityOnErrorShutdown sucessfully reverted to SavedSensitivity %i"), LocalSavedSensitivity);
 		}
+		UE_LOG(LogWindowsTouchpad, Verbose, TEXT("RemoveForceMaxTouchpadSensitivityOnErrorShutdown Completed."));
 	}
 
-	// Note: this function could be called from threads other than the main game thread.
-	// It is NOT fully safe to do so, ForceMaxTouchpadSensitivity_SavedSensitivity is not threadsafe and its SET and VALUE could be mismatched and we are not going to wait for 
-	// a potential worker-thread system call to complete.
-	// However we are in an error state already, it is very likely that it will work fine, we can do some simple checks to make sure it doesn't do anything really strange,
-	// and I want to avoid adding locks to this error case.
-	// TOptional, at the time of writing, is implemented such that a Reset does not change a saved enum value, so cases where another thread is removing the setting and this code also runs should
-	// result in a double-restoration of the saved value.  This is OK.
-	// TOptional, at the time of writing, is written so that the value is saved before the Set flag is set however if it does actually run the operations in the reverse order then the very first set passes through
-	// a state where the flag is set and the value is not saved.  In that case we might incorrectly restore to the MostSensitive setting.  In other cases we would restore the correct value,
-	// regardless of the flag value.
-	// To almost perfectly handle error cases we need an external process watching our game and restoring the setting on any abnormal termination.
-	void RemoveForceMaxTouchpadSensitivityOnErrorShutdown()
+	// Note that the other thee app event hooks only need to do anything in EForceMaxTouchpadSensitivityMode::SetTouchpadParameters
+	void HandleStartupApp()
 	{
-		if (bForceMaxTouchpadSensitivity && GetForceMaxTouchpadSensitivityMode() == EForceMaxTouchpadSensitivityMode::SetTouchpadParameters)
+		UE_LOG(LogWindowsTouchpad, Verbose, TEXT("HandleStartupApp"));
+		CacheForceMaxTouchpadSensitivityMode();
+		if (bForceMaxTouchpadSensitivity)
 		{
-			// Make a local copy of the saved sensitivity to work with.  Now we know it won't change out from under us.  We only have to worry about it being half-set.
-			TOptional<TOUCHPAD_SENSITIVITY_LEVEL> LocalSavedSensitivity = ForceMaxTouchpadSensitivity_SavedSensitivity;
-
-			// Only if we did set the sensitivity using the EForceMaxTouchpadSensitivityMode::TouchpadParameters implementation
-			// will we have cached a value, and need to restore it.
-			if (LocalSavedSensitivity.IsSet())
-			{
-				// The value should be in the valid range, and not MOST_SENSITIVE, because in that case we would not have needed to change it.
-				// If it is MOST_SENSITIVE or too large the game thread might have been setting or resetting it.
-				if (LocalSavedSensitivity.GetValue() <= TOUCHPAD_SENSITIVITY_LEVEL_MOST_SENSITIVE || LocalSavedSensitivity.GetValue() > TOUCHPAD_SENSITIVITY_LEVEL_LEAST_SENSITIVE)
-				{
-					// No restore needed, or invalid value.
-					return;
-				}
-
-				// We are not doing this wait in this error version.
-				//WaitOnPreviousApplyAsync();
-
-				// Note: we will get the params, change just the one, then write them back so that the reserved param block is preserved.
-				UE_LOG(LogWindowsTouchpad, Verbose, TEXT("RemoveForceMaxTouchpadSensitivityOnErrorShutdown reverting to SavedSensitivity %i"), LocalSavedSensitivity.GetValue());
-
-				// Get the user's touchpad settings.
-				TOUCHPAD_PARAMETERS_V0 params = {};
-				params.versionNumber = 0;
-				BOOL bSuccess = ::SystemParametersInfo(SPI_GETTOUCHPADPARAMETERS, sizeof(params), &params, 0 /*fWinIni*/);
-				if (!bSuccess)
-				{
-					UE_LOG(LogWindowsTouchpad, Warning, TEXT("RemoveForceMaxTouchpadSensitivityOnErrorShutdown SystemParametersInfo SPI_GETTOUCHPADPARAMETERS failed"));
-					return;
-				}
-
-				// Restore the saved sensitivity level.
-				// Leave all other bytes in the struct untouched.
-				params.sensitivityLevel = LocalSavedSensitivity.GetValue();
-
-				// Do not specify SPIF_UPDATEINIFILE so that the registry is not updated.
-				// We will revert this update later, and not writing to the registry is a
-				// failsafe in case we fail to do so (e.g. due to a crash).
-				bSuccess = ::SystemParametersInfo(SPI_SETTOUCHPADPARAMETERS, sizeof(params), &params, SPIF_SENDCHANGE);
-				if (!bSuccess)
-				{
-					UE_LOG(LogWindowsTouchpad, Warning, TEXT("RemoveForceMaxTouchpadSensitivityOnErrorShutdown SystemParametersInfo SPI_SETTOUCHPADPARAMETERS failed"));
-					return;
-				}
-				UE_LOG(LogWindowsTouchpad, Log, TEXT("RemoveForceMaxTouchpadSensitivityOnErrorShutdown sucessfully reverted to SavedSensitivity %i"), LocalSavedSensitivity.GetValue());
-			}
+			ApplyForceMaxTouchpadSensitivity();
 		}
 	}
 
@@ -731,7 +742,7 @@ namespace UE::WindowsApplication::TouchpadSensitivity
 			UE_LOG(LogWindowsTouchpad, Verbose, TEXT("HandleActivateApp %i"), bActivate);
 			if (bActivate)
 			{
-				ApplyForceMaxTouchpadSensitivity(FWindowsApplication::EForceMaxTouchpadSensitivityAsyncBehavior::AllowAsynchronous);
+				ApplyForceMaxTouchpadSensitivity();
 			}
 			else
 			{
@@ -740,16 +751,6 @@ namespace UE::WindowsApplication::TouchpadSensitivity
 		}
 	}
 
-	void HandleStartupApp()
-	{
-		UE_LOG(LogWindowsTouchpad, Verbose, TEXT("HandleStartupApp"));
-		CacheForceMaxTouchpadSensitivityMode();
-		if (bForceMaxTouchpadSensitivity)
-		{
-			ApplyForceMaxTouchpadSensitivity(FWindowsApplication::EForceMaxTouchpadSensitivityAsyncBehavior::AllowAsynchronous);
-		}
-	}
-
 	void HandleDestroyApplication()
 	{
 		if (bForceMaxTouchpadSensitivity && (GetForceMaxTouchpadSensitivityMode() == EForceMaxTouchpadSensitivityMode::SetTouchpadParameters))
@@ -784,7 +785,7 @@ namespace UE::WindowsApplication::TouchpadSensitivity
 			{
 				UE_LOG(LogWindowsTouchpad, Verbose, TEXT("WindowsApplication.ForceMaxTouchpadSensitivity CVarSinkHandler calling ApplyForceMaxTouchpadSensitivity"));
 				CacheForceMaxTouchpadSensitivityMode();
-				ApplyForceMaxTouchpadSensitivity(FWindowsApplication::EForceMaxTouchpadSensitivityAsyncBehavior::AllowAsynchronous);
+				ApplyForceMaxTouchpadSensitivity();
 			}
 			else
 			{
@@ -987,9 +988,9 @@ bool FWindowsApplication::SupportsForceMaxTouchpadSensitivity()
 	return UE::WindowsApplication::TouchpadSensitivity::SupportsForceMaxTouchpadSensitivity();
 }
 
-void FWindowsApplication::ApplyForceMaxTouchpadSensitivity(FWindowsApplication::EForceMaxTouchpadSensitivityAsyncBehavior AsyncBehavior)
+void FWindowsApplication::ApplyForceMaxTouchpadSensitivity()
 {
-	UE::WindowsApplication::TouchpadSensitivity::ApplyForceMaxTouchpadSensitivity(AsyncBehavior);
+	UE::WindowsApplication::TouchpadSensitivity::ApplyForceMaxTouchpadSensitivity();
 }
 
 void FWindowsApplication::RemoveForceMaxTouchpadSensitivity()
diff --git a/Engine/Source/Runtime/ApplicationCore/Public/Windows/WindowsApplication.h b/Engine/Source/Runtime/ApplicationCore/Public/Windows/WindowsApplication.h
--- a/Engine/Source/Runtime/ApplicationCore/Public/Windows/WindowsApplication.h
+++ b/Engine/Source/Runtime/ApplicationCore/Public/Windows/WindowsApplication.h
@@ -389,7 +389,7 @@ public:
 	 * In Windows 11 24H2 a new api was added which can change this setting only for this app.
 	 * 
 	 * We reccomend that this feature be exposed as a user setting which is enabled by default and hidden/disabled if SupportsForceMaxTouchpadSensitivity() returns false and that the CVar be used to activate it rather than directly calling 
-	 * ApplyForceMaxTouchpadSensitivity/RemoveForceMaxTouchpadSensitivity due to the complexity of handling the focus/exit/crash cases with the 22H2 api.
+	 * ApplyForceMaxTouchpadSensitivity/RemoveForceMaxTouchpadSensitivity due to the complexity of handling the focus/exit/crash cases with the SetTouchpadParameters api.
 	 * 
 	 * See WindowsApplication.cpp namespace UE::WindowsApplication::TouchpadSensitivity for more information about this system.
 	 */
@@ -399,7 +399,7 @@ public:
 		RequireBlocking // This function will not return until the system call returns.
 	};
 	APPLICATIONCORE_API bool SupportsForceMaxTouchpadSensitivity();
-	APPLICATIONCORE_API void ApplyForceMaxTouchpadSensitivity(EForceMaxTouchpadSensitivityAsyncBehavior AsyncBehavior);
+	APPLICATIONCORE_API void ApplyForceMaxTouchpadSensitivity();
 	APPLICATIONCORE_API void RemoveForceMaxTouchpadSensitivity();
 
 	/** If we are using the system wide setting it is possible that we could crash and fail to set it back.  Call this function to
