Initialize FDeferredAndroidMessage.KeyEventData with default values.
#rnx

diff --git a/Engine/Source/Runtime/ApplicationCore/Public/Android/AndroidInputInterface.h b/Engine/Source/Runtime/ApplicationCore/Public/Android/AndroidInputInterface.h
--- a/Engine/Source/Runtime/ApplicationCore/Public/Android/AndroidInputInterface.h
+++ b/Engine/Source/Runtime/ApplicationCore/Public/Android/AndroidInputInterface.h
@@ -255,12 +255,12 @@ struct FDeferredAndroidMessage
 	{
 		struct
 		{
-			int32 keyId;
-			int32 unichar;
-			uint32 modifier;
-			bool  isRepeat;
+			int32 keyId{};
+			int32 unichar{};
+			uint32 modifier{};
+			bool isRepeat{};
 		}
-		KeyEventData;
+		KeyEventData{};
 
 	};
 };
