SCKit to replace deprecated API for eyedropper

#rnx

diff --git a/Engine/Source/Runtime/ApplicationCore/Private/Mac/MacPlatformApplicationMisc.cpp b/Engine/Source/Runtime/ApplicationCore/Private/Mac/MacPlatformApplicationMisc.cpp
--- a/Engine/Source/Runtime/ApplicationCore/Private/Mac/MacPlatformApplicationMisc.cpp
+++ b/Engine/Source/Runtime/ApplicationCore/Private/Mac/MacPlatformApplicationMisc.cpp
@@ -32,8 +32,11 @@
 #include <libproc.h>
 #include <notify.h>
 #include <uuid/uuid.h>
+#import <ScreenCaptureKit/ScreenCaptureKit.h>
 #include "Apple/PostAppleSystemHeaders.h"
 
+
+
 extern FMacMallocCrashHandler* GCrashMalloc;
 
 MacApplicationExternalCb FPlatformApplicationMisc::UpdateCachedMacMenuStateCb = nullptr;
@@ -420,19 +423,60 @@ bool FMacPlatformApplicationMisc::ControlScreensaver(EScreenSaverAction Action)
 FLinearColor FMacPlatformApplicationMisc::GetScreenPixelColor(const FVector2D& InScreenPos, float /*InGamma*/)
 {
 	SCOPED_AUTORELEASE_POOL;
-
-	const CGPoint Pt = FMacApplication::ConvertSlatePositionToCGPoint(InScreenPos.X, InScreenPos.Y);
-	CGImageRef ScreenImage = CGWindowListCreateImage(CGRectMake(Pt.x, Pt.y, 1, 1), kCGWindowListOptionOnScreenBelowWindow, kCGNullWindowID, kCGWindowImageDefault);
-	
-	CGDataProviderRef provider = CGImageGetDataProvider(ScreenImage);
-	NSData* data = (id)CGDataProviderCopyData(provider);
-	[data autorelease];
-	const uint8* bytes = (const uint8*)[data bytes];
+	__block bool screenshotDone = false;
+	__block FColor color(0, 0, 0);
+#ifdef __MAC_15_2
+	if (@available(macOS 15.2, *))
+	{
+	   const CGPoint Pt = FMacApplication::ConvertSlatePositionToCGPoint(InScreenPos.X, InScreenPos.Y);
+	   [SCScreenshotManager captureImageInRect:CGRectMake(Pt.x, Pt.y, 1, 1) completionHandler:^(CGImageRef sampleBuffer, NSError *error)
+		{
+			if (!error)
+			{
+				CFDataRef ImageDataRef = CGDataProviderCopyData(CGImageGetDataProvider(sampleBuffer));
+				const UInt8* rawData = CFDataGetBytePtr(ImageDataRef);
+				color = FColor(rawData[2], rawData[1], rawData[0]);
+				CFRelease(ImageDataRef);
+				CGImageRelease(sampleBuffer);
+			}
+			screenshotDone = true;
+		}];
+	}
+	else
+#endif
+	{
+		[SCShareableContent getShareableContentWithCompletionHandler:^(SCShareableContent *shareableContent, NSError *error)
+		 {
+			if (!error)
+			{
+				// todo: find the right display
+				SCContentFilter* Filter = [[[SCContentFilter alloc] initWithDisplay:[shareableContent.displays objectAtIndex:0] excludingWindows:[NSArray array]] autorelease];
+				[SCScreenshotManager captureImageWithFilter:Filter
+											  configuration:[SCStreamConfiguration streamConfigurationWithPreset:SCStreamConfigurationPresetCaptureHDRScreenshotLocalDisplay]
+										  completionHandler:^(CGImageRef sampleBuffer, NSError *error)
+				 {
+					// todo: get the right pixel
+					screenshotDone = true;
+				 }];
+			}
+			else
+			{
+				// most likely permission denied
+				screenshotDone = true;
+			}
+		 }];
+	}
 	
-	// Mac colors are gamma corrected in Pow(2.2) space, so do the conversion using the 2.2 to linear conversion.
-	FColor ScreenColor(bytes[2], bytes[1], bytes[0]);
-	FLinearColor ScreenLinearColor = FLinearColor::FromPow22Color(ScreenColor);
-	CGImageRelease(ScreenImage);
+	float timer = 0;
+	do {
+		FPlatformProcess::Sleep(0.002f);
+		if (screenshotDone || timer > 1)	// usually takes ~0.05s
+			break; 
+		else 
+			timer += 0.002f;
+	} while (true);
+
+	FLinearColor ScreenLinearColor = FLinearColor::FromPow22Color(color);
 
 	return ScreenLinearColor;
 }
diff --git a/Engine/Source/Runtime/Core/Core.Build.cs b/Engine/Source/Runtime/Core/Core.Build.cs
--- a/Engine/Source/Runtime/Core/Core.Build.cs
+++ b/Engine/Source/Runtime/Core/Core.Build.cs
@@ -111,7 +111,7 @@ public class Core : ModuleRules
 				"zlib",
 				"PLCrashReporter"
 				);
-			PublicFrameworks.AddRange(new string[] { "Cocoa", "Carbon", "IOKit", "Security", "UniformTypeIdentifiers", "Network" });
+			PublicFrameworks.AddRange(new string[] { "Cocoa", "Carbon", "IOKit", "Security", "UniformTypeIdentifiers", "Network", "ScreenCaptureKit" });
 
 			PrivateDependencyModuleNames.Add("mimalloc");
 			PrivateDefinitions.Add("PLATFORM_BUILDS_MIMALLOC=1");
