fix compile error on mac SDK 14.0

diff --git a/Engine/Source/Runtime/ApplicationCore/Private/Mac/MacPlatformApplicationMisc.cpp b/Engine/Source/Runtime/ApplicationCore/Private/Mac/MacPlatformApplicationMisc.cpp
--- a/Engine/Source/Runtime/ApplicationCore/Private/Mac/MacPlatformApplicationMisc.cpp
+++ b/Engine/Source/Runtime/ApplicationCore/Private/Mac/MacPlatformApplicationMisc.cpp
@@ -452,7 +452,7 @@ FLinearColor FMacPlatformApplicationMisc::GetScreenPixelColor(const FVector2D& I
 				// todo: find the right display
 				SCContentFilter* Filter = [[[SCContentFilter alloc] initWithDisplay:[shareableContent.displays objectAtIndex:0] excludingWindows:[NSArray array]] autorelease];
 				[SCScreenshotManager captureImageWithFilter:Filter
-											  configuration:[SCStreamConfiguration streamConfigurationWithPreset:SCStreamConfigurationPresetCaptureHDRScreenshotLocalDisplay]
+											  configuration:[[[SCStreamConfiguration alloc] init] autorelease]//WithPreset:SCStreamConfigurationPresetCaptureHDRScreenshotLocalDisplay]
 										  completionHandler:^(CGImageRef sampleBuffer, NSError *error)
 				 {
 					// todo: get the right pixel
