Fix compile error

#rnx

diff --git a/Engine/Source/Runtime/ApplicationCore/Private/Mac/MacPlatformApplicationMisc.cpp b/Engine/Source/Runtime/ApplicationCore/Private/Mac/MacPlatformApplicationMisc.cpp
--- a/Engine/Source/Runtime/ApplicationCore/Private/Mac/MacPlatformApplicationMisc.cpp
+++ b/Engine/Source/Runtime/ApplicationCore/Private/Mac/MacPlatformApplicationMisc.cpp
@@ -453,7 +453,7 @@ FLinearColor FMacPlatformApplicationMisc::GetScreenPixelColor(const FVector2D& I
 				SCContentFilter* Filter = [[[SCContentFilter alloc] initWithDisplay:[shareableContent.displays objectAtIndex:0] excludingWindows:[NSArray array]] autorelease];
 				[SCScreenshotManager captureImageWithFilter:Filter
 											  configuration:[[[SCStreamConfiguration alloc] init] autorelease]//WithPreset:SCStreamConfigurationPresetCaptureHDRScreenshotLocalDisplay]
-										  completionHandler:^(CGImageRef sampleBuffer, NSError *error)
+										  completionHandler:^(CGImageRef sampleBuffer, NSError *error2)
 				 {
 					// todo: get the right pixel
 					screenshotDone = true;
