Restore eyedropper tool functionality under macOS 14.0~15.1

#rnx

diff --git a/Engine/Source/Runtime/ApplicationCore/Private/Mac/MacPlatformApplicationMisc.cpp b/Engine/Source/Runtime/ApplicationCore/Private/Mac/MacPlatformApplicationMisc.cpp
--- a/Engine/Source/Runtime/ApplicationCore/Private/Mac/MacPlatformApplicationMisc.cpp
+++ b/Engine/Source/Runtime/ApplicationCore/Private/Mac/MacPlatformApplicationMisc.cpp
@@ -423,23 +423,25 @@ bool FMacPlatformApplicationMisc::ControlScreensaver(EScreenSaverAction Action)
 FLinearColor FMacPlatformApplicationMisc::GetScreenPixelColor(const FVector2D& InScreenPos, float /*InGamma*/)
 {
 	SCOPED_AUTORELEASE_POOL;
-	__block bool screenshotDone = false;
-	__block FColor color(0, 0, 0);
+	__block bool ScreenshotDone = false;
+	__block FColor Color(0, 0, 0);
 #ifdef __MAC_15_2
 	if (@available(macOS 15.2, *))
 	{
 	   const CGPoint Pt = FMacApplication::ConvertSlatePositionToCGPoint(InScreenPos.X, InScreenPos.Y);
 	   [SCScreenshotManager captureImageInRect:CGRectMake(Pt.x, Pt.y, 1, 1) completionHandler:^(CGImageRef sampleBuffer, NSError *error)
 		{
-			if (!error)
+			if (!error && sampleBuffer)
 			{
 				CFDataRef ImageDataRef = CGDataProviderCopyData(CGImageGetDataProvider(sampleBuffer));
 				const UInt8* rawData = CFDataGetBytePtr(ImageDataRef);
-				color = FColor(rawData[2], rawData[1], rawData[0]);
+				Color = FColor(rawData[2], rawData[1], rawData[0]);
 				CFRelease(ImageDataRef);
+				// Somehow we need to release sampleBuffer manually, or else the memory will leak
+				// This is contradictory to Cocoa API naming guidelines, and to a similar API captureImageWithFilter:configuration:completionHandler: 
 				CGImageRelease(sampleBuffer);
 			}
-			screenshotDone = true;
+			ScreenshotDone = true;
 		}];
 	}
 	else
@@ -447,36 +449,62 @@ FLinearColor FMacPlatformApplicationMisc::GetScreenPixelColor(const FVector2D& I
 	{
 		[SCShareableContent getShareableContentWithCompletionHandler:^(SCShareableContent *shareableContent, NSError *error)
 		 {
-			if (!error)
+			if (!error && shareableContent)
 			{
-				// todo: find the right display
-				SCContentFilter* Filter = [[[SCContentFilter alloc] initWithDisplay:[shareableContent.displays objectAtIndex:0] excludingWindows:[NSArray array]] autorelease];
+				const CGPoint Pt = FMacApplication::ConvertSlatePositionToCGPoint(InScreenPos.X, InScreenPos.Y);
+				CGDirectDisplayID matchingDisplay;
+				uint32_t matchingDisplayCount;
+				CGError error = CGGetDisplaysWithRect(CGRectMake(Pt.x, Pt.y, 1, 1), 64, &matchingDisplay, &matchingDisplayCount);
+				assert(error == kCGErrorSuccess);
+				assert(matchingDisplayCount == 1);
+				SCDisplay* shareableDisplay = nil;
+				for (SCDisplay* scDisplay in shareableContent.displays)
+				{
+					if (scDisplay.displayID == matchingDisplay)
+					{
+						shareableDisplay = scDisplay;
+						break;
+					}
+				}
+				SCContentFilter* Filter = [[[SCContentFilter alloc] initWithDisplay:shareableDisplay excludingWindows:[NSArray array]] autorelease];
+				FMacScreenRef Screen = FMacApplication::FindScreenBySlatePosition(InScreenPos.X, InScreenPos.Y);
+				SCStreamConfiguration* stream = [[[SCStreamConfiguration alloc] init] autorelease];
+				stream.width = 100;
+				stream.height = 100;
+				stream.sourceRect = CGRectMake(Pt.x - Screen->Frame.origin.x, Pt.y - Screen->Frame.origin.y, 1, 1);
+				stream.pixelFormat = 'BGRA';
+				stream.showsCursor = false;
 				[SCScreenshotManager captureImageWithFilter:Filter
-											  configuration:[[[SCStreamConfiguration alloc] init] autorelease]//WithPreset:SCStreamConfigurationPresetCaptureHDRScreenshotLocalDisplay]
+											  configuration:stream
 										  completionHandler:^(CGImageRef sampleBuffer, NSError *error2)
 				 {
-					// todo: get the right pixel
-					screenshotDone = true;
+					if (!error2 && sampleBuffer)
+					{
+						CFDataRef ImageDataRef = CGDataProviderCopyData(CGImageGetDataProvider(sampleBuffer));
+						const UInt8* rawData = CFDataGetBytePtr(ImageDataRef);
+						Color = FColor(rawData[2], rawData[1], rawData[0]);
+						CFRelease(ImageDataRef);
+					}
+					ScreenshotDone = true;
 				 }];
 			}
 			else
 			{
-				// most likely permission denied
-				screenshotDone = true;
+				ScreenshotDone = true;
 			}
 		 }];
 	}
 	
-	float timer = 0;
+	float Timer = 0;
 	do {
-		FPlatformProcess::Sleep(0.002f);
-		if (screenshotDone || timer > 1)	// usually takes ~0.05s
+		FPlatformProcess::Sleep(0.01f);
+		if (ScreenshotDone || Timer > 1.0f)	// usually takes ~0.05s
 			break; 
 		else 
-			timer += 0.002f;
+			Timer += 0.01f;
 	} while (true);
 
-	FLinearColor ScreenLinearColor = FLinearColor::FromPow22Color(color);
+	FLinearColor ScreenLinearColor = FLinearColor::FromPow22Color(Color);
 
 	return ScreenLinearColor;
 }
