Fix compiler warning

#rnx

diff --git a/Engine/Source/Runtime/ApplicationCore/Private/Mac/MacPlatformApplicationMisc.cpp b/Engine/Source/Runtime/ApplicationCore/Private/Mac/MacPlatformApplicationMisc.cpp
--- a/Engine/Source/Runtime/ApplicationCore/Private/Mac/MacPlatformApplicationMisc.cpp
+++ b/Engine/Source/Runtime/ApplicationCore/Private/Mac/MacPlatformApplicationMisc.cpp
@@ -454,8 +454,8 @@ FLinearColor FMacPlatformApplicationMisc::GetScreenPixelColor(const FVector2D& I
 				const CGPoint Pt = FMacApplication::ConvertSlatePositionToCGPoint(InScreenPos.X, InScreenPos.Y);
 				CGDirectDisplayID matchingDisplay;
 				uint32_t matchingDisplayCount;
-				CGError error = CGGetDisplaysWithRect(CGRectMake(Pt.x, Pt.y, 1, 1), 64, &matchingDisplay, &matchingDisplayCount);
-				assert(error == kCGErrorSuccess);
+				CGError cgError = CGGetDisplaysWithRect(CGRectMake(Pt.x, Pt.y, 1, 1), 64, &matchingDisplay, &matchingDisplayCount);
+				assert(cgError == kCGErrorSuccess);
 				assert(matchingDisplayCount == 1);
 				SCDisplay* shareableDisplay = nil;
 				for (SCDisplay* scDisplay in shareableContent.displays)
