IOS config rules parsing buffer overrun fix

Config rules were uncompressed as an ansi string, which is not null-terminated. However, string conversion code expected it to be, which caused buffer overrun due to strlen scanning the buffer.

Fix is to use a different StringCast overload, passing an explicit length, which we know.

#rnx

diff --git a/Engine/Source/Runtime/ApplicationCore/Private/IOS/IOSConfigRules.cpp b/Engine/Source/Runtime/ApplicationCore/Private/IOS/IOSConfigRules.cpp
--- a/Engine/Source/Runtime/ApplicationCore/Private/IOS/IOSConfigRules.cpp
+++ b/Engine/Source/Runtime/ApplicationCore/Private/IOS/IOSConfigRules.cpp
@@ -402,10 +402,11 @@ namespace ConfigRules
 			FindEnd
 		};
 
-		FAnsiStringView ConfigRulesAnsi((const char*)ConfigRulesData.GetData(), ConfigRulesData.Num());
 		// TODO: we should be able to parse this entire file without conversion to tchar.
 		// Some functions such as ParseLines, LexFromString etc. expect tchar.
-		FString ConfigRules = StringCast<TCHAR>(ConfigRulesAnsi.GetData()).Get();
+		auto ConvertedString = StringCast<TCHAR>((const char*)ConfigRulesData.GetData(), ConfigRulesData.Num());
+		FStringView ConfigRules = FStringView(ConvertedString.Get(), ConvertedString.Length());
+
 		TMap<FString, FString> ConfigRuleVars;
 		ConfigRuleVars.Append(MoveTemp(PredefinedVariables));
 		TArray<EConfRuleState> StateStack;
