Fix mouse high precision mode been turned on during PIE

#rnx

diff --git a/Engine/Source/Runtime/ApplicationCore/Private/Mac/MacCursor.cpp b/Engine/Source/Runtime/ApplicationCore/Private/Mac/MacCursor.cpp
--- a/Engine/Source/Runtime/ApplicationCore/Private/Mac/MacCursor.cpp
+++ b/Engine/Source/Runtime/ApplicationCore/Private/Mac/MacCursor.cpp
@@ -386,7 +386,10 @@ void FMacCursor::Lock(const RECT* const Bounds)
 
 	MacApplication->OnCursorLock();
 
-	bIsPositionInitialised = false; // Force GetIntPosition() to update its cached position in case the cursor was warped by another app while we were in high precision mode
+	if (!bUseHighPrecisionMode)
+	{
+		bIsPositionInitialised = false; // After exiting HighPrecisionMode, force GetIntPosition() to update its cached position in case the cursor was warped by another app while we were in high precision mode
+	}
 	FIntVector2 Position = GetIntPosition();
 	if (UpdateCursorClipping(Position))
 	{
