TouchpadSensitivity Improvements
-Discovered the source of the slowness in the function call when changing the setting, the ::SystemParametersInfo(SPI_SETTOUCHPADPARAMETERS call with the SPIF_SENDCHANGE argument sends WM_SETTINGSCHANGE to all windows on the system and waits for them to handle it.  We do not need to wait.  In fact in this particular case we do not want to send the message at all, its not a very important setting and we are going to set it back before other apps would get touchpad input.  Removing that means that the function returns very quickly and we don't need any async tasks (note: the error exit case still potentially happens on a non-game thread, so we do still use one atomic).
-Also we can handle the WM_SETTINGSCHANGE message sent from the system UI or other apps that change the setting and update our restore value.  This is a more reliable way to update our restore value if the user changes it.

diff --git a/Engine/Source/Runtime/ApplicationCore/Private/Windows/WindowsApplication.cpp b/Engine/Source/Runtime/ApplicationCore/Private/Windows/WindowsApplication.cpp
--- a/Engine/Source/Runtime/ApplicationCore/Private/Windows/WindowsApplication.cpp
+++ b/Engine/Source/Runtime/ApplicationCore/Private/Windows/WindowsApplication.cpp
@@ -409,29 +409,8 @@ namespace UE::WindowsApplication::TouchpadSensitivity
 		}
 	}
 
-	// Async model for SetTouchpadParameters mode.
-	// The ::SystemParametersInfo(SPI_SETTOUCHPADPARAMETERS call can take up to about 1 second to return.  We don't want 1 second hitches on startup, tab-out, tab-in, or normal exit.  We are going to tolerate the hitch on error-out to avoid any kind of locking or waiting there.
-	// Because this call is so slow it is possible to tab in/out/in/out a number of times before even the first setting change completes.  In order to avoid having several seconds of flipping back and forth in such a situation we are going to use an eventual-consistency
-	// design.  We will ensure that the LAST state reached is the LAST state requested.  We will not necessarily do all the requested state transitions.
-	// Therefore setting changes will be requested in two stages. 
-	// 1) Set the atomic ForceMaxTouchpadSensitivity_DesiredSensitivity to the desired sensitivity (either MostSensitive, or the value we want to restore on tab-out/exit).
-	// 2) Enqueue a task to set the sensitivity to the DesiredSensitivity.
-	// The result of this is that if we do multiple requests we will change the atomic DesiredSensitivity multiple times and queue up multiple tasks, but the next task to execute will set the system to the most recently set state.
-	// As long as no additional requests are made the remaining queued tasks will complete very quickly because the current sensitivity setting is already the DesiredSensitivity.  Any time we change the sensitivity we know that at least
-	// one task is in the pipe to execute that change.
-	// 
-	// Flaws:
-	// 1) In a crash we cannot restore the setting from this process.  There is a mechanism to write the saved sensitivity to the registery and restore it from there either on the next app run, or in some external process.  However this could potentially overwrite the user's setting
-	// if they change it to most sensitive.  Still that seems relatively unlikely and the user aparrantly at least knows about the setting.
-	// 2) If, while tabbed out, the setting were to change and multiple tasks were in the pipe I think it is possible we would restore a no-longer current saved state.  This probably can't be done manually, but if one tabbed between two applications using this system I expect they could
-	// end up leaving the setting at MostSensitive.
-
-	/** Async task pipe for processing the incoming profiler service data */
-	static UE::Tasks::FPipe AsyncTaskPipe{ TEXT("ForceMaxTouchpadSensitivityPipe") };
-	static std::atomic<TOUCHPAD_SENSITIVITY_LEVEL> ForceMaxTouchpadSensitivity_DesiredSensitivity = TOUCHPAD_SENSITIVITY_LEVEL_NOT_SET;
 	static std::atomic<TOUCHPAD_SENSITIVITY_LEVEL> ForceMaxTouchpadSensitivity_SavedSensitivity = TOUCHPAD_SENSITIVITY_LEVEL_NOT_SET;
 	static bool bRegistrySavedValueConsumed = false;
-	static std::atomic<bool> bAbortAsyncTasks = false;
 
 	FWindowsApplication::EForceMaxTouchpadSensitivityRestorationValues GetForceMaxTouchpadSensitivityRestorationValue()
 	{
@@ -478,125 +457,38 @@ namespace UE::WindowsApplication::TouchpadSensitivity
 		}
 	}
 
-
-	void EnqueueAsyncTouchpadSensitivityChangeTask()
+	void OnWMSettingChangeSetTouchpadParameters()
 	{
-		UE_LOG(LogWindowsTouchpad, Verbose, TEXT("TouchpadSensitivityChangeTask Task launching from pipe."));
-		AsyncTaskPipe.Launch(
-			TEXT("TouchpadSensitivityChangeTask"),
-			[]
-			{
-				UE_LOG(LogWindowsTouchpad, Verbose, TEXT("TouchpadSensitivityChangeTask Task running."));
-
-				if (bAbortAsyncTasks.load())
-				{
-					UE_LOG(LogWindowsTouchpad, Verbose, TEXT("TouchpadSensitivityChangeTask aborted, doing nothing.  Task Complete."));
-					return;
-				}
-
-				// Get the user's touchpad settings.
-				TOUCHPAD_PARAMETERS_V0 params = {};
-				params.versionNumber = 0;
-				{
-					BOOL bSuccess = ::SystemParametersInfo(SPI_GETTOUCHPADPARAMETERS, sizeof(params), &params, 0 /*fWinIni*/);
-					if (!bSuccess)
-					{
-						UE_LOG(LogWindowsTouchpad, Warning, TEXT("TouchpadSensitivityChangeTask SystemParametersInfo SPI_GETTOUCHPADPARAMETERS failed. Task Complete."));
-						return;
-					}
-				}
-
-				TOUCHPAD_SENSITIVITY_LEVEL DesiredSensitivity = ForceMaxTouchpadSensitivity_DesiredSensitivity.load();
-
-				if (DesiredSensitivity == TOUCHPAD_SENSITIVITY_LEVEL_NOT_SET)
-				{
-					UE_LOG(LogWindowsTouchpad, Log, TEXT("TouchpadSensitivityChangeTask DesiredSensitivity is TOUCHPAD_SENSITIVITY_LEVEL_NOT_SET.  This could happen if we tried to restore the setting, but had never changed it.  Setting not changed. Task Complete."));
-					return;
-				}
-
-				// If we are forcing to max we need to figure out our restore value and save it.
-				if (DesiredSensitivity == TOUCHPAD_SENSITIVITY_LEVEL_MOST_SENSITIVE)
-				{
-					// If we previously exited without restoring (crashed, killed in debugger, etc) the sensitivity the system's current sensitivity will already be MOST, but we may
-					// have saved some other value into the registry so that we can restore it.
-					// We only want to read that registry value once per run of the app.
-					// So first we will mark that we did that, then we might apply it.
-					if (bRestoreTouchpadSensitivityFromRegistry && !bRegistrySavedValueConsumed)
-					{
-						// We only want to attempt to read the registry value once.
-						bRegistrySavedValueConsumed = true;
-
-						// Only if our current sensitivity is already MOST we suspect we might have crashed and want to use the registry value as our saved sensitivity.
-						if (params.sensitivityLevel == DesiredSensitivity)
-						{
-							TOUCHPAD_SENSITIVITY_LEVEL RegistrySensitivity = TOUCHPAD_SENSITIVITY_LEVEL_MOST_SENSITIVE;
-							bool bSuccess = false;
-							ReadSavedSensitivityFromRegistry(RegistrySensitivity, bSuccess);
-							if (bSuccess && RegistrySensitivity != TOUCHPAD_SENSITIVITY_LEVEL_MOST_SENSITIVE)
-							{
-								UE_LOG(LogWindowsTouchpad, Log, TEXT("TouchpadSensitivityChangeTask restoring saved sensitivity %i from registry, we are already in MostSensitive and do not need to change the setting now. Task Complete."), (int32)RegistrySensitivity);
-								ForceMaxTouchpadSensitivity_SavedSensitivity.store(RegistrySensitivity);
-								// The sensitivity setting is already MOST, so we are done now.
-								return;
-							}
-						}
-						else
-						{
-							// We are not already in MOST, so we can ignore the registry saved value and restore to the current value.
-							// We already marked the registry value consumed, so we can proceed as normal.
-						}
-					}
-				}
-					
-				if (params.sensitivityLevel == DesiredSensitivity)
-				{
-					UE_LOG(LogWindowsTouchpad, Log, TEXT("TouchpadSensitivityChangeTask doing nothing because the current sensitivity already matches DesiredSensitivity %i. Task Complete."), DesiredSensitivity);
-					return;
-				}
-
-				if (DesiredSensitivity == TOUCHPAD_SENSITIVITY_LEVEL_MOST_SENSITIVE)
-				{
-					// We are doing the force, save the current value so we can restore it later.
-					UE_LOG(LogWindowsTouchpad, Log, TEXT("TouchpadSensitivityChangeTask setting saved sensitivity %i."), (int32)params.sensitivityLevel);
-					ForceMaxTouchpadSensitivity_SavedSensitivity.store(params.sensitivityLevel);
-					WriteSavedSensitivityToRegistry(params.sensitivityLevel);
-				}
-
-				// And now we will update the sensitivity to the desired value.
-				params.sensitivityLevel = DesiredSensitivity;
-
-				// Last chance to avoid the long system call where we change the setting.
-				if (bAbortAsyncTasks.load())
-				{
-					UE_LOG(LogWindowsTouchpad, Log, TEXT("TouchpadSensitivityChangeTask aborted, doing nothing.  Task Complete."));
-					return;
-				}
-
-				// Do not specify SPIF_UPDATEINIFILE so that the registry is not updated.
-				// We will revert this update later, and not writing to the registry is a
-				// failsafe in case we fail to do so (e.g. due to a crash).
-				bool bSuccess = ::SystemParametersInfo(SPI_SETTOUCHPADPARAMETERS, sizeof(params), const_cast<TOUCHPAD_PARAMETERS_V0*>(&params), SPIF_SENDCHANGE);
-				if (!bSuccess)
-				{
-					UE_LOG(LogWindowsTouchpad, Warning, TEXT("TouchpadSensitivityChangeTask SystemParametersInfo SPI_SETTOUCHPADPARAMETERS Failed.  Task Complete."));
-					return;
-				}
-				else
-				{
-					UE_LOG(LogWindowsTouchpad, Log, TEXT("TouchpadSensitivityChangeTask sucessfully changed setting to %i."), DesiredSensitivity);
+		// Note: we do not use GetForceMaxTouchpadSensitivityMode() here because in this case it is not an error if the mode has not yet been cached.
+		if (bForceMaxTouchpadSensitivity && CachedForceMaxTouchpadSensitivityMode == EForceMaxTouchpadSensitivityMode::SetTouchpadParameters)
+		{
+			// We got the windows message telling us that the touchpad settings have changed
+			// We don't know which settings have changed, but we can go ahead and use whatever the value is right now
+			// as our new restore value, and also ensure that our registry saved value matches.
+			// This relies on the fact that we never send the settings change message, so we are not talking to ourselves here.
+			UE_LOG(LogWindowsTouchpad, Verbose, TEXT("OnWMSettingChangeSetTouchpadParameters touchpad parameters have changed."));
 
-					// If we have restored a value we do not want a value stored in the registry anymore. (if the user changes their setting now we do not want to overwrite it later)
-					if (DesiredSensitivity != TOUCHPAD_SENSITIVITY_LEVEL_MOST_SENSITIVE)
-					{
-						UE_LOG(LogWindowsTouchpad, Log, TEXT("TouchpadSensitivityChangeTask ClearSavedSensitivityFromRegistry."), DesiredSensitivity);
-						ClearSavedSensitivityFromRegistry();
-					}
+			// Get the user's touchpad settings.
+			TOUCHPAD_PARAMETERS_V0 params = {};
+			params.versionNumber = 0;
+			BOOL bSuccess = ::SystemParametersInfo(SPI_GETTOUCHPADPARAMETERS, sizeof(params), &params, 0);
+			if (!bSuccess)
+			{
+				UE_LOG(LogWindowsTouchpad, Warning, TEXT("OnWMSettingChangeSetTouchpadParameters SystemParametersInfo SPI_GETTOUCHPADPARAMETERS failed"));
+				return;
+			}
 
-					UE_LOG(LogWindowsTouchpad, Verbose, TEXT("TouchpadSensitivityChangeTask Task Complete."), DesiredSensitivity);
-					return;
-				}
-			},
-			UE::Tasks::ETaskPriority::Normal);
+			UE_LOG(LogWindowsTouchpad, Log, TEXT("OnWMSettingChangeSetTouchpadParameters changing saved senstivity to %i"), params.sensitivityLevel);
+			ForceMaxTouchpadSensitivity_SavedSensitivity.store(params.sensitivityLevel);
+			if (params.sensitivityLevel == TOUCHPAD_SENSITIVITY_LEVEL_MOST_SENSITIVE)
+			{
+				ClearSavedSensitivityFromRegistry();
+			}
+			else
+			{
+				WriteSavedSensitivityToRegistry(params.sensitivityLevel);
+			}
+		}
 	}
 
 	void ApplyForceMaxTouchpadSensitivity()
@@ -615,8 +507,83 @@ namespace UE::WindowsApplication::TouchpadSensitivity
 			break;
 		case EForceMaxTouchpadSensitivityMode::SetTouchpadParameters:
 		{
-			ForceMaxTouchpadSensitivity_DesiredSensitivity.store(TOUCHPAD_SENSITIVITY_LEVEL_MOST_SENSITIVE);
-			EnqueueAsyncTouchpadSensitivityChangeTask();
+			// Note: we will get the params, change just the one, then write them back so that the reserved param block is preserved.
+			UE_LOG(LogWindowsTouchpad, Verbose, TEXT("ApplyForceMaxTouchpadSensitivity setting sensitivity to MOST_SENSITIVE"));
+
+			// Get the user's touchpad settings.
+			TOUCHPAD_PARAMETERS_V0 params = {};
+			params.versionNumber = 0;
+			BOOL bSuccess = ::SystemParametersInfo(SPI_GETTOUCHPADPARAMETERS, sizeof(params), &params, 0);
+			if (!bSuccess)
+			{
+				UE_LOG(LogWindowsTouchpad, Warning, TEXT("ApplyForceMaxTouchpadSensitivity SystemParametersInfo SPI_GETTOUCHPADPARAMETERS failed"));
+				return;
+			}
+
+			// If we previously exited without restoring (crashed, killed in debugger, etc) the sensitivity the system's current sensitivity will already be MOST, but we may
+			// have saved some other value into the registry so that we can restore it.
+			// We only want to read that registry value once per run of the app.
+			// So first we will mark that we did that, then we might apply it.
+			if (bRestoreTouchpadSensitivityFromRegistry && !bRegistrySavedValueConsumed)
+			{
+				// We only want to attempt to read the registry value once.
+				bRegistrySavedValueConsumed = true;
+
+				// Only if our current sensitivity is already MOST we suspect we might have crashed and want to use the registry value as our saved sensitivity.
+				if (params.sensitivityLevel == TOUCHPAD_SENSITIVITY_LEVEL_MOST_SENSITIVE)
+				{
+					TOUCHPAD_SENSITIVITY_LEVEL RegistrySensitivity = TOUCHPAD_SENSITIVITY_LEVEL_MOST_SENSITIVE;
+					bool bReadSuccess = false;
+					ReadSavedSensitivityFromRegistry(RegistrySensitivity, bReadSuccess);
+					if (bReadSuccess && RegistrySensitivity != TOUCHPAD_SENSITIVITY_LEVEL_MOST_SENSITIVE)
+					{
+						UE_LOG(LogWindowsTouchpad, Log, TEXT("ApplyForceMaxTouchpadSensitivity restoring saved sensitivity %i from registry, we are already in MostSensitive and do not need to change the setting now."), (int32)RegistrySensitivity);
+						ForceMaxTouchpadSensitivity_SavedSensitivity.store(RegistrySensitivity);
+						// The sensitivity setting is already MOST, so we are done now.
+						return;
+					}
+				}
+				else
+				{
+					// We are not already in MOST, so we can ignore the registry saved value and restore to the current value.
+					// We already marked the registry value consumed, so we can proceed as normal.
+				}
+			}
+
+			if (params.sensitivityLevel == TOUCHPAD_SENSITIVITY_LEVEL_MOST_SENSITIVE)
+			{
+				UE_LOG(LogWindowsTouchpad, Log, TEXT("ApplyForceMaxTouchpadSensitivity found that sensitivity is already MOST_SENSITIVE. Doing nothing."));
+				return;
+			}
+
+			// When we first try to change the sensitivity we will store the value we want to restore to.
+			if (ForceMaxTouchpadSensitivity_SavedSensitivity.load() == TOUCHPAD_SENSITIVITY_LEVEL_NOT_SET)
+			{
+				UE_LOG(LogWindowsTouchpad, Verbose, TEXT("ApplyForceMaxTouchpadSensitivity setting ForceMaxTouchpadSensitivity_SavedSensitivity to %i"), params.sensitivityLevel);
+				ForceMaxTouchpadSensitivity_SavedSensitivity.store(params.sensitivityLevel);
+				WriteSavedSensitivityToRegistry(params.sensitivityLevel);
+			}
+
+			// Change the sensitivity level.
+			// Leave all other bytes in the struct untouched.
+			params.sensitivityLevel = TOUCHPAD_SENSITIVITY_LEVEL_MOST_SENSITIVE;
+
+			// Do not specify SPIF_UPDATEINIFILE so that the registry is not updated.
+			// We will revert this update later, and not writing to the registry is a
+			// failsafe in case we fail to do so (e.g. due to a crash).
+			// Do not specify SPIF_SENDCHANGE because the message will be sent to every window and then this
+			// function will block until the other windows handle it.
+			bSuccess = ::SystemParametersInfo(SPI_SETTOUCHPADPARAMETERS, sizeof(params), &params, 0);
+			if (!bSuccess)
+			{
+				UE_LOG(LogWindowsTouchpad, Warning, TEXT("ApplyForceMaxTouchpadSensitivity SystemParametersInfo SPI_SETTOUCHPADPARAMETERS failed"));
+				return;
+			}
+
+			// If we want to send the WM_SETTINGCHANGE message without blocking this should do it:
+			//SendNotifyMessage(HWND_BROADCAST, WM_SETTINGCHANGE, SPI_SETTOUCHPADPARAMETERS, 0);
+
+			UE_LOG(LogWindowsTouchpad, Log, TEXT("ApplyForceMaxTouchpadSensitivity sucessfully set sensitivity to MOST_SENSITIVE"));
 		}
 		break;
 		case EForceMaxTouchpadSensitivityMode::SetMaxTouchpadSensitivity:
@@ -625,6 +592,10 @@ namespace UE::WindowsApplication::TouchpadSensitivity
 			{
 				UE_LOG(LogWindowsTouchpad, Warning, TEXT("ApplyForceMaxTouchpadSensitivity SetMaxTouchpadSensitivity Failed"));
 			}
+			else
+			{
+				UE_LOG(LogWindowsTouchpad, Log, TEXT("ApplyForceMaxTouchpadSensitivity SetMaxTouchpadSensitivity set to true"));
+			}
 		}
 		break;
 		default:
@@ -649,8 +620,51 @@ namespace UE::WindowsApplication::TouchpadSensitivity
 			break;
 		case EForceMaxTouchpadSensitivityMode::SetTouchpadParameters:
 		{
-			ForceMaxTouchpadSensitivity_DesiredSensitivity.store(ForceMaxTouchpadSensitivity_SavedSensitivity.load());
-			EnqueueAsyncTouchpadSensitivityChangeTask();
+			TOUCHPAD_SENSITIVITY_LEVEL LocalSavedSensitivity = ForceMaxTouchpadSensitivity_SavedSensitivity.load();
+
+			if (LocalSavedSensitivity <= TOUCHPAD_SENSITIVITY_LEVEL_MOST_SENSITIVE || LocalSavedSensitivity > TOUCHPAD_SENSITIVITY_LEVEL_LEAST_SENSITIVE)
+			{
+				// No restore needed (saved value is MOST_SENSITIVE), NOT_SET, or somehow an invalid value.
+				return;
+			}
+
+			// Note: we will get the params, change just the one, then write them back so that the reserved param block is preserved.
+			UE_LOG(LogWindowsTouchpad, Verbose, TEXT("RemoveForceMaxTouchpadSensitivity reverting to SavedSensitivity %i"), LocalSavedSensitivity);
+
+			// Get the user's touchpad settings.
+			TOUCHPAD_PARAMETERS_V0 params = {};
+			params.versionNumber = 0;
+			BOOL bSuccess = ::SystemParametersInfo(SPI_GETTOUCHPADPARAMETERS, sizeof(params), &params, 0);
+			if (!bSuccess)
+			{
+				UE_LOG(LogWindowsTouchpad, Warning, TEXT("RemoveForceMaxTouchpadSensitivity SystemParametersInfo SPI_GETTOUCHPADPARAMETERS failed"));
+				return;
+			}
+
+			// Restore the saved sensitivity level.
+			// Leave all other bytes in the struct untouched.
+			params.sensitivityLevel = LocalSavedSensitivity;
+
+			// Do not specify SPIF_UPDATEINIFILE so that the registry is not updated.
+			// We will revert this update later, and not writing to the registry is a
+			// failsafe in case we fail to do so (e.g. due to a crash).
+			// Do not specify SPIF_SENDCHANGE because the message will be sent to every window and then this
+			// function will block until the other windows handle it.
+			bSuccess = ::SystemParametersInfo(SPI_SETTOUCHPADPARAMETERS, sizeof(params), &params, 0);
+			if (!bSuccess)
+			{
+				UE_LOG(LogWindowsTouchpad, Warning, TEXT("RemoveForceMaxTouchpadSensitivity SystemParametersInfo SPI_SETTOUCHPADPARAMETERS failed"));
+				return;
+			}
+
+			// If we want to send the WM_SETTINGCHANGE message without blocking this should do it:
+			//SendNotifyMessage(HWND_BROADCAST, WM_SETTINGCHANGE, SPI_SETTOUCHPADPARAMETERS, 0);
+
+			// We have removed our sensitivity change, so we can clear the registry saved value.
+			ClearSavedSensitivityFromRegistry();
+
+			UE_LOG(LogWindowsTouchpad, Log, TEXT("RemoveForceMaxTouchpadSensitivity sucessfully reverted to SavedSensitivity %i"), LocalSavedSensitivity);
+
 		}
 		break;
 		case EForceMaxTouchpadSensitivityMode::SetMaxTouchpadSensitivity:
@@ -659,6 +673,10 @@ namespace UE::WindowsApplication::TouchpadSensitivity
 			{
 				UE_LOG(LogWindowsTouchpad, Warning, TEXT("RemoveForceMaxTouchpadSensitivity SetMaxTouchpadSensitivity Failed"));
 			}
+			else
+			{
+				UE_LOG(LogWindowsTouchpad, Log, TEXT("ApplyForceMaxTouchpadSensitivity SetMaxTouchpadSensitivity set to false"));
+			}		
 		}
 		break;
 		default:
@@ -677,9 +695,6 @@ namespace UE::WindowsApplication::TouchpadSensitivity
 		{
 			UE_LOG(LogWindowsTouchpad, Verbose, TEXT("RemoveForceMaxTouchpadSensitivityOnErrorShutdown Starting."));
 
-			// New Async tasks that start should do nothing.
-			bAbortAsyncTasks.store(true);
-
 			TOUCHPAD_SENSITIVITY_LEVEL LocalSavedSensitivity = ForceMaxTouchpadSensitivity_SavedSensitivity.load();
 
 			if (LocalSavedSensitivity <= TOUCHPAD_SENSITIVITY_LEVEL_MOST_SENSITIVE || LocalSavedSensitivity > TOUCHPAD_SENSITIVITY_LEVEL_LEAST_SENSITIVE)
@@ -690,9 +705,6 @@ namespace UE::WindowsApplication::TouchpadSensitivity
 
 			// At this point we have a saved sensitivity value that it makes sense to restore.
 
-			// Set the atomic desired sensitivity.  If there is a task currently executing this might catch it and have it do the right thing.
-			ForceMaxTouchpadSensitivity_DesiredSensitivity.store(LocalSavedSensitivity);
-
 			// Note: we will get the params, change just the one, then write them back so that the reserved param block is preserved.
 			UE_LOG(LogWindowsTouchpad, Verbose, TEXT("RemoveForceMaxTouchpadSensitivityOnErrorShutdown reverting to SavedSensitivity %i"), LocalSavedSensitivity);
 
@@ -713,7 +725,7 @@ namespace UE::WindowsApplication::TouchpadSensitivity
 			// Do not specify SPIF_UPDATEINIFILE so that the registry is not updated.
 			// We will revert this update later, and not writing to the registry is a
 			// failsafe in case we fail to do so (e.g. due to a crash).
-			bSuccess = ::SystemParametersInfo(SPI_SETTOUCHPADPARAMETERS, sizeof(params), &params, SPIF_SENDCHANGE);
+			bSuccess = ::SystemParametersInfo(SPI_SETTOUCHPADPARAMETERS, sizeof(params), &params, 0);
 			if (!bSuccess)
 			{
 				UE_LOG(LogWindowsTouchpad, Warning, TEXT("RemoveForceMaxTouchpadSensitivityOnErrorShutdown SystemParametersInfo SPI_SETTOUCHPADPARAMETERS failed"));
@@ -2455,6 +2467,10 @@ int32 FWindowsApplication::ProcessMessage( HWND hwnd, uint32 msg, WPARAM wParam,
 				{
 					DeferMessage(CurrentNativeEventWindowPtr, hwnd, msg, wParam, lParam);
 				}
+				if (wParam == UE::WindowsApplication::TouchpadSensitivity::SPI_SETTOUCHPADPARAMETERS)
+				{
+					DeferMessage(CurrentNativeEventWindowPtr, hwnd, msg, wParam, lParam);
+				}
 			}
 			break;
 
@@ -3249,6 +3265,11 @@ int32 FWindowsApplication::ProcessDeferredMessage( const FDeferredWindowsMessage
 			{
 				MessageHandler->OnConvertibleLaptopModeChanged();
 			}
+			if (wParam == UE::WindowsApplication::TouchpadSensitivity::SPI_SETTOUCHPADPARAMETERS)
+			{
+				UE::WindowsApplication::TouchpadSensitivity::OnWMSettingChangeSetTouchpadParameters();
+			}
+
 			break;
 	
 		case WM_NCACTIVATE:
