[ApplicationCore]
* Fixed so console output device 2 is not doing TerminateProcess when a proper shutdown is about to end (static deinitialization is happening)

diff --git a/Engine/Source/Runtime/ApplicationCore/Private/Windows/WindowsConsoleOutputDevice2.cpp b/Engine/Source/Runtime/ApplicationCore/Private/Windows/WindowsConsoleOutputDevice2.cpp
--- a/Engine/Source/Runtime/ApplicationCore/Private/Windows/WindowsConsoleOutputDevice2.cpp
+++ b/Engine/Source/Runtime/ApplicationCore/Private/Windows/WindowsConsoleOutputDevice2.cpp
@@ -101,8 +101,9 @@ public:
 
 	~FConsoleWindow()
 	{
+		bExiting = true;
 		PostMessageW(MainHwnd, WM_CLOSE, 0, 0);
-		WaitForSingleObject(Thread, 1000);
+		WaitForSingleObject(Thread, 3000);
 		CloseHandle(Thread);
 		Thread = nullptr;
 
@@ -348,7 +349,7 @@ public:
 
 			FWindowsPlatformMisc::CallGracefulTerminationHandler();
 		}
-		else
+		else if (!bExiting)
 		{
 			// This means that we've pressed close after the process has already been exited. It is likely the shutdown process is hung and we should just terminate the process
 			Sleep(500);
@@ -2548,6 +2549,7 @@ public:
 	int32 CheckpointIndex = 0;
 
 	bool bIsVisible = false;
+	bool bExiting = false;
 	int32 ConsoleWidth = 160;
 	int32 ConsoleHeight = 4000;
 	int32 ConsolePosX = 0;
