Added PlayerController::SetControllerGyroAutoCalibration so gyroscope calibration can be disabled and enabled manually on some platforms

#tests PIE with gyro controllers

diff --git a/Engine/Source/Runtime/ApplicationCore/Public/GenericPlatform/IInputInterface.h b/Engine/Source/Runtime/ApplicationCore/Public/GenericPlatform/IInputInterface.h
--- a/Engine/Source/Runtime/ApplicationCore/Public/GenericPlatform/IInputInterface.h
+++ b/Engine/Source/Runtime/ApplicationCore/Public/GenericPlatform/IInputInterface.h
@@ -125,6 +125,24 @@ struct FInputDeviceProperty
 	FName Name;
 };
 
+/** A generic gyro calibration property for input devices whose gyroscopes (angular velocity sensors) can have their calibration enabled or disabled */
+struct FInputDeviceGyroAutoCalibrationProperty : public FInputDeviceProperty
+{
+	FInputDeviceGyroAutoCalibrationProperty()
+		: FInputDeviceProperty(PropertyName())
+	{}
+
+	FInputDeviceGyroAutoCalibrationProperty(bool bInEnable)
+		: FInputDeviceProperty(PropertyName())
+		, bEnable(bInEnable)
+	{}
+
+	static FName PropertyName() { return FName("InputDeviceGyroAutoCalibration"); }
+
+	/** If gyro auto calibration should be enabled. */
+	bool bEnable = true;
+};
+
 /** A generic light color property for input devices that have lights on them */
 struct FInputDeviceLightColorProperty : public FInputDeviceProperty
 {
diff --git a/Engine/Source/Runtime/Engine/Classes/GameFramework/PlayerController.h b/Engine/Source/Runtime/Engine/Classes/GameFramework/PlayerController.h
--- a/Engine/Source/Runtime/Engine/Classes/GameFramework/PlayerController.h
+++ b/Engine/Source/Runtime/Engine/Classes/GameFramework/PlayerController.h
@@ -1379,6 +1379,13 @@ public:
 	UFUNCTION(BlueprintCallable, Category = "Input")
 	ENGINE_API void ResetControllerDeadZones();
 
+	/**
+	 * Sets whether the player's controller's gyro auto calibration is enabled
+	 * @param	bEnabled	 	Whether gyro auto calibration should be enabled
+	 */
+	UFUNCTION(BlueprintCallable, Category = "Input")
+	ENGINE_API void SetControllerGyroAutoCalibration(bool bEnabled);
+
 	/**
 	 * Travel to a different map or IP address. Calls the PreClientTravel event before doing anything.
 	 * NOTE: This is implemented as a locally executed wrapper for ClientTravelInternal, to avoid API compatability breakage
diff --git a/Engine/Source/Runtime/Engine/Private/PlayerController.cpp b/Engine/Source/Runtime/Engine/Private/PlayerController.cpp
--- a/Engine/Source/Runtime/Engine/Private/PlayerController.cpp
+++ b/Engine/Source/Runtime/Engine/Private/PlayerController.cpp
@@ -4577,6 +4577,25 @@ void APlayerController::ResetControllerDeadZones()
 	}
 }
 
+void APlayerController::SetControllerGyroAutoCalibration(bool bEnabled)
+{
+	if (Player == nullptr)
+	{
+		return;
+	}
+
+	if (FSlateApplication::IsInitialized())
+	{
+		IInputInterface* InputInterface = FSlateApplication::Get().GetInputInterface();
+		if (InputInterface)
+		{
+			const int32 ControllerId = CastChecked<ULocalPlayer>(Player)->GetControllerId();
+			FInputDeviceGyroAutoCalibrationProperty GyroCalibration(bEnabled);
+			InputInterface->SetDeviceProperty(ControllerId, &GyroCalibration);
+		}
+	}
+}
+
 void APlayerController::ProcessForceFeedbackAndHaptics(const float DeltaTime, const bool bGamePaused)
 {
 	if (Player == nullptr)
