[ApplicationCore]
* Fixed so new console handles dpi scaling when unreal application change dpi scale handling

diff --git a/Engine/Source/Runtime/ApplicationCore/Private/Windows/WindowsConsoleOutputDevice2.cpp b/Engine/Source/Runtime/ApplicationCore/Private/Windows/WindowsConsoleOutputDevice2.cpp
--- a/Engine/Source/Runtime/ApplicationCore/Private/Windows/WindowsConsoleOutputDevice2.cpp
+++ b/Engine/Source/Runtime/ApplicationCore/Private/Windows/WindowsConsoleOutputDevice2.cpp
@@ -249,7 +249,7 @@ public:
 
 		RECT Rect;
 		GetClientRect(MainHwnd, &Rect);
-		UpdateSize(Rect.right, Rect.bottom, false);
+		UpdateSize(FromScale(Rect.right), FromScale(Rect.bottom), false);
 
 		UpdateTime(0);
 
@@ -277,7 +277,6 @@ public:
 
 		SetForegroundWindow(MainHwnd);
 
-
 		TipHwnd = CreateWindowEx(WS_EX_TOOLWINDOW | WS_EX_TOPMOST | WS_EX_NOACTIVATE, WC_STATIC, TEXT(""), SS_OWNERDRAW | WS_VISIBLE | WS_POPUP | WS_DISABLED | WS_CLIPCHILDREN | WS_CLIPSIBLINGS, 1, 1, 1, 1, MainHwnd, NULL, HInstance, NULL);
 		SendMessageW(TipHwnd, WM_SETFONT, (WPARAM)Font, 0);
 
@@ -291,10 +290,12 @@ public:
 			{
 				if (bIsVisible)
 				{
-					MoveWindow(MainHwnd, ConsolePosX, ConsolePosY, ConsoleWidth, ConsoleHeight, true);
+					SetWindowPos(MainHwnd, 0, ConsolePosX, ConsolePosY, ConsoleWidth, ConsoleHeight, SWP_SHOWWINDOW|SWP_NOZORDER|SWP_NOACTIVATE|SWP_NOOWNERZORDER|SWP_FRAMECHANGED);
+				}
+				else
+				{
+					ShowWindow(MainHwnd, SW_HIDE);
 				}
-
-				ShowWindow(MainHwnd, bIsVisible ? SW_SHOW : SW_HIDE);
 				bWasVisible = bIsVisible;
 
 				if (bIsVisible)
@@ -313,6 +314,65 @@ public:
 
 			DWORD Timeout = 200;
 			DWORD WaitResult = MsgWaitForMultipleObjects(1, &DirtyEvent, false, Timeout, QS_ALLINPUT);
+
+			{
+				typedef UINT(WINAPI* LPGetDpiForWindow)(HWND Hwnd);
+				static LPGetDpiForWindow GetDpiForWindow = (LPGetDpiForWindow)(void*)GetProcAddress(GetModuleHandle(L"user32.dll"), "GetDpiForWindow");
+				if (GetDpiForWindow != nullptr)
+				{
+					uint32 dpi = GetDpiForWindow(MainHwnd);
+					static uint32 lastDpi = dpi;
+					if (lastDpi != dpi)
+					{
+						DpiScale = double(dpi) / 96;
+						ShowWindow(MainHwnd, SW_HIDE);
+						lastDpi = dpi;
+						bWasVisible = false;
+
+						ConsolePosX = ToScale(ConsolePosX);
+						ConsolePosY = ToScale(ConsolePosY);
+						ConsoleWidth = ToScale(ConsoleWidth);
+						ConsoleHeight = ToScale(ConsoleHeight);
+
+						SystemParametersInfo(SPI_GETNONCLIENTMETRICS, sizeof(NonClientMetrics), &NonClientMetrics, 0);
+						Font = (HFONT)CreateFontIndirect(&NonClientMetrics.lfMessageFont);
+
+						FontHeight = -MulDiv(8, dpi, 72);
+						LogFont = (HFONT)CreateFontW(FontHeight, 0, 0, 0, FW_NORMAL, false, false, false, ANSI_CHARSET, OUT_DEFAULT_PRECIS, CLIP_DEFAULT_PRECIS, ANTIALIASED_QUALITY, FIXED_PITCH | FF_MODERN, FontName);
+
+						FontSize;
+						Hdc = GetDC(MainHwnd);
+						SelectObject(Hdc, LogFont);
+						GetTextExtentPoint32(Hdc, TEXT("A"), 1, &FontSize);
+						LogFontWidth = FontSize.cx;
+						LogFontHeight = FontSize.cy;
+						ReleaseDC(MainHwnd, Hdc);
+
+						auto UpdateFont = [](HWND hwnd, HFONT font) { SendMessageW(hwnd, WM_SETFONT, (WPARAM)font, 0); };
+						
+						UpdateFont(IncludeFilterHwnd, Font);
+						UpdateFont(ExcludeFilterHwnd, Font);
+						UpdateFont(ClearLogButtonHwnd, Font);
+						UpdateFont(CheckpointButtonHwnd, Font);
+						UpdateFont(CommandHwnd, Font);
+						UpdateFont(RunCommandHwnd, Font);
+
+						for (Activity& A : Activities)
+						{
+							if (A.NameHwnd)
+								UpdateFont(A.NameHwnd, Font);
+							if (A.StatusHwnd)
+								UpdateFont(A.StatusHwnd, Font);
+						}
+
+						UpdateFont(LogHwnd, LogFont);
+						SendMessageW(LogHwnd, LB_SETITEMHEIGHT, 0, MAKELPARAM(LogFontHeight, 0));
+
+						WindowBorderPen = CreatePen(PS_INSIDEFRAME, int(ceil(DpiScale)), BorderColor);
+					}
+				}
+			}
+
 			if (WaitResult == WAIT_TIMEOUT)
 				continue;
 
@@ -361,7 +421,7 @@ public:
 
 	void CreateColors()
 	{
-		COLORREF BorderColor = RGB(0,0,0);
+		BorderColor = RGB(0,0,0);
 		COLORREF LogBackgroundColor = RGB(0, 0, 0);
 		switch (ConsoleColor)
 		{
@@ -661,22 +721,22 @@ public:
 		int32 Y = 10;
 		int32 ButtonStart = ClientWidth - 235;
 
-		SetWindowPos(ClearLogButtonHwnd, 0, ButtonStart, Y - 2, 80, 22, FilterFlags);
-		SetWindowPos(CheckpointButtonHwnd, 0, ButtonStart + 94, Y - 2, 135, 22, FilterFlags);
+		SetWindowPosScale(ClearLogButtonHwnd, 0, ButtonStart, Y - 2, 80, 22, FilterFlags);
+		SetWindowPosScale(CheckpointButtonHwnd, 0, ButtonStart + 94, Y - 2, 135, 22, FilterFlags);
 
 		int32 WidthForFilters = ButtonStart - 45;
 		if (WidthForFilters < 100)
 			FilterFlags = Flags | SWP_HIDEWINDOW;
 		int32 X = 8;
 		int32 WidthForEditBoxes = WidthForFilters / 2;
-		SetWindowPos(IncludeFilterHwnd, 0, X, Y - 1, WidthForEditBoxes, 20, FilterFlags);
+		SetWindowPosScale(IncludeFilterHwnd, 0, X, Y - 1, WidthForEditBoxes, 20, FilterFlags);
 		X += WidthForEditBoxes + 20;
-		SetWindowPos(ExcludeFilterHwnd, 0, X, Y - 1, WidthForEditBoxes, 20, FilterFlags);
+		SetWindowPosScale(ExcludeFilterHwnd, 0, X, Y - 1, WidthForEditBoxes, 20, FilterFlags);
 
-		SetWindowPos(LogHwnd, 0, 7, LogTop, ClientWidth - 33, LogHeight, Flags);
-		SetWindowPos(LogScrollHwnd, 0, ClientWidth - 26, LogTop, 19, LogHeight, Flags);
-		SetWindowPos(CommandHwnd, 0, 8, CommandY, ClientWidth - ButtonWidth - 20, 20, CommandFlags);
-		SetWindowPos(RunCommandHwnd, 0, ClientWidth - ButtonWidth - 7, CommandY - 1, ButtonWidth, 22, CommandFlags);
+		SetWindowPosScale(LogHwnd, 0, 7, LogTop, ClientWidth - 33, LogHeight, Flags);
+		SetWindowPosScale(LogScrollHwnd, 0, ClientWidth - 26, LogTop, 19, LogHeight, Flags);
+		SetWindowPosScale(CommandHwnd, 0, 8, CommandY, ClientWidth - ButtonWidth - 20, 20, CommandFlags);
+		SetWindowPosScale(RunCommandHwnd, 0, ClientWidth - ButtonWidth - 7, CommandY - 1, ButtonWidth, 22, CommandFlags);
 
 		if (bRedraw)
 			SendMessageW(LogHwnd, WM_SETREDRAW, true, 0);
@@ -692,10 +752,10 @@ public:
 				X += 16;
 				Width -= 16;
 				int32 NameWidth = FMath::Min(90, Width);
-				SetWindowPos(A.NameHwnd, 0, X, Y, NameWidth, 18, Flags);
+				SetWindowPosScale(A.NameHwnd, 0, X, Y, NameWidth, 18, Flags);
 				SetWindowLong(A.NameHwnd, GWLP_USERDATA, RowIndex);
 				int32 StatusWidth = FMath::Max(Width - NameWidth, 0);
-				SetWindowPos(A.StatusHwnd, 0, X + NameWidth, Y, StatusWidth, 18, Flags);
+				SetWindowPosScale(A.StatusHwnd, 0, X + NameWidth, Y, StatusWidth, 18, Flags);
 				SetWindowLong(A.StatusHwnd, GWLP_USERDATA, RowIndex);
 			});
 	}
@@ -844,7 +904,7 @@ public:
 
 	void GetThumbPosAndHeight(int32& OutPos, int32& OutHeight, const RECT& ClientRect)
 	{
-		int32 ButtonHeight = ClientRect.right + 1;
+		int32 ButtonHeight = ClientRect.right + ToScale(1);
 
 		int32 TotalCount = (int)SendMessageW(LogHwnd, LB_GETCOUNT, 0, 0);
 		int32 PageSize = ClientRect.bottom / LogFontHeight;
@@ -873,7 +933,7 @@ public:
 
 	int32 GetTopVisible(int32 ThumbPos, const RECT& ClientRect)
 	{
-		int32 ButtonHeight = ClientRect.right + 1;
+		int32 ButtonHeight = ClientRect.right + ToScale(1);
 
 		int32 TotalCount = (int)SendMessageW(LogHwnd, LB_GETCOUNT, 0, 0);
 		int32 PageSize = ClientRect.bottom / LogFontHeight;
@@ -929,26 +989,30 @@ public:
 			HDC hdc = (HDC)wParam;
 			RECT Rect;
 			GetClientRect(hWnd, &Rect);
-			int32 ButtonHeight = Rect.right + 1;
+			int32 ButtonHeight = Rect.right + ToScale(1);
 
 			RECT ScrollRect{ 0, ButtonHeight, Rect.right, Rect.bottom - ButtonHeight };
-			FillRect(hdc, &ScrollRect, ScrollBackgroundBrush);
+			FillRect(hdc, ScrollRect, ScrollBackgroundBrush);
 
 			RECT ButtonRect{ 0, 0, Rect.right, ButtonHeight };
-			FillRect(hdc, &ButtonRect, ButtonBrush);
+			FillRect(hdc, ButtonRect, ButtonBrush);
 
 			RECT ButtonRect2{ 0, Rect.bottom - ButtonHeight, Rect.right, Rect.bottom };
-			FillRect(hdc, &ButtonRect2, ButtonBrush);
+			FillRect(hdc, ButtonRect2, ButtonBrush);
+
+			int One = ToScale(1);
+			int Two = ToScale(2);
+			int Three = ToScale(3);
 
 			SelectObject(hdc, NoPen);
 			SelectObject(hdc, TextBrush);
 			int32 X = Rect.right / 2;
 			int32 Y = Rect.right / 2;
-			POINT vertices[] = { {X - 3, Y + 1}, {X, Y - 2}, {X + 3, Y + 1} };
+			POINT vertices[] = { {X - Three, Y + One}, {X, Y - Two}, {X + Three, Y + One} };
 			Polygon(hdc, vertices, 3);
 
 			Y = Rect.bottom - Y;
-			POINT vertices2[] = { {X - 3, Y - 2}, {X, Y + 1}, {X + 3, Y - 2} };
+			POINT vertices2[] = { {X - Three, Y - Two}, {X, Y + One}, {X + Three, Y - Two} };
 			Polygon(hdc, vertices2, 3);
 
 			return 1;
@@ -957,13 +1021,13 @@ public:
 		{
 			RECT Rect;
 			GetClientRect(hWnd, &Rect);
-			int32 ButtonHeight = Rect.right + 1;
+			int32 ButtonHeight = Rect.right + ToScale(1);
 
 			PAINTSTRUCT ps;
 			BeginPaint(hWnd, &ps);
 
 			RECT ScrollRect{ 0, ButtonHeight, Rect.right, Rect.bottom - ButtonHeight };
-			FillRect(ps.hdc, &ScrollRect, ScrollBackgroundBrush);
+			FillRect(ps.hdc, ScrollRect, ScrollBackgroundBrush);
 
 			int ThumbPos;
 			int ThumbHeight;
@@ -972,7 +1036,7 @@ public:
 			if (ThumbHeight > 0)
 			{
 				RECT ThumbRect{ 1, ThumbPos, Rect.right-1, ThumbPos + ThumbHeight };
-				FillRect(ps.hdc, &ThumbRect, ThumbBrush);
+				FillRect(ps.hdc, ThumbRect, ThumbBrush);
 			}
 
 			EndPaint(hWnd, &ps);
@@ -984,7 +1048,7 @@ public:
 			GetClientRect(hWnd, &Rect);
 			int32 Y = HIWORD(lParam);
 
-			int32 ButtonHeight = Rect.right + 1;
+			int32 ButtonHeight = Rect.right + ToScale(1);
 			if (Y <= ButtonHeight)
 			{
 				MoveTopVisible(Rect, -1);
@@ -1049,9 +1113,12 @@ public:
 
 	LRESULT HitTestNCA(HWND hWnd, LPARAM lParam)
 	{
-		POINT MousePt = { ((int)(short)LOWORD(lParam)), ((int)(short)HIWORD(lParam))};
 		RECT WindowRect;
 		GetWindowRect(hWnd, &WindowRect);
+		WindowRect = FromScale(WindowRect);
+		int MouseX = FromScale((int)LOWORD(lParam));
+		int MouseY = FromScale((int)HIWORD(lParam));
+
 		RECT FrameRect = { 0 };
 		AdjustWindowRectEx(&FrameRect, WS_OVERLAPPEDWINDOW & ~WS_CAPTION, false, NULL);
 		int32 Row = 1;
@@ -1060,17 +1127,17 @@ public:
 
 		int BorderWidth = 2;
 
-		if (MousePt.y >= WindowRect.top && MousePt.y < WindowRect.top + 27)
+		if (MouseY >= WindowRect.top && MouseY < WindowRect.top + 27)
 		{
-			bOnResizeBorder = (MousePt.y < (WindowRect.top - FrameRect.top));
+			bOnResizeBorder = (MouseY < (WindowRect.top - FrameRect.top));
 			Row = 0;
 		}
-		else if (MousePt.y < WindowRect.bottom && MousePt.y >= WindowRect.bottom - BorderWidth)
+		else if (MouseY < WindowRect.bottom && MouseY >= WindowRect.bottom - BorderWidth)
 			Row = 2;
 
-		if (MousePt.x >= WindowRect.left && MousePt.x < WindowRect.left + BorderWidth)
+		if (MouseX >= WindowRect.left && MouseX < WindowRect.left + BorderWidth)
 			Col = 0;
-		else if (MousePt.x < WindowRect.right && MousePt.x >= WindowRect.right - BorderWidth)
+		else if (MouseX < WindowRect.right && MouseX >= WindowRect.right - BorderWidth)
 			Col = 2;
 
 		LRESULT hitTests[3][3] =
@@ -1087,6 +1154,10 @@ public:
 	{
 		RECT WindowRect;
 		GetWindowRect(MainHwnd, &WindowRect);
+		WindowRect = FromScale(WindowRect);
+		X = FromScale(X);
+		Y = FromScale(Y);
+
 		if (Y < WindowRect.top + 2 || Y > WindowRect.top + 27 || X > WindowRect.right - 2)
 			return HTNOWHERE;
 		int32 ButtonWidth = 27;
@@ -1698,6 +1769,8 @@ public:
 			SendMessageW(LogHwnd, WM_SETREDRAW, false, 0);
 			RECT Rect;
 			GetClientRect(MainHwnd, &Rect);
+			Rect = FromScale(Rect);
+
 			int NewActivitiesTotalHeight = TraverseActivityPositions(Rect.right, Rect.bottom, [](Activity& A, int32 X, int32 Y, int32 Width, int32 RowIndex) {});
 			if (NewActivitiesTotalHeight == ActivitiesTotalHeight)
 			{
@@ -1713,6 +1786,58 @@ public:
 		}
 	}
 
+	template<typename T>
+	T ToScale(T Value) const { return T(Value*DpiScale); }
+	template<typename T>
+	T FromScale(T Value) const { return T(Value/DpiScale); }
+
+	RECT ToScale(const RECT& Rect)
+	{
+		return { ToScale(Rect.left), ToScale(Rect.top), ToScale(Rect.right), ToScale(Rect.bottom) };
+	}
+
+	RECT FromScale(const RECT& Rect)
+	{
+		return { FromScale(Rect.left), FromScale(Rect.top), FromScale(Rect.right), FromScale(Rect.bottom) };
+	}
+
+	void SetWindowPosScale(HWND hWnd, HWND hWndInsertAfter, int X, int Y, int Cx, int Cy, UINT Flags)
+	{
+		::SetWindowPos(hWnd, hWndInsertAfter, ToScale(X), ToScale(Y), ToScale(Cx), ToScale(Cy), Flags);
+	}
+
+	void FillRect(HDC Dc, const RECT& Rect, HBRUSH Brush)
+	{
+		RECT R = Rect;
+		::FillRect(Dc, &R, Brush);
+	}
+
+	void RectangleScale(HDC Dc, int32 X, int32 Y, int32 Width, int32 Height)
+	{
+		::Rectangle(Dc, ToScale(X), ToScale(Y), ToScale(Width), ToScale(Height));
+	}
+
+	void GetTextRectW(HDC Dc, const wchar_t* Text, int32 Len, RECT* Rect, int32 Flags)
+	{
+		::DrawTextW(Dc, Text, Len, Rect, Flags | DT_MODIFYSTRING);
+	}
+
+	void DrawTextW(HDC Dc, const wchar_t* Text, int32 Len, const RECT& Rect, int32 Flags)
+	{
+		RECT R = Rect;
+		::DrawTextW(Dc, Text, Len, &R, Flags);
+	}
+
+	void MoveToScale(HDC Dc, int32 X, int32 Y)
+	{
+		::MoveToEx(Dc, ToScale(X), ToScale(Y), NULL);
+	}
+
+	void LineToScale(HDC Dc, int32 X, int32 Y)
+	{
+		::LineTo(Dc, ToScale(X), ToScale(Y));
+	}
+
 	void HandleNewLogEntries()
 	{
 		if (bSuspendAddingEntries)
@@ -1805,6 +1930,130 @@ public:
 		SendMessageW(LogHwnd, WM_SETREDRAW, true, 0);
 	}
 
+	void DrawNonClientArea(HWND hWnd, WPARAM wParam)
+	{
+		RECT windowRect;
+		GetWindowRect(hWnd, &windowRect);
+		HRGN region = NULL;
+		if (wParam == NULLREGION)
+		{
+			region = CreateRectRgn(windowRect.left, windowRect.top, windowRect.right, windowRect.bottom);
+		}
+		else
+		{
+			HRGN copy = CreateRectRgn(0, 0, 0, 0);
+			if (CombineRgn(copy, (HRGN)wParam, NULL, RGN_COPY))
+				region = copy;
+			else
+				DeleteObject(copy);
+		}
+
+		HDC dc = GetDCEx(hWnd, region, DCX_WINDOW | DCX_CACHE | DCX_INTERSECTRGN | DCX_LOCKWINDOWUPDATE);
+		if (!dc)
+		{
+			if (region)
+			{
+				DeleteObject(region);
+			}
+			return;
+		}
+
+		RECT rect = FromScale(windowRect);
+
+		int32 width = rect.right - rect.left;
+		int32 height = rect.bottom - rect.top;
+
+		int32 ButtonWidth = 27;
+
+		RECT CaptionRect{ 0,0, width, 27 };
+		FillRect(dc, ToScale(CaptionRect), BackgroundBrush);
+
+		SelectObject(dc, GetStockObject(NULL_BRUSH));
+		SelectObject(dc, WindowBorderPen);
+		Rectangle(dc, 0, 0, windowRect.right - windowRect.left, windowRect.bottom - windowRect.top);
+
+		DrawIconEx(dc, ToScale(3), ToScale(3), Icon, ToScale(21), ToScale(21), 0, NULL, DI_NORMAL);
+
+		SetBkColor(dc, BackgroundColor);
+		SetTextColor(dc, TextColor);
+		SelectObject(dc, Font);
+		CaptionRect.left += 32;
+		CaptionRect.right -= ButtonWidth * 4;
+		DrawTextW(dc, *ConsoleTitle, ConsoleTitle.Len(), ToScale(CaptionRect), DT_SINGLELINE | DT_VCENTER | DT_END_ELLIPSIS);
+
+		SelectObject(dc, TextPen);
+
+		if (NcButtonHot != -1 && (NcButtonDown == -1 || NcButtonHot == NcButtonDown))
+		{
+			int32 NcButton = NcButtonHot;
+			HBRUSH Brush = ButtonBrush;
+			if (NcButtonDown != -1)
+				Brush = ButtonPressedBrush;
+			int32 ButtonIndex = 0;
+			if (NcButton == HTMAXBUTTON)
+				ButtonIndex = 1;
+			else if (NcButton == HTMINBUTTON)
+				ButtonIndex = 2;
+			else if (NcButton == HTEXPAND)
+				ButtonIndex = 3;
+			RECT ButtonRect{ width - ButtonWidth*(ButtonIndex+1), 1, width - 1 - ButtonWidth*ButtonIndex, 27};
+			FillRect(dc, ToScale(ButtonRect), Brush);
+		}
+
+		int32 MiddleY = 27 / 2;
+		int32 MiddleX = width - ButtonWidth / 2 - ButtonWidth * 3 - 1;
+
+		// Double arrow
+		if (bConsoleExpanded)
+		{
+			MoveToScale(dc, MiddleX - 4, MiddleY - 4);
+			LineToScale(dc, MiddleX, MiddleY);
+			LineToScale(dc, MiddleX + 5, MiddleY - 5);
+			MoveToScale(dc, MiddleX - 4, MiddleY);
+			LineToScale(dc, MiddleX, MiddleY + 4);
+			LineToScale(dc, MiddleX + 5, MiddleY - 1);
+		}
+		else
+		{
+			MoveToScale(dc, MiddleX - 4, MiddleY);
+			LineToScale(dc, MiddleX, MiddleY - 4);
+			LineToScale(dc, MiddleX + 5, MiddleY + 1);
+			MoveToScale(dc, MiddleX - 4, MiddleY + 4);
+			LineToScale(dc, MiddleX, MiddleY);
+			LineToScale(dc, MiddleX + 5, MiddleY + 5);
+		}
+
+		// Minimize
+		MiddleX += ButtonWidth;
+		MoveToScale(dc, MiddleX - 4, MiddleY);
+		LineToScale(dc, MiddleX + 5, MiddleY);
+			
+		// Maximize
+		MiddleX += ButtonWidth;
+		WINDOWPLACEMENT Placement;
+		Placement.length = sizeof(Placement);
+		GetWindowPlacement(hWnd, &Placement);
+		if (Placement.showCmd == SW_MAXIMIZE)
+		{
+			RectangleScale(dc, MiddleX, MiddleY - 3, MiddleX + 8, MiddleY + 5);
+			SelectObject(dc, BackgroundBrush);
+			RectangleScale(dc, MiddleX - 3, MiddleY, MiddleX + 5, MiddleY + 8);
+		}
+		else
+		{
+			RectangleScale(dc, MiddleX - 5, MiddleY - 5, MiddleX + 5, MiddleY + 5);
+		}
+			
+		// Close
+		MiddleX += ButtonWidth;
+		MoveToScale(dc, MiddleX - 4, MiddleY - 4);
+		LineToScale(dc, MiddleX + 5, MiddleY + 5);
+		MoveToScale(dc, MiddleX + 4, MiddleY - 4);
+		LineToScale(dc, MiddleX - 5, MiddleY + 5);
+
+		ReleaseDC(hWnd, dc);
+	}
+
 	LRESULT MainWinProc(HWND hWnd, UINT Msg, WPARAM wParam, LPARAM lParam)
 	{
 		switch (Msg)
@@ -1823,7 +2072,7 @@ public:
 			{
 			case SIZE_MAXIMIZED:
 			case SIZE_RESTORED:
-				UpdateSize(LOWORD(lParam), HIWORD(lParam), true);
+				UpdateSize(FromScale(LOWORD(lParam)), FromScale(HIWORD(lParam)), true);
 				RECT WindowRect;
 				GetWindowRect(MainHwnd, &WindowRect);
 				ConsoleWidth = WindowRect.right - WindowRect.left;
@@ -1856,7 +2105,7 @@ public:
 			HDC Hdc = (HDC)wParam;
 			RECT Rect;
 			GetClientRect(hWnd, &Rect);
-			FillRect(Hdc, &Rect, BackgroundBrush);
+			FillRect(Hdc, Rect, BackgroundBrush);
 
 			// Borders around edit boxes
 			SelectObject(Hdc, BorderPen);
@@ -1889,7 +2138,7 @@ public:
 			SetBkColor(Hdc, EditBackgroundColor);
 			RECT Rect;
 			GetClientRect(filterHwnd, &Rect);
-			FillRect(Hdc, &Rect, EditBackgroundBrush);
+			FillRect(Hdc, Rect, EditBackgroundBrush);
 			SetTextColor(Hdc, TextColor);
 			if (!GetWindowTextLength(filterHwnd))
 			{
@@ -1910,7 +2159,7 @@ public:
 				if (ItemCount - TopIndex < VisibleCount)
 				{
 					Rect.top = (ItemCount - TopIndex) * LogFontHeight;
-					FillRect((HDC)wParam, &Rect, LogBackgroundBrush); // Draw the padding below the last line
+					FillRect((HDC)wParam, Rect, LogBackgroundBrush); // Draw the padding below the last line
 				}
 				return (INT_PTR)GetStockObject(NULL_BRUSH);
 			}
@@ -1952,10 +2201,10 @@ public:
 					Color = ButtonHighlightColor;
 				}
 				SelectObject(nmcd.hdc, GetStockObject(NULL_BRUSH));
-				FillRect(nmcd.hdc, &nmcd.rc, Brush);
+				FillRect(nmcd.hdc, nmcd.rc, Brush);
 				SetBkColor(nmcd.hdc, Color);
 				SetTextColor(nmcd.hdc, TextColor);
-				DrawTextW(nmcd.hdc, Str, StrLen, &nmcd.rc, DT_SINGLELINE | DT_VCENTER | DT_CENTER | DT_NOCLIP);
+				DrawTextW(nmcd.hdc, Str, StrLen, nmcd.rc, DT_SINGLELINE | DT_VCENTER | DT_CENTER | DT_NOCLIP);
 				SelectObject(nmcd.hdc, BorderPen);
 				Rectangle(nmcd.hdc, nmcd.rc.left, nmcd.rc.top, nmcd.rc.right, nmcd.rc.bottom);
 				return CDRF_SKIPDEFAULT;
@@ -1978,10 +2227,10 @@ public:
 				{
 					if (pdis->rcItem.bottom != 1)
 					{
-						FillRect(Hdc, &pdis->rcItem, BackgroundBrush);
+						FillRect(Hdc, pdis->rcItem, BackgroundBrush);
 						SetTextColor(Hdc, GetSysColor(COLOR_HIGHLIGHTTEXT));
 						SelectObject(Hdc, LogFont);
-						DrawTextW(Hdc, Str, StrLen, &pdis->rcItem, DT_SINGLELINE | DT_VCENTER);
+						DrawTextW(Hdc, Str, StrLen, pdis->rcItem, DT_SINGLELINE | DT_VCENTER);
 					}
 					break;
 				}
@@ -2001,21 +2250,21 @@ public:
 				// All this stuff just to prevent flickering when text changes fast
 				int32 TextOffset = 1;
 				RECT ModifiedRect = pdis->rcItem;
-				DrawTextW(Hdc, Str, StrLen, &ModifiedRect, DT_SINGLELINE | DT_END_ELLIPSIS | DT_CALCRECT | DT_MODIFYSTRING);
+				GetTextRectW(Hdc, Str, StrLen, &ModifiedRect, DT_SINGLELINE | DT_END_ELLIPSIS | DT_CALCRECT);
 				RECT TempRect = pdis->rcItem;
 				TempRect.bottom = TempRect.top + TextOffset;
 				TempRect.right = ModifiedRect.right;
-				FillRect(Hdc, &TempRect, Brush);
+				FillRect(Hdc, TempRect, Brush);
 
 				TempRect = pdis->rcItem;
 				TempRect.left = ModifiedRect.right;
-				FillRect(Hdc, &TempRect, Brush);
+				FillRect(Hdc, TempRect, Brush);
 				if (ModifiedRect.bottom < pdis->rcItem.bottom)
 				{
 					TempRect = pdis->rcItem;
 					TempRect.right = ModifiedRect.right;
 					TempRect.top = ModifiedRect.bottom;
-					FillRect(Hdc, &TempRect, Brush);
+					FillRect(Hdc, TempRect, Brush);
 				}
 
 				StrLen = FCString::Strlen(Str);
@@ -2039,12 +2288,12 @@ public:
 
 				if (pdis->itemState & ODS_SELECTED)
 				{
-					FillRect(Hdc, &pdis->rcItem, GetSysColorBrush(COLOR_HIGHLIGHT));
+					FillRect(Hdc, pdis->rcItem, GetSysColorBrush(COLOR_HIGHLIGHT));
 					SetTextColor(Hdc, GetSysColor(COLOR_HIGHLIGHTTEXT));
 				}
 				else
 				{
-					FillRect(Hdc, &pdis->rcItem, LogBackgroundBrush);
+					FillRect(Hdc, pdis->rcItem, LogBackgroundBrush);
 					if (Entry.TextAttribute == (FOREGROUND_INTENSITY | FOREGROUND_RED))
 						SetTextColor(Hdc, RGB(220, 0, 0));
 					else if (Entry.TextAttribute == (FOREGROUND_INTENSITY | FOREGROUND_RED | FOREGROUND_GREEN))
@@ -2309,7 +2558,7 @@ public:
 				RedrawNC();
 				RECT Rect;
 				GetClientRect(MainHwnd, &Rect);
-				UpdateSize(Rect.right, Rect.bottom, true);
+				UpdateSize(FromScale(Rect.right), FromScale(Rect.bottom), true);
 				RedrawWindow(MainHwnd, NULL, NULL, RDW_INVALIDATE);
 				if (bConsoleExpanded)
 					SetFocus(CommandHwnd);
@@ -2331,124 +2580,15 @@ public:
 		case WM_NCCALCSIZE:
 		{
 			LPNCCALCSIZE_PARAMS ncParams = (LPNCCALCSIZE_PARAMS)lParam;
-			ncParams->rgrc[0].top += 27;
-			ncParams->rgrc[0].left += 1;
-			ncParams->rgrc[0].bottom -= 1;
-			ncParams->rgrc[0].right -= 1;
+			ncParams->rgrc[0].top += ToScale(27);
+			ncParams->rgrc[0].left += ToScale(1);
+			ncParams->rgrc[0].bottom -= ToScale(1);
+			ncParams->rgrc[0].right -= ToScale(1);
 			return 0;
 		}
 		case WM_NCPAINT:
 		{
-			RECT rect;
-			GetWindowRect(hWnd, &rect);
-			HRGN region = NULL;
-			if (wParam == NULLREGION)
-				region = CreateRectRgn(rect.left, rect.top, rect.right, rect.bottom);
-			else
-			{
-				HRGN copy = CreateRectRgn(0, 0, 0, 0);
-				if (CombineRgn(copy, (HRGN)wParam, NULL, RGN_COPY))
-					region = copy;
-				else
-					DeleteObject(copy);
-			}
-
-			HDC dc = GetDCEx(hWnd, region, DCX_WINDOW | DCX_CACHE | DCX_INTERSECTRGN | DCX_LOCKWINDOWUPDATE);
-			if (!dc && region)
-				DeleteObject(region);
-
-			int32 width = rect.right - rect.left;
-			int32 height = rect.bottom - rect.top;
-
-			int32 ButtonWidth = 27;
-
-			RECT CaptionRect{ 0,0,width, 27 };
-			FillRect(dc, &CaptionRect, BackgroundBrush);
-
-			SelectObject(dc, GetStockObject(NULL_BRUSH));
-			SelectObject(dc, WindowBorderPen);
-			Rectangle(dc, 0, 0, width, height);
-
-			DrawIconEx(dc, 3, 3, Icon, 21, 21, 0, NULL, DI_NORMAL);
-
-			SetBkColor(dc, BackgroundColor);
-			SetTextColor(dc, TextColor);
-			SelectObject(dc, Font);
-			CaptionRect.left += 32;
-			CaptionRect.right -= ButtonWidth * 4;
-			DrawTextW(dc, *ConsoleTitle, ConsoleTitle.Len(), &CaptionRect, DT_SINGLELINE | DT_VCENTER | DT_END_ELLIPSIS);
-
-			SelectObject(dc, TextPen);
-
-			if (NcButtonHot != -1 && (NcButtonDown == -1 || NcButtonHot == NcButtonDown))
-			{
-				int32 NcButton = NcButtonHot;
-				HBRUSH Brush = ButtonBrush;
-				if (NcButtonDown != -1)
-					Brush = ButtonPressedBrush;
-				int32 ButtonIndex = 0;
-				if (NcButton == HTMAXBUTTON)
-					ButtonIndex = 1;
-				else if (NcButton == HTMINBUTTON)
-					ButtonIndex = 2;
-				else if (NcButton == HTEXPAND)
-					ButtonIndex = 3;
-				RECT ButtonRect{ width - ButtonWidth*(ButtonIndex+1), 1, width - 1 - ButtonWidth*ButtonIndex, 27};
-				FillRect(dc, &ButtonRect, Brush);
-			}
-
-			int32 MiddleY = 27 / 2;
-			int32 MiddleX = width - ButtonWidth / 2 - ButtonWidth * 3 - 1;
-
-			// Double arrow
-			if (bConsoleExpanded)
-			{
-				MoveToEx(dc, MiddleX - 4, MiddleY - 4, NULL);
-				LineTo(dc, MiddleX, MiddleY);
-				LineTo(dc, MiddleX + 5, MiddleY - 5);
-				MoveToEx(dc, MiddleX - 4, MiddleY, NULL);
-				LineTo(dc, MiddleX, MiddleY + 4);
-				LineTo(dc, MiddleX + 5, MiddleY - 1);
-			}
-			else
-			{
-				MoveToEx(dc, MiddleX - 4, MiddleY, NULL);
-				LineTo(dc, MiddleX, MiddleY - 4);
-				LineTo(dc, MiddleX + 5, MiddleY + 1);
-				MoveToEx(dc, MiddleX - 4, MiddleY + 4, NULL);
-				LineTo(dc, MiddleX, MiddleY);
-				LineTo(dc, MiddleX + 5, MiddleY + 5);
-			}
-
-			// Minimize
-			MiddleX += ButtonWidth;
-			MoveToEx(dc, MiddleX - 4, MiddleY, NULL);
-			LineTo(dc, MiddleX + 5, MiddleY);
-			
-			// Maximize
-			MiddleX += ButtonWidth;
-			WINDOWPLACEMENT Placement;
-			Placement.length = sizeof(Placement);
-			GetWindowPlacement(hWnd, &Placement);
-			if (Placement.showCmd == SW_MAXIMIZE)
-			{
-				Rectangle(dc, MiddleX, MiddleY - 3, MiddleX + 8, MiddleY + 5);
-				SelectObject(dc, BackgroundBrush);
-				Rectangle(dc, MiddleX - 3, MiddleY, MiddleX + 5, MiddleY + 8);
-			}
-			else
-			{
-				Rectangle(dc, MiddleX - 5, MiddleY - 5, MiddleX + 5, MiddleY + 5);
-			}
-			
-			// Close
-			MiddleX += ButtonWidth;
-			MoveToEx(dc, MiddleX - 4, MiddleY - 4, NULL);
-			LineTo(dc, MiddleX + 5, MiddleY + 5);
-			MoveToEx(dc, MiddleX + 4, MiddleY - 4, NULL);
-			LineTo(dc, MiddleX - 5, MiddleY + 5);
-
-			ReleaseDC(hWnd, dc);
+			DrawNonClientArea(hWnd, wParam);
 			return 0;
 		}
 		case WM_NCACTIVATE:
@@ -2482,6 +2622,7 @@ public:
 	TRingBuffer<LogEntry> Log;
 	FDateTime StartDateTime;
 	double StartTime;
+	double DpiScale = 1.0;
 	int32 AddedEntryLogVirtualIndex = -1;
 	HICON Icon;
 	HFONT Font;
@@ -2498,6 +2639,7 @@ public:
 	HWND TipHwnd;
 	POINT TipPosition;
 	WNDPROC LogHwndWndProcPtr;
+	COLORREF BorderColor;
 	COLORREF TextColor;
 	COLORREF BackgroundColor;
 	COLORREF EditBackgroundColor;
@@ -2709,10 +2851,10 @@ void FWindowsConsoleOutputDevice2::SaveToINI()
 		WindowRWLock.WriteUnlock();
 		return;
 	}
-	int32 ConsoleWidth = Window->ConsoleWidth;
-	int32 ConsoleHeight = Window->ConsoleHeight;
-	int32 ConsolePosX = Window->ConsolePosX;
-	int32 ConsolePosY = Window->ConsolePosY;
+	int32 ConsoleWidth = Window->FromScale(Window->ConsoleWidth);
+	int32 ConsoleHeight = Window->FromScale(Window->ConsoleHeight);
+	int32 ConsolePosX = Window->FromScale(Window->ConsolePosX);
+	int32 ConsolePosY = Window->FromScale(Window->ConsolePosY);
 	int32 ConsoleColor = (int32)Window->ConsoleColor;
 	int32 ConsoleShowDateTime = Window->ConsoleShowDateTime;
 	FString ConsoleIncludeFilterStr = Window->ConsoleIncludeFilterStr;
