[ApplicationCore]
* Removed ownerless expression

diff --git a/Engine/Source/Runtime/ApplicationCore/Private/Windows/WindowsConsoleOutputDevice2.cpp b/Engine/Source/Runtime/ApplicationCore/Private/Windows/WindowsConsoleOutputDevice2.cpp
--- a/Engine/Source/Runtime/ApplicationCore/Private/Windows/WindowsConsoleOutputDevice2.cpp
+++ b/Engine/Source/Runtime/ApplicationCore/Private/Windows/WindowsConsoleOutputDevice2.cpp
@@ -340,7 +340,6 @@ public:
 						FontHeight = -MulDiv(8, dpi, 72);
 						LogFont = (HFONT)CreateFontW(FontHeight, 0, 0, 0, FW_NORMAL, false, false, false, ANSI_CHARSET, OUT_DEFAULT_PRECIS, CLIP_DEFAULT_PRECIS, ANTIALIASED_QUALITY, FIXED_PITCH | FF_MODERN, FontName);
 
-						FontSize;
 						Hdc = GetDC(MainHwnd);
 						SelectObject(Hdc, LogFont);
 						GetTextExtentPoint32(Hdc, TEXT("A"), 1, &FontSize);
