[ApplicationCore]
* Fixed position of activity status dots when dpi scaling is enabled

diff --git a/Engine/Source/Runtime/ApplicationCore/Private/Windows/WindowsConsoleOutputDevice2.cpp b/Engine/Source/Runtime/ApplicationCore/Private/Windows/WindowsConsoleOutputDevice2.cpp
--- a/Engine/Source/Runtime/ApplicationCore/Private/Windows/WindowsConsoleOutputDevice2.cpp
+++ b/Engine/Source/Runtime/ApplicationCore/Private/Windows/WindowsConsoleOutputDevice2.cpp
@@ -2320,7 +2320,7 @@ public:
 			RECT Rect;
 			GetClientRect(hWnd, &Rect);
 
-			TraverseActivityPositions(Rect.right, Rect.bottom, [&](Activity& A, int32 X, int32 Y, int32 Width, int32 RowIndex)
+			TraverseActivityPositions(FromScale(Rect.right), FromScale(Rect.bottom), [&](Activity& A, int32 X, int32 Y, int32 Width, int32 RowIndex)
 				{
 					if (A.Light == 0)
 						return;
@@ -2330,7 +2330,7 @@ public:
 					int32 FontHeight = 16;
 					int32 CircleSize = 11;
 					int32 HalfDiff = (FontHeight - CircleSize) / 2;
-					RoundRect(Hdc, X - CircleSize + XOffset, Y + HalfDiff + YOffset, X + XOffset, Y + FontHeight - HalfDiff + YOffset, CircleSize, CircleSize);
+					RoundRect(Hdc, ToScale(X - CircleSize + XOffset), ToScale(Y + HalfDiff + YOffset), ToScale(X + XOffset), ToScale(Y + FontHeight - HalfDiff + YOffset), ToScale(CircleSize), ToScale(CircleSize));
 				});
 			EndPaint(hWnd, &ps);
 			return 0;
