Added "SuppressConsoleOutputVerboseLogging" commandline arg that suppresses Verbose and VeryVerbose logging from output console devices.

This is useful for when you want to turn on verbose file logging but don't want to incur the (extreme on some platforms) performance hit of also logging that stuff to the console device

diff --git a/Engine/Source/Runtime/ApplicationCore/Private/Linux/LinuxConsoleOutputDevice.cpp b/Engine/Source/Runtime/ApplicationCore/Private/Linux/LinuxConsoleOutputDevice.cpp
--- a/Engine/Source/Runtime/ApplicationCore/Private/Linux/LinuxConsoleOutputDevice.cpp
+++ b/Engine/Source/Runtime/ApplicationCore/Private/Linux/LinuxConsoleOutputDevice.cpp
@@ -37,6 +37,11 @@ FLinuxConsoleOutputDevice::FLinuxConsoleOutputDevice()
 
 	// If -stdout is getting set we are doing a printf in LaunchEngineLoop.cpp, so lets avoid double printing here
 	bIsStdoutSet = FParse::Param(*CommandLine, TEXT("stdout"));
+
+	if (FParse::Param(*CommandLine, TEXT("SuppressConsoleOutputVerboseLogging")))
+	{
+		AllowedLogVerbosity = ELogVerbosity::Log;
+	}
 }
 
 FLinuxConsoleOutputDevice::~FLinuxConsoleOutputDevice()
@@ -58,6 +63,9 @@ void FLinuxConsoleOutputDevice::Serialize(const TCHAR* Data, ELogVerbosity::Type
 	if (!bIsWindowShown || bIsStdoutSet)
 		return;
 
+	if ((Verbosity & ELogVerbosity::VerbosityMask) > AllowedLogVerbosity)
+		return;
+
 	static bool bEntry=false;
 	if (!GIsCriticalError || bEntry)
 	{
diff --git a/Engine/Source/Runtime/ApplicationCore/Private/Mac/MacConsoleOutputDevice.cpp b/Engine/Source/Runtime/ApplicationCore/Private/Mac/MacConsoleOutputDevice.cpp
--- a/Engine/Source/Runtime/ApplicationCore/Private/Mac/MacConsoleOutputDevice.cpp
+++ b/Engine/Source/Runtime/ApplicationCore/Private/Mac/MacConsoleOutputDevice.cpp
@@ -20,6 +20,10 @@ FMacConsoleOutputDevice::FMacConsoleOutputDevice()
 	, TextViewTextColor(NULL)
 	, OutstandingTasks(0)
 {
+	if (FParse::Param(FCommandLine::Get(), TEXT("SuppressConsoleOutputVerboseLogging")))
+	{
+		AllowedLogVerbosity = ELogVerbosity::Log;
+	}
 }
 
 FMacConsoleOutputDevice::~FMacConsoleOutputDevice()
@@ -160,6 +164,9 @@ bool FMacConsoleOutputDevice::IsShown()
 
 void FMacConsoleOutputDevice::Serialize( const TCHAR* Data, ELogVerbosity::Type Verbosity, const class FName& Category )
 {
+	if ((Verbosity & ELogVerbosity::VerbosityMask) > AllowedLogVerbosity)
+		return;
+
 	if( ConsoleHandle )
 	{
 		FScopeLock ScopeLock( &CriticalSection );
diff --git a/Engine/Source/Runtime/ApplicationCore/Private/Windows/WindowsConsoleOutputDevice.cpp b/Engine/Source/Runtime/ApplicationCore/Private/Windows/WindowsConsoleOutputDevice.cpp
--- a/Engine/Source/Runtime/ApplicationCore/Private/Windows/WindowsConsoleOutputDevice.cpp
+++ b/Engine/Source/Runtime/ApplicationCore/Private/Windows/WindowsConsoleOutputDevice.cpp
@@ -159,6 +159,11 @@ FWindowsConsoleOutputDevice::FWindowsConsoleOutputDevice()
 			}
 		}
 	}
+
+	if (FParse::Param(FCommandLine::Get(), TEXT("SuppressConsoleOutputVerboseLogging")))
+	{
+		AllowedLogVerbosity = ELogVerbosity::Log;
+	}
 #endif
 }
 
@@ -307,6 +312,9 @@ void FWindowsConsoleOutputDevice::Serialize( const TCHAR* Data, ELogVerbosity::T
 {
 	using namespace UE::String;
 
+	if ((Verbosity & ELogVerbosity::VerbosityMask) > AllowedLogVerbosity)
+		return;
+
 	if( ConsoleHandle )
 	{
 		const double RealTime = Time == -1.0f ? FPlatformTime::Seconds() - GStartTime : Time;
diff --git a/Engine/Source/Runtime/ApplicationCore/Private/Windows/WindowsConsoleOutputDevice2.cpp b/Engine/Source/Runtime/ApplicationCore/Private/Windows/WindowsConsoleOutputDevice2.cpp
--- a/Engine/Source/Runtime/ApplicationCore/Private/Windows/WindowsConsoleOutputDevice2.cpp
+++ b/Engine/Source/Runtime/ApplicationCore/Private/Windows/WindowsConsoleOutputDevice2.cpp
@@ -2834,6 +2834,11 @@ FWindowsConsoleOutputDevice2::FWindowsConsoleOutputDevice2()
 			}
 		}
 	}
+
+	if (FParse::Param(FCommandLine::Get(), TEXT("SuppressConsoleOutputVerboseLogging")))
+	{
+		AllowedLogVerbosity = ELogVerbosity::Log;
+	}
 #endif
 }
 
@@ -3026,6 +3031,9 @@ bool FWindowsConsoleOutputDevice2::IsShown()
 void FWindowsConsoleOutputDevice2::Serialize( const TCHAR* Data, ELogVerbosity::Type Verbosity, const class FName& Category, const double Time )
 {
 	using namespace UE::String;
+	
+	if ((Verbosity & ELogVerbosity::VerbosityMask) > AllowedLogVerbosity)
+		return;
 
 	const double RealTime = Time == -1.0f ? FPlatformTime::Seconds() - GStartTime : Time;
 
diff --git a/Engine/Source/Runtime/ApplicationCore/Public/Linux/LinuxConsoleOutputDevice.h b/Engine/Source/Runtime/ApplicationCore/Public/Linux/LinuxConsoleOutputDevice.h
--- a/Engine/Source/Runtime/ApplicationCore/Public/Linux/LinuxConsoleOutputDevice.h
+++ b/Engine/Source/Runtime/ApplicationCore/Public/Linux/LinuxConsoleOutputDevice.h
@@ -25,6 +25,9 @@ private:
 	/** true if stdout is set via the command line */
 	bool bIsStdoutSet;
 
+	/** maximum allowed log verbosity */
+	ELogVerbosity::Type AllowedLogVerbosity = ELogVerbosity::All;
+
 public:
 
 	/**
diff --git a/Engine/Source/Runtime/ApplicationCore/Public/Mac/MacConsoleOutputDevice.h b/Engine/Source/Runtime/ApplicationCore/Public/Mac/MacConsoleOutputDevice.h
--- a/Engine/Source/Runtime/ApplicationCore/Public/Mac/MacConsoleOutputDevice.h
+++ b/Engine/Source/Runtime/ApplicationCore/Public/Mac/MacConsoleOutputDevice.h
@@ -32,6 +32,9 @@ private:
 	/** Number of outstanding dispatched Cocoa tasks */
 	uint64 OutstandingTasks;
 
+	/** Maximum allowed log verbosity */
+	ELogVerbosity::Type AllowedLogVerbosity = ELogVerbosity::All;
+
 	/**
 	 * Saves the console window's position and size to the game .ini
 	 */
diff --git a/Engine/Source/Runtime/ApplicationCore/Public/Windows/WindowsConsoleOutputDevice.h b/Engine/Source/Runtime/ApplicationCore/Public/Windows/WindowsConsoleOutputDevice.h
--- a/Engine/Source/Runtime/ApplicationCore/Public/Windows/WindowsConsoleOutputDevice.h
+++ b/Engine/Source/Runtime/ApplicationCore/Public/Windows/WindowsConsoleOutputDevice.h
@@ -21,6 +21,9 @@ private:
 	/** true if this console device is attached to an existing console window. */
 	bool bIsAttached;
 
+	/** maximum allowed log verbosity */
+	ELogVerbosity::Type AllowedLogVerbosity = ELogVerbosity::All;
+
 #if !UE_BUILD_SHIPPING
 	/** An entry for log category highlighting */
 	struct FLogHighlight
diff --git a/Engine/Source/Runtime/ApplicationCore/Public/Windows/WindowsConsoleOutputDevice2.h b/Engine/Source/Runtime/ApplicationCore/Public/Windows/WindowsConsoleOutputDevice2.h
--- a/Engine/Source/Runtime/ApplicationCore/Public/Windows/WindowsConsoleOutputDevice2.h
+++ b/Engine/Source/Runtime/ApplicationCore/Public/Windows/WindowsConsoleOutputDevice2.h
@@ -21,6 +21,9 @@ private:
 	FConsoleWindow* Window;
 	FRWLock WindowRWLock;
 
+	/** maximum allowed log verbosity */
+	ELogVerbosity::Type AllowedLogVerbosity = ELogVerbosity::All;
+
 #if !UE_BUILD_SHIPPING
 	/** An entry for log category highlighting */
 	struct FLogHighlight
