{Mobile} User is unable to change enhanced input Controller Bindings on mobile
Used direct strings when creating the input scope for mobile.
This is intended as a temporary solution.

#rnx
#tests Android

diff --git a/Engine/Source/Runtime/ApplicationCore/Private/Android/AndroidInputInterface.cpp b/Engine/Source/Runtime/ApplicationCore/Private/Android/AndroidInputInterface.cpp
--- a/Engine/Source/Runtime/ApplicationCore/Private/Android/AndroidInputInterface.cpp
+++ b/Engine/Source/Runtime/ApplicationCore/Private/Android/AndroidInputInterface.cpp
@@ -1338,7 +1338,7 @@ void FAndroidInputInterface::SendControllerEvents()
 				continue;
 			}
 
-			FInputDeviceScope InputScope(nullptr, *(TouchDeviceInfo->Name), TouchInputDeviceId.GetId(), TouchDeviceInfo->Descriptor);
+			FInputDeviceScope InputScope(nullptr, *(TouchDeviceInfo->Name), TouchInputDeviceId.GetId(), "MobileTouch");
 		}
 
 		IPlatformInputDeviceMapper& DeviceMapper = IPlatformInputDeviceMapper::Get();
@@ -1374,7 +1374,7 @@ void FAndroidInputInterface::SendControllerEvents()
 			const FInputDeviceId InputDeviceId = InternalDeviceIdMappings.FindDeviceId(DeviceInfo.Descriptor);
 			check(InputDeviceId != INPUTDEVICEID_NONE);
 
-			FInputDeviceScope InputScope(nullptr, *(DeviceInfo.Name), InputDeviceId.GetId(), DeviceInfo.Descriptor);
+			FInputDeviceScope InputScope(nullptr, *(DeviceInfo.Name), InputDeviceId.GetId(), "Gamepad");
 
 			FAndroidGamepadDeviceMapping* DeviceData = GameControllerDataMap.Find(InputDeviceId);
 			if (DeviceData == nullptr)
