1. Fixed the issue that the most recently used hardware input device is reported with the wrong type.
2. The input device ID 0 is handled as from the default touch screen.
3. Changed the HardwareDeviceIdentifier of FInputDeviceScope to predefined values, "Gamepad" for game controllers, and "MobileInput" for the touch screen.
4. Changed the type of FAndroidInputDeviceInfo::Name to FName.

diff --git a/Engine/Source/Runtime/ApplicationCore/Private/Android/AndroidInputInterface.cpp b/Engine/Source/Runtime/ApplicationCore/Private/Android/AndroidInputInterface.cpp
--- a/Engine/Source/Runtime/ApplicationCore/Private/Android/AndroidInputInterface.cpp
+++ b/Engine/Source/Runtime/ApplicationCore/Private/Android/AndroidInputInterface.cpp
@@ -29,6 +29,11 @@ TMap<int32, FAndroidInputDeviceInfo> FAndroidInputInterface::InputDeviceInfoMap;
 TMap<FInputDeviceId, FAndroidGamepadDeviceMapping> FAndroidInputInterface::GameControllerDataMap;
 TArray<int32> FAndroidInputInterface::GameControllerIdMapping;
 
+FName FAndroidInputInterface::InputClassName_DefaultMobileTouch;
+FName FAndroidInputInterface::InputClassName_DefaultGamepad;
+FString FAndroidInputInterface::HardwareDeviceIdentifier_DefaultMobileTouch;
+FString FAndroidInputInterface::HardwareDeviceIdentifier_DefaultGamepad;
+
 int32 FAndroidInputInterface::CurrentVibeIntensity;
 int32 FAndroidInputInterface::MaxVibeTime = 1000;
 double FAndroidInputInterface::LastVibeUpdateTime = 0.0;
@@ -121,8 +126,10 @@ JNI_METHOD void Java_com_epicgames_unreal_GameControllerManager_nativeOnInputDev
 	FAndroidInputInterface::HandleInputDeviceStateEvent(device_id, StateEvent, DeviceType);
 }
 
-void FAndroidGamepadDeviceMapping::Init(const FString& DeviceName)
+void FAndroidGamepadDeviceMapping::Init(const FName InDeviceName)
 {
+	FString DeviceName = InDeviceName.ToString();
+
 	// Use device name to decide on mapping scheme
 	if (DeviceName.StartsWith(TEXT("Amazon")))
 	{
@@ -331,6 +338,11 @@ FAndroidInputInterface::FAndroidInputInterface(const TSharedRef< FGenericApplica
 	GameControllerDataMap.Empty();
 	GameControllerIdMapping.Empty();
 
+	InputClassName_DefaultMobileTouch = TEXT("DefaultMobileTouch");
+	InputClassName_DefaultGamepad = TEXT("DefaultGamepad");
+	HardwareDeviceIdentifier_DefaultMobileTouch = TEXT("MobileTouch");
+	HardwareDeviceIdentifier_DefaultGamepad = TEXT("Gamepad");
+
 	if (AndroidEnableInputDeviceListener)
 	{
 		extern bool AndroidThunkCpp_InitGameControllerManager();
@@ -404,10 +416,8 @@ const FInputDeviceId FAndroidInputInterface::GetMappedInputDeviceId(int32 Contro
 	return INPUTDEVICEID_NONE;
 }
 
-const FString& FAndroidInputInterface::GetGamepadControllerName(int32 ControllerId)
+const FName FAndroidInputInterface::GetGamepadControllerName(int32 ControllerId)
 {
-	static const FString NoName(TEXT("None"));
-
 	FAndroidInputDeviceInfo* DeviceInfo = nullptr;
 	if (GetInputDeviceByControllerId(ControllerId, &DeviceInfo))
 	{
@@ -417,7 +427,7 @@ const FString& FAndroidInputInterface::GetGamepadControllerName(int32 Controller
 			return DeviceInfo->Name;
 		}
 	}
-	return NoName;
+	return InputClassName_DefaultGamepad;
 }
 
 void FAndroidInputInterface::SetMessageHandler( const TSharedRef< FGenericApplicationMessageHandler >& InMessageHandler )
@@ -1254,7 +1264,7 @@ void FAndroidInputInterface::SendControllerEvents()
 						{
 							DeviceToRemove.Push(Pair.Key);
 							UE_LOG(LogAndroid, Error, TEXT("Found input device with same descriptor! DeviceId = %d, DeviceName=%s, Descriptor=%s"),
-								Pair.Key, *(Pair.Value.Name), *(Pair.Value.Descriptor));
+								Pair.Key, *(Pair.Value.Name.ToString()), *(Pair.Value.Descriptor));
 						}
 					}
 
@@ -1274,7 +1284,7 @@ void FAndroidInputInterface::SendControllerEvents()
 					UE_LOG(LogAndroid, Log, 
 						TEXT("New input device recorded: ControllerId = %d, DeviceId=%d, ControllerType=%d, DeviceName=%s, Descriptor=%s, FeedbackMotorCount=%d"),
 						FindControllerId(CurrentDevice.DeviceId), CurrentDevice.DeviceId, CurrentDevice.DeviceType, 
-						*CurrentDevice.Name, *CurrentDevice.Descriptor, CurrentDevice.FeedbackMotorCount);
+						*(CurrentDevice.Name.ToString()), *CurrentDevice.Descriptor, CurrentDevice.FeedbackMotorCount);
 				}
 				else
 				{
@@ -1301,11 +1311,12 @@ void FAndroidInputInterface::SendControllerEvents()
 		// some special inputs has -1 as their device id (e.g. scrcpy inputs)
 		// use default touchscreen instead to handle their inputs
 		FInputDeviceId TouchInputDeviceId = IPlatformInputDeviceMapper::Get().GetDefaultInputDevice();
-
+		
+		FAndroidInputDeviceInfo* TouchDeviceInfo = nullptr;
 		TouchInput Touch = FAndroidInputInterface::TouchInputStack[i];
-		if (Touch.DeviceId >= 0)
+		if (Touch.DeviceId > 0)
 		{
-			FAndroidInputDeviceInfo* TouchDeviceInfo = InputDeviceInfoMap.Find(Touch.DeviceId);
+			TouchDeviceInfo = InputDeviceInfoMap.Find(Touch.DeviceId);
 			if (TouchDeviceInfo == nullptr)
 			{
 				continue;
@@ -1337,10 +1348,14 @@ void FAndroidInputInterface::SendControllerEvents()
 			{
 				continue;
 			}
-
-			FInputDeviceScope InputScope(nullptr, *(TouchDeviceInfo->Name), TouchInputDeviceId.GetId(), "MobileTouch");
 		}
 
+		///The FInputDeviceScope::HardwareDeviceIdentifier has to be one of values in UInputPlatformSettings::HardwareDevices.
+		// This is a temp solution with a hardcoded string which can be mapped to FHardwareDeviceIdentifier::DefaultMobileTouch.
+		// TODO: Future improvement is needed to acquire them by values in the Android input device info.
+		FName DeviceName = TouchDeviceInfo ? TouchDeviceInfo->Name : InputClassName_DefaultMobileTouch;
+		FInputDeviceScope InputScope(nullptr, DeviceName, TouchInputDeviceId.GetId(), HardwareDeviceIdentifier_DefaultMobileTouch);
+
 		IPlatformInputDeviceMapper& DeviceMapper = IPlatformInputDeviceMapper::Get();
 		FPlatformUserId PlatformUserId = DeviceMapper.GetUserForInputDevice(TouchInputDeviceId);
 
@@ -1374,7 +1389,10 @@ void FAndroidInputInterface::SendControllerEvents()
 			const FInputDeviceId InputDeviceId = InternalDeviceIdMappings.FindDeviceId(DeviceInfo.Descriptor);
 			check(InputDeviceId != INPUTDEVICEID_NONE);
 
-			FInputDeviceScope InputScope(nullptr, *(DeviceInfo.Name), InputDeviceId.GetId(), "Gamepad");
+			///The FInputDeviceScope::HardwareDeviceIdentifier has to be one of values in UInputPlatformSettings::HardwareDevices.
+			// This is a temp solution with a hardcoded string which can be mapped to FHardwareDeviceIdentifier::DefaultGamepad.
+			// TODO: Future improvement is needed to acquire them by values in the Android input device info.
+			FInputDeviceScope InputScope(nullptr, DeviceInfo.Name, InputDeviceId.GetId(), HardwareDeviceIdentifier_DefaultGamepad);
 
 			FAndroidGamepadDeviceMapping* DeviceData = GameControllerDataMap.Find(InputDeviceId);
 			if (DeviceData == nullptr)
@@ -1738,7 +1756,7 @@ void FAndroidInputInterface::DumpInputDevices()
 		IPlatformInputDeviceMapper& DeviceMapper = IPlatformInputDeviceMapper::Get();
 		FPlatformUserId UserId = DeviceMapper.GetUserForInputDevice(InputDeviceId);
 		UE_LOG(LogAndroid, Log, TEXT("DeviceInfo, DeviceId = %03d, Type = %d, State = %d, IsExternal = %d, InputDeviceId = %02d, UserId = %02d, Name = %s"),
-			Info.DeviceId, Info.DeviceType, Info.DeviceState, Info.IsExternal, InputDeviceId.GetId(), UserId.GetInternalId(), *Info.Name);
+			Info.DeviceId, Info.DeviceType, Info.DeviceState, Info.IsExternal, InputDeviceId.GetId(), UserId.GetInternalId(), *(Info.Name.ToString()));
 	}
 	
 	UE_LOG(LogAndroid, Log, TEXT("---"));
diff --git a/Engine/Source/Runtime/ApplicationCore/Private/Android/AndroidPlatformApplicationMisc.cpp b/Engine/Source/Runtime/ApplicationCore/Private/Android/AndroidPlatformApplicationMisc.cpp
--- a/Engine/Source/Runtime/ApplicationCore/Private/Android/AndroidPlatformApplicationMisc.cpp
+++ b/Engine/Source/Runtime/ApplicationCore/Private/Android/AndroidPlatformApplicationMisc.cpp
@@ -112,7 +112,7 @@ bool FAndroidApplicationMisc::IsControllerAssignedToGamepad(int32 ControllerId)
 
 FString FAndroidApplicationMisc::GetGamepadControllerName(int32 ControllerId)
 {
-	return FAndroidInputInterface::GetGamepadControllerName(ControllerId);
+	return (FAndroidInputInterface::GetGamepadControllerName(ControllerId).ToString());
 }
 
 extern void AndroidThunkCpp_ClipboardCopy(const FString& Str);
diff --git a/Engine/Source/Runtime/ApplicationCore/Public/Android/AndroidInputInterface.h b/Engine/Source/Runtime/ApplicationCore/Public/Android/AndroidInputInterface.h
--- a/Engine/Source/Runtime/ApplicationCore/Public/Android/AndroidInputInterface.h
+++ b/Engine/Source/Runtime/ApplicationCore/Public/Android/AndroidInputInterface.h
@@ -124,7 +124,6 @@ struct FAndroidInputDeviceInfo
 {
 	FAndroidInputDeviceInfo()
 	{
-		Name.Empty();
 		Descriptor.Empty();
 	}
 
@@ -132,7 +131,7 @@ struct FAndroidInputDeviceInfo
 	int32 VendorId = 0;
 	int32 ProductId = 0;
 	int32 ControllerId = 0;
-	FString Name;
+	FName Name;
 	FString Descriptor;
 	int32 FeedbackMotorCount = 0;
 	bool IsExternal = false;
@@ -190,12 +189,12 @@ struct FAndroidGamepadDeviceMapping
 		ResetRuntimeData();
 	}
 
-	FAndroidGamepadDeviceMapping(const FString& DeviceName)
+	FAndroidGamepadDeviceMapping(const FName DeviceName)
 	{
 		Init(DeviceName);
 	}
 
-	void Init(const FString& DeviceName);
+	void Init(const FName DeviceName);
 
 	void ResetRuntimeData()
 	{
@@ -295,7 +294,7 @@ public:
 	static void ResetGamepadAssignmentToController(int32 ControllerId);
 	static bool IsControllerAssignedToGamepad(int32 ControllerId);
 	static const FInputDeviceId GetMappedInputDeviceId(int32 ControllerId);
-	static const FString& GetGamepadControllerName(int32 ControllerId);
+	static const FName GetGamepadControllerName(int32 ControllerId);
 
 	static void HandleInputDeviceStateEvent(int32 DeviceId, InputDeviceStateEvent StateEvent, InputDeviceType DeviceType);
 
@@ -409,6 +408,12 @@ private:
 	static TArray<int32/*Android DeviceId*/> GameControllerIdMapping;
 	static TMap<int32/*Android deviceId*/, FAndroidInputDeviceInfo> InputDeviceInfoMap;
 	static TMap<FInputDeviceId, FAndroidGamepadDeviceMapping> GameControllerDataMap;
+	
+	static FName InputClassName_DefaultMobileTouch;
+	static FName InputClassName_DefaultGamepad;
+	// The following two vars should be FName. Waiting for a refactoring of the FInputDeviceScope's ctor
+	static FString HardwareDeviceIdentifier_DefaultMobileTouch;
+	static FString HardwareDeviceIdentifier_DefaultGamepad;
 
 	static FGamepadKeyNames::Type ButtonMapping[MAX_NUM_CONTROLLER_BUTTONS];
 
diff --git a/Engine/Source/Runtime/Launch/Private/Android/AndroidJNI.cpp b/Engine/Source/Runtime/Launch/Private/Android/AndroidJNI.cpp
--- a/Engine/Source/Runtime/Launch/Private/Android/AndroidJNI.cpp
+++ b/Engine/Source/Runtime/Launch/Private/Android/AndroidJNI.cpp
@@ -729,7 +729,11 @@ bool AndroidThunkCpp_GetInputDeviceInfo(int32 deviceId, FAndroidInputDeviceInfo
 			results.ProductId = (int32)Env->GetIntField(*deviceInfo, FJavaWrapper::InputDeviceInfo_ProductId);
 			results.ControllerId = (int32)Env->GetIntField(*deviceInfo, FJavaWrapper::InputDeviceInfo_ControllerId);
 
-			results.Name = FJavaHelper::FStringFromLocalRef(Env, (jstring)Env->GetObjectField(*deviceInfo, FJavaWrapper::InputDeviceInfo_Name));
+			FScopedJavaObject JavaInputDeviceName = NewScopedJavaObject(Env, (jstring)Env->GetObjectField(*deviceInfo, FJavaWrapper::InputDeviceInfo_Name));
+			const auto chars = Env->GetStringUTFChars(*JavaInputDeviceName, 0);
+			results.Name = chars; // UTF8_TO_TCHAR(chars)
+			Env->ReleaseStringUTFChars(*JavaInputDeviceName, chars);
+
 			results.Descriptor = FJavaHelper::FStringFromLocalRef(Env, (jstring)Env->GetObjectField(*deviceInfo, FJavaWrapper::InputDeviceInfo_Descriptor));
 
 			results.FeedbackMotorCount = (int32)Env->GetIntField(*deviceInfo, FJavaWrapper::InputDeviceInfo_FeedbackMotorCount);
@@ -743,7 +747,7 @@ bool AndroidThunkCpp_GetInputDeviceInfo(int32 deviceId, FAndroidInputDeviceInfo
 	results.VendorId = 0;
 	results.ProductId = 0;
 	results.ControllerId = -1;
-	results.Name = FString("Unknown");
+	results.Name = TEXT("Unknown");
 	results.Descriptor = FString("Unknown");
 	results.FeedbackMotorCount = 0;
 	return false;
