Fixed a compilation error caused by a wrong header file including order.

diff --git a/Engine/Source/Runtime/ApplicationCore/Private/Android/AndroidJavaMediaFrameData.cpp b/Engine/Source/Runtime/ApplicationCore/Private/Android/AndroidJavaMediaFrameData.cpp
--- a/Engine/Source/Runtime/ApplicationCore/Private/Android/AndroidJavaMediaFrameData.cpp
+++ b/Engine/Source/Runtime/ApplicationCore/Private/Android/AndroidJavaMediaFrameData.cpp
@@ -1,7 +1,6 @@
 // Copyright Epic Games, Inc. All Rights Reserved.
 
 #include "Android/AndroidJavaMediaFrameData.h"
-#include "Android/AndroidJava.h"
 #include "Android/AndroidApplication.h"
 
 #include "GlobalShader.h"
