Don't add user with invalidate platform user id in AuthNull impl.

diff --git a/Engine/Source/Runtime/ApplicationCore/Private/GenericPlatform/GenericPlatformInputDeviceMapper.cpp b/Engine/Source/Runtime/ApplicationCore/Private/GenericPlatform/GenericPlatformInputDeviceMapper.cpp
--- a/Engine/Source/Runtime/ApplicationCore/Private/GenericPlatform/GenericPlatformInputDeviceMapper.cpp
+++ b/Engine/Source/Runtime/ApplicationCore/Private/GenericPlatform/GenericPlatformInputDeviceMapper.cpp
@@ -158,7 +158,10 @@ int32 IPlatformInputDeviceMapper::GetAllActiveUsers(TArray<FPlatformUserId>& Out
 	// Add the owning platform user for each input device
 	for (TPair<FInputDeviceId, FPlatformInputDeviceState> MappedDevice : MappedInputDevices)
 	{
-		OutUsers.AddUnique(MappedDevice.Value.OwningPlatformUser);
+		if (MappedDevice.Value.OwningPlatformUser.IsValid())
+		{
+			OutUsers.AddUnique(MappedDevice.Value.OwningPlatformUser);
+		}
 	}
 	
 	return OutUsers.Num();
