Added an overload of FPlatformSplash::SetSplashText that takes a FText

This allows you to call it before localization data has loaded, and still get a localized result when the splash screen opens

diff --git a/Engine/Source/Runtime/ApplicationCore/Private/GenericPlatform/GenericPlatformSplash.cpp b/Engine/Source/Runtime/ApplicationCore/Private/GenericPlatform/GenericPlatformSplash.cpp
--- a/Engine/Source/Runtime/ApplicationCore/Private/GenericPlatform/GenericPlatformSplash.cpp
+++ b/Engine/Source/Runtime/ApplicationCore/Private/GenericPlatform/GenericPlatformSplash.cpp
@@ -1,8 +1,11 @@
 // Copyright Epic Games, Inc. All Rights Reserved.
 
 #include "GenericPlatform/GenericPlatformSplash.h"
+
 #include "Containers/UnrealString.h"
 #include "HAL/FileManager.h"
+#include "HAL/PlatformSplash.h"
+#include "Internationalization/Text.h"
 #include "Misc/Paths.h"
 
 // Supported file extension for splash image
@@ -42,6 +45,11 @@ void FGenericPlatformSplash::SetCustomSplashImage(const TCHAR* SplashFilename)
 	GCustomSplashScreenFileName = SplashFilename;
 }
 
+void FGenericPlatformSplash::SetSplashText(const SplashTextType::Type InType, const TCHAR* InText)
+{
+	FPlatformSplash::SetSplashText(InType, FText::AsCultureInvariant(InText));
+}
+
 /**
 * Finds a usable splash pathname for the given filename
 *
diff --git a/Engine/Source/Runtime/ApplicationCore/Private/Linux/LinuxPlatformSplash.cpp b/Engine/Source/Runtime/ApplicationCore/Private/Linux/LinuxPlatformSplash.cpp
--- a/Engine/Source/Runtime/ApplicationCore/Private/Linux/LinuxPlatformSplash.cpp
+++ b/Engine/Source/Runtime/ApplicationCore/Private/Linux/LinuxPlatformSplash.cpp
@@ -620,7 +620,7 @@ bool FLinuxSplashState::InitSplashResources(const FText &AppName, const FString
 void FLinuxSplashState::SetSplashText( const SplashTextType::Type InType, const FText& InText )
 {
 #if WITH_EDITOR
-	if (!InText.EqualTo(SplashText[ InType ]))
+	if (!InText.IdenticalTo(SplashText[ InType ], ETextIdenticalModeFlags::LexicalCompareInvariants))
 	{
 		SplashText[ InType ] = InText;
 		bStringsChanged = true;
@@ -816,7 +816,7 @@ void FLinuxPlatformSplash::Hide()
 * @param	InType		Type of text to change
 * @param	InText		Text to display
 */
-void FLinuxPlatformSplash::SetSplashText( const SplashTextType::Type InType, const TCHAR* InText )
+void FLinuxPlatformSplash::SetSplashText( const SplashTextType::Type InType, const FText& InText )
 {
 	if (GSplashState)
 	{
@@ -824,7 +824,7 @@ void FLinuxPlatformSplash::SetSplashText( const SplashTextType::Type InType, con
 		// not interesting to an end-user (also, it's not usually localized properly.)
 		if (InType == SplashTextType::CopyrightInfo || GIsEditor)
 		{
-			GSplashState->SetSplashText(InType, FText::FromString(InText));
+			GSplashState->SetSplashText(InType, InText);
 		}
 		GSplashState->Pump();
 	}
diff --git a/Engine/Source/Runtime/ApplicationCore/Private/Mac/MacPlatformSplash.cpp b/Engine/Source/Runtime/ApplicationCore/Private/Mac/MacPlatformSplash.cpp
--- a/Engine/Source/Runtime/ApplicationCore/Private/Mac/MacPlatformSplash.cpp
+++ b/Engine/Source/Runtime/ApplicationCore/Private/Mac/MacPlatformSplash.cpp
@@ -332,7 +332,7 @@ bool FMacPlatformSplash::IsShown()
 	return (GSplashWindow != NULL);
 }
 
-void FMacPlatformSplash::SetSplashText(const SplashTextType::Type InType, const TCHAR* InText)
+void FMacPlatformSplash::SetSplashText(const SplashTextType::Type InType, const FText& InText)
 {
 	// We only want to bother drawing startup progress in the editor, since this information is
 	// not interesting to an end-user (also, it's not usually localized properly.)
@@ -346,9 +346,9 @@ void FMacPlatformSplash::SetSplashText(const SplashTextType::Type InType, const
 				FScopeLock Lock(&GSplashMutex);
 
 				// Update splash text
-				if( FCString::Strcmp( InText, *GSplashScreenText[ InType ].ToString() ) != 0 )
+				if (!InText.IdenticalTo(GSplashScreenText[InType], ETextIdenticalModeFlags::LexicalCompareInvariants))
 				{
-					GSplashScreenText[ InType ] = FText::FromString( InText );
+					GSplashScreenText[ InType ] = InText;
 					bWasUpdated = true;
 				}
 			}
diff --git a/Engine/Source/Runtime/ApplicationCore/Private/Windows/WindowsPlatformSplash.cpp b/Engine/Source/Runtime/ApplicationCore/Private/Windows/WindowsPlatformSplash.cpp
--- a/Engine/Source/Runtime/ApplicationCore/Private/Windows/WindowsPlatformSplash.cpp
+++ b/Engine/Source/Runtime/ApplicationCore/Private/Windows/WindowsPlatformSplash.cpp
@@ -780,7 +780,7 @@ void FWindowsPlatformSplash::SetProgress(int ProgressPercent)
 	}
 }
 
-void FWindowsPlatformSplash::SetSplashText( const SplashTextType::Type InType, const TCHAR* InText )
+void FWindowsPlatformSplash::SetSplashText( const SplashTextType::Type InType, const FText& InText )
 {
 	// We only want to bother drawing startup progress in the editor, since this information is
 	// not interesting to an end-user (also, it's not usually localized properly.)
@@ -795,9 +795,9 @@ void FWindowsPlatformSplash::SetSplashText( const SplashTextType::Type InType, c
 			bool bWasUpdated = false;
 			{
 				// Update splash text
-				if( FCString::Strcmp( InText, *GSplashScreenText[ InType ].ToString() ) != 0 )
+				if (!InText.IdenticalTo(GSplashScreenText[InType], ETextIdenticalModeFlags::LexicalCompareInvariants))
 				{
-					GSplashScreenText[ InType ] = FText::FromString( InText );
+					GSplashScreenText[ InType ] = InText;
 					bWasUpdated = true;
 				}
 			}
@@ -813,10 +813,7 @@ void FWindowsPlatformSplash::SetSplashText( const SplashTextType::Type InType, c
 	else // We haven't started drawing yet we can set text ahead of time to show
 	{
 		// Update splash text
-		if (FCString::Strcmp(InText, *GSplashScreenText[InType].ToString()) != 0)
-		{
-			GSplashScreenText[InType] = FText::FromString(InText);
-		}
+		GSplashScreenText[InType] = InText;
 	}
 }
 
diff --git a/Engine/Source/Runtime/ApplicationCore/Public/GenericPlatform/GenericPlatformSplash.h b/Engine/Source/Runtime/ApplicationCore/Public/GenericPlatform/GenericPlatformSplash.h
--- a/Engine/Source/Runtime/ApplicationCore/Public/GenericPlatform/GenericPlatformSplash.h
+++ b/Engine/Source/Runtime/ApplicationCore/Public/GenericPlatform/GenericPlatformSplash.h
@@ -32,6 +32,7 @@ namespace SplashTextType
 }
 
 class FString;
+class FText;
 
 /**
  * Generic implementation for most platforms
@@ -68,9 +69,10 @@ struct FGenericPlatformSplash
 	 * @param	InType		Type of text to change
 	 * @param	InText		Text to display
 	 */
-	FORCEINLINE static void SetSplashText( const SplashTextType::Type InType, const TCHAR* InText )
+	// TODO: Deprecate the const TCHAR* overload
+	static APPLICATIONCORE_API void SetSplashText( const SplashTextType::Type InType, const TCHAR* InText );
+	FORCEINLINE static void SetSplashText(const SplashTextType::Type InType, const FText& InText)
 	{
-
 	}
 
 	/**
diff --git a/Engine/Source/Runtime/ApplicationCore/Public/Linux/LinuxPlatformSplash.h b/Engine/Source/Runtime/ApplicationCore/Public/Linux/LinuxPlatformSplash.h
--- a/Engine/Source/Runtime/ApplicationCore/Public/Linux/LinuxPlatformSplash.h
+++ b/Engine/Source/Runtime/ApplicationCore/Public/Linux/LinuxPlatformSplash.h
@@ -29,7 +29,8 @@ struct FLinuxPlatformSplash : public FGenericPlatformSplash
 	 * @param	InType		Type of text to change
 	 * @param	InText		Text to display
 	 */
-	static APPLICATIONCORE_API void SetSplashText(const SplashTextType::Type InType, const TCHAR* InText);
+	using FGenericPlatformSplash::SetSplashText;
+	static APPLICATIONCORE_API void SetSplashText(const SplashTextType::Type InType, const FText& InText);
 };
 
 typedef FLinuxPlatformSplash FPlatformSplash;
diff --git a/Engine/Source/Runtime/ApplicationCore/Public/Mac/MacPlatformSplash.h b/Engine/Source/Runtime/ApplicationCore/Public/Mac/MacPlatformSplash.h
--- a/Engine/Source/Runtime/ApplicationCore/Public/Mac/MacPlatformSplash.h
+++ b/Engine/Source/Runtime/ApplicationCore/Public/Mac/MacPlatformSplash.h
@@ -26,7 +26,8 @@ struct APPLICATIONCORE_API FMacPlatformSplash : public FGenericPlatformSplash
 	 * @param	InType		Type of text to change
 	 * @param	InText		Text to display
 	 */
-	static void SetSplashText( const SplashTextType::Type InType, const TCHAR* InText );
+	using FGenericPlatformSplash::SetSplashText;
+	static void SetSplashText( const SplashTextType::Type InType, const FText& InText );
 
 	/**
 	 * Return whether the splash screen is being shown or not
diff --git a/Engine/Source/Runtime/ApplicationCore/Public/Windows/WindowsPlatformSplash.h b/Engine/Source/Runtime/ApplicationCore/Public/Windows/WindowsPlatformSplash.h
--- a/Engine/Source/Runtime/ApplicationCore/Public/Windows/WindowsPlatformSplash.h
+++ b/Engine/Source/Runtime/ApplicationCore/Public/Windows/WindowsPlatformSplash.h
@@ -5,7 +5,6 @@
 #include "CoreTypes.h"
 #include "GenericPlatform/GenericPlatformSplash.h"
 
-
 /**
  * Windows splash implementation.
  */
@@ -31,7 +30,8 @@ struct FWindowsPlatformSplash
 	 * @param InType Type of text to change.
 	 * @param InText Text to display.
 	 */
-	static APPLICATIONCORE_API void SetSplashText( const SplashTextType::Type InType, const TCHAR* InText );
+	using FGenericPlatformSplash::SetSplashText;
+	static APPLICATIONCORE_API void SetSplashText( const SplashTextType::Type InType, const FText& InText );
 
 	/**
 	 * Return whether the splash screen is being shown or not
