Deprecated FPlatformSplash::SetSplashText overload that takes a const TCHAR*

diff --git a/Engine/Source/Runtime/ApplicationCore/Public/GenericPlatform/GenericPlatformSplash.h b/Engine/Source/Runtime/ApplicationCore/Public/GenericPlatform/GenericPlatformSplash.h
--- a/Engine/Source/Runtime/ApplicationCore/Public/GenericPlatform/GenericPlatformSplash.h
+++ b/Engine/Source/Runtime/ApplicationCore/Public/GenericPlatform/GenericPlatformSplash.h
@@ -69,7 +69,7 @@ struct FGenericPlatformSplash
 	 * @param	InType		Type of text to change
 	 * @param	InText		Text to display
 	 */
-	// TODO: Deprecate the const TCHAR* overload
+	UE_DEPRECATED(5.6, "Use the overload of FPlatformSlash::SetSplashText that takes a FText")
 	static APPLICATIONCORE_API void SetSplashText( const SplashTextType::Type InType, const TCHAR* InText );
 	FORCEINLINE static void SetSplashText(const SplashTextType::Type InType, const FText& InText)
 	{
