Add set Window Mode for Android Windows
remove unused class members

#rnx

diff --git a/Engine/Source/Runtime/ApplicationCore/Public/Android/AndroidWindow.h b/Engine/Source/Runtime/ApplicationCore/Public/Android/AndroidWindow.h
--- a/Engine/Source/Runtime/ApplicationCore/Public/Android/AndroidWindow.h
+++ b/Engine/Source/Runtime/ApplicationCore/Public/Android/AndroidWindow.h
@@ -74,6 +74,8 @@ public:
 
 	APPLICATIONCORE_API bool GetNativeWindowResolution(int32_t& OutWidth, int32_t& OutHeight) const;
 
+	APPLICATIONCORE_API virtual void SetWindowMode(EWindowMode::Type InNewWindowMode) { WindowMode = InNewWindowMode; }
+
 	class FNativeAccessor
 	{
 		bool bWriteAccess;
@@ -138,12 +140,13 @@ public:
 
 protected:
 	/** @return true if the native window is currently in fullscreen mode, false otherwise */
-	virtual EWindowMode::Type GetWindowMode() const override { return EWindowMode::Fullscreen; }
+	virtual EWindowMode::Type GetWindowMode() const override { return WindowMode; }
 
 private:
 	ANativeWindow* CurrentANativeHandle = nullptr;
 	ANativeWindow* PreviousANativeHandle = nullptr;
 	UE::FSharedMutex ANativeHandleMutex;
+	EWindowMode::Type WindowMode = EWindowMode::Fullscreen;
 
 	/**
 	 * Protect the constructor; only TSharedRefs of this class can be made.
@@ -152,10 +155,6 @@ private:
 
 	FAndroidApplication* OwningApplication;
 
-	/** Store the window region size for querying whether a point lies within the window */
-	int32 RegionX;
-	int32 RegionY;
-
 	static APPLICATIONCORE_API void* NativeWindow;
 
 	// Waits for the event thread to report an initial window size.
