Remove oudtated UIScaleFactor logic fro Android Cursor

#rnx

diff --git a/Engine/Source/Runtime/ApplicationCore/Private/Android/AndroidApplication.cpp b/Engine/Source/Runtime/ApplicationCore/Private/Android/AndroidApplication.cpp
--- a/Engine/Source/Runtime/ApplicationCore/Private/Android/AndroidApplication.cpp
+++ b/Engine/Source/Runtime/ApplicationCore/Private/Android/AndroidApplication.cpp
@@ -102,13 +102,6 @@ void FAndroidApplication::PollGameDeviceState( const float TimeDelta )
 		BroadcastDisplayMetricsChanged(DisplayMetrics);
 		FCoreDelegates::OnSafeFrameChangedEvent.Broadcast();
 
-		// the cursor needs to compute the proper slate scaling factor each time the display metrics change
-		TSharedPtr<FAndroidCursor> AndroidCursor = StaticCastSharedPtr<FAndroidCursor>(Cursor);
-		if (AndroidCursor.IsValid())
-		{
-			AndroidCursor->ComputeUIScaleFactor();
-		}
-
 		bWindowSizeChanged = false;
 	}
     
diff --git a/Engine/Source/Runtime/ApplicationCore/Private/Android/AndroidCursor.cpp b/Engine/Source/Runtime/ApplicationCore/Private/Android/AndroidCursor.cpp
--- a/Engine/Source/Runtime/ApplicationCore/Private/Android/AndroidCursor.cpp
+++ b/Engine/Source/Runtime/ApplicationCore/Private/Android/AndroidCursor.cpp
@@ -8,33 +8,13 @@ FAndroidCursor::FAndroidCursor()
 	, CurrentPosition(FVector2D::ZeroVector)
 	, CursorClipRect()
 	, bShow(false)
-	, UIScaleFactor(1.0f)
 {
-	ComputeUIScaleFactor();
-}
-
-void FAndroidCursor::ComputeUIScaleFactor()
-{
-	FPlatformRect ScreenRect = FAndroidWindow::GetScreenRect();
-	int32_t ScreenWidth, ScreenHeight;
-	FAndroidWindow::CalculateSurfaceSize(ScreenWidth, ScreenHeight);
-
-	if (ScreenWidth > ScreenHeight)
-	{
-		UIScaleFactor = (float)(ScreenRect.Bottom - ScreenRect.Top) / (float)ScreenHeight;
-	}
-	else
-	{
-		UIScaleFactor = (float)(ScreenRect.Right - ScreenRect.Left) / (float)ScreenWidth;
-	}
 }
 
 void FAndroidCursor::SetPosition( const int32 X, const int32 Y )
 {
 	FVector2D NewPosition(X, Y);
 
-//	NewPosition *= UIScaleFactor;
-
 	UpdateCursorClipping(NewPosition);
 	CurrentPosition = NewPosition;
 }
diff --git a/Engine/Source/Runtime/ApplicationCore/Public/Android/AndroidCursor.h b/Engine/Source/Runtime/ApplicationCore/Public/Android/AndroidCursor.h
--- a/Engine/Source/Runtime/ApplicationCore/Public/Android/AndroidCursor.h
+++ b/Engine/Source/Runtime/ApplicationCore/Public/Android/AndroidCursor.h
@@ -37,9 +37,6 @@ public:
 
 	virtual void SetTypeShape(EMouseCursor::Type InCursorType, void* CursorHandle) override { }
 
-	/** this function will compute the necessary position scale factor needed to transform the cursor in Slate space */
-	void ComputeUIScaleFactor();
-
 private:
 	bool UpdateCursorClipping(FVector2D& CursorPosition);
 
@@ -47,7 +44,4 @@ private:
 	FVector2D CurrentPosition;
 	FIntRect CursorClipRect;
 	bool bShow;
-
-	/** scale factor computed by ComputeUIScaleFactor() and used in SetPosition() */
-	float UIScaleFactor;
 };
