Fixed possibly shutdown crashes by unregistering FWindowsApplication::ShutDownAfterError against FCoreDelegates::OnShutdownAfterError in its destructor.

diff --git a/Engine/Source/Runtime/ApplicationCore/Private/Windows/WindowsApplication.cpp b/Engine/Source/Runtime/ApplicationCore/Private/Windows/WindowsApplication.cpp
--- a/Engine/Source/Runtime/ApplicationCore/Private/Windows/WindowsApplication.cpp
+++ b/Engine/Source/Runtime/ApplicationCore/Private/Windows/WindowsApplication.cpp
@@ -1139,6 +1139,8 @@ bool FWindowsApplication::RegisterClass( const HINSTANCE HInstance, const HICON
 
 FWindowsApplication::~FWindowsApplication()
 {
+	FCoreDelegates::OnShutdownAfterError.RemoveAll(this);
+
 	if (TextInputMethodSystem.IsValid())
 	{
 		TextInputMethodSystem->Terminate();
