Rudimentary support for command history in -newconsole

diff --git a/Engine/Source/Runtime/ApplicationCore/Private/Windows/WindowsConsoleOutputDevice2.cpp b/Engine/Source/Runtime/ApplicationCore/Private/Windows/WindowsConsoleOutputDevice2.cpp
--- a/Engine/Source/Runtime/ApplicationCore/Private/Windows/WindowsConsoleOutputDevice2.cpp
+++ b/Engine/Source/Runtime/ApplicationCore/Private/Windows/WindowsConsoleOutputDevice2.cpp
@@ -243,6 +243,9 @@ public:
 		SetWindowLongPtr(LogScrollHwnd, GWLP_USERDATA, (LONG_PTR)this);
 
 		CommandHwnd = CreateWindow2(WC_EDIT, NULL, ES_AUTOHSCROLL, Font, 0, 0, 1, 1, ID_COMMAND);
+		SetWindowLongPtr(CommandHwnd, GWLP_USERDATA, (LONG_PTR)this);
+		CommandHwndWndProcPtr = (WNDPROC)SetWindowLongPtr(CommandHwnd, GWLP_WNDPROC, (LONG_PTR)StaticCommandHwndWndProc);
+
 		RunCommandHwnd = CreateWindow2(WC_BUTTON, TEXT("RunCommand"), 0, Font, 0, 0, 1, 1, ID_COMMANDBUTTON);
 
 		SetFocus(CommandHwnd);
@@ -878,6 +881,12 @@ public:
 		return ThisPtr->LogHwndWndProc(hWnd, Msg, wParam, lParam);
 	}
 
+	static LRESULT CALLBACK StaticCommandHwndWndProc(HWND hWnd, UINT Msg, WPARAM wParam, LPARAM lParam)
+	{
+		FConsoleWindow* ThisPtr = (FConsoleWindow*)GetWindowLongPtr(hWnd, GWLP_USERDATA);
+		return ThisPtr->CommandHwndWndProc(hWnd, Msg, wParam, lParam);
+	}
+
 	static LRESULT CALLBACK StaticMainWinProc(HWND hWnd, UINT Msg, WPARAM wParam, LPARAM lParam)
 	{
 		FConsoleWindow* ThisPtr = (FConsoleWindow*)GetWindowLongPtr(hWnd, GWLP_USERDATA);
@@ -1577,6 +1586,39 @@ public:
 		return LogHwndWndProcPtr(hWnd, Msg, wParam, lParam);
 	}
 
+	LRESULT CommandHwndWndProc(HWND hWnd, UINT Msg, WPARAM wParam, LPARAM lParam)
+	{
+		switch (Msg)
+		{
+		case WM_KEYUP:
+		{
+			if (wParam == VK_UP)
+			{
+				TCHAR CurrentCommand[1024] = { '\0' };
+				GetWindowText(CommandHwnd, CurrentCommand, 1024);
+
+				FString Command;
+				Owner->GetLastCommand(Command);
+
+				if (Command == CurrentCommand)
+				{
+					Owner->GetPreviousCommand(Command);
+				}
+				SetWindowText(CommandHwnd, *Command);
+			}
+			else if (wParam == VK_DOWN)
+			{
+				FString NextCommand;
+				Owner->GetNextCommand(NextCommand);
+				SetWindowText(CommandHwnd, *NextCommand);
+			}
+			break;
+		}
+		}
+
+		return CommandHwndWndProcPtr(hWnd, Msg, wParam, lParam);
+	}
+
 	static LRESULT CALLBACK StaticMessageBoxHookProc(int nCode, WPARAM wParam, LPARAM lParam)
 	{
 		if (nCode == HCBT_ACTIVATE)
@@ -2431,6 +2473,8 @@ public:
 					break;
 				}
 
+				Owner->AddCommand(Command);
+
 				Async(EAsyncExecution::TaskGraphMainThread, [Cmd = FString(Command)]
 					{
 						if (IModularFeatures::Get().IsModularFeatureAvailable(IConsoleCommandExecutor::ModularFeatureName()))
@@ -2638,6 +2682,7 @@ public:
 	HWND TipHwnd;
 	POINT TipPosition;
 	WNDPROC LogHwndWndProcPtr;
+	WNDPROC CommandHwndWndProcPtr;
 	COLORREF BorderColor;
 	COLORREF TextColor;
 	COLORREF BackgroundColor;
@@ -2898,6 +2943,8 @@ void FWindowsConsoleOutputDevice2::SaveToINI()
 	Config.SetString(Selection, TEXT("IncludeFilter"), *ConsoleIncludeFilterStr, Filename);
 	Config.SetString(Selection, TEXT("ExcludeFilter"), *ConsoleExcludeFilterStr, Filename);
 
+	Config.SetArray(Selection, TEXT("CommandHistory"), CommandHistory, Filename);
+
 	Config.Flush(false, Filename);
 }
 
@@ -2945,6 +2992,8 @@ void FWindowsConsoleOutputDevice2::Show( bool ShowWindow )
 		Config.GetInt(Selection, TEXT("ConsoleShowDateTime"), ConsoleShowDateTime, Filename);
 		Config.GetString(Selection, TEXT("IncludeFilter"), ConsoleIncludeFilterStr, Filename);
 		Config.GetString(Selection, TEXT("ExcludeFilter"), ConsoleExcludeFilterStr, Filename);
+		Config.GetArray(Selection, TEXT("CommandHistory"), CommandHistory, Filename);
+		CommandHistoryCursor = CommandHistory.Num();
 
 		if (!FParse::Value(FCommandLine::Get(), TEXT("ConsoleX="), ConsolePosX) && !bHasX)
 		{
@@ -3151,3 +3200,50 @@ bool FWindowsConsoleOutputDevice2::CanBeUsedOnAnyThread() const
 	return true;
 }
 
+bool FWindowsConsoleOutputDevice2::GetPreviousCommand(FString& OutCommand)
+{
+	if (CommandHistoryCursor > 0)
+	{
+		CommandHistoryCursor--;
+
+		if (CommandHistoryCursor < CommandHistory.Num())
+		{
+			OutCommand = CommandHistory[CommandHistoryCursor];
+			return true;
+		}
+	}
+	return false;
+}
+
+bool FWindowsConsoleOutputDevice2::GetNextCommand(FString& OutCommand)
+{
+	if (CommandHistoryCursor < (CommandHistory.Num() - 1))
+	{
+		CommandHistoryCursor++;
+		OutCommand = CommandHistory[CommandHistoryCursor];
+		return true;
+	}
+	return false;
+}
+
+void FWindowsConsoleOutputDevice2::AddCommand(const FString& InCommand)
+{	
+	if (CommandHistoryCursor >= 0 && CommandHistoryCursor < CommandHistory.Num() && CommandHistory[CommandHistoryCursor] == InCommand)
+	{
+		// The command was probably executed from the command history. No need to store again.
+		return;
+	}
+
+	CommandHistoryCursor = CommandHistory.Add(InCommand);
+	SaveToINI();
+}
+
+bool FWindowsConsoleOutputDevice2::GetLastCommand(FString& OutCommand)
+{
+	if (CommandHistoryCursor >= 0 && CommandHistoryCursor < CommandHistory.Num())
+	{
+		OutCommand = CommandHistory[CommandHistoryCursor];
+		return true;
+	}
+	return false;
+}
diff --git a/Engine/Source/Runtime/ApplicationCore/Public/Windows/WindowsConsoleOutputDevice2.h b/Engine/Source/Runtime/ApplicationCore/Public/Windows/WindowsConsoleOutputDevice2.h
--- a/Engine/Source/Runtime/ApplicationCore/Public/Windows/WindowsConsoleOutputDevice2.h
+++ b/Engine/Source/Runtime/ApplicationCore/Public/Windows/WindowsConsoleOutputDevice2.h
@@ -20,6 +20,8 @@ private:
 	class FConsoleWindow;
 	FConsoleWindow* Window;
 	FRWLock WindowRWLock;
+	TArray<FString> CommandHistory;
+	int32 CommandHistoryCursor = 0;
 
 	/** maximum allowed log verbosity */
 	ELogVerbosity::Type AllowedLogVerbosity = ELogVerbosity::All;
@@ -101,6 +103,10 @@ public:
 
 	APPLICATIONCORE_API void SetColor( const TCHAR* Color );
 
+	bool GetPreviousCommand(FString& OutCommand);
+	bool GetNextCommand(FString& OutCommand);
+	bool GetLastCommand(FString& OutCommand);
+	void AddCommand(const FString& InCommand);
 private:
 	FWindowsConsoleOutputDevice2(const FWindowsConsoleOutputDevice2&) = delete;
 	FWindowsConsoleOutputDevice2& operator=(const FWindowsConsoleOutputDevice2&) = delete;
