fix SubmitTool and Zen prevents Mac from sleeping

#rnx

diff --git a/Engine/Source/Runtime/ApplicationCore/Private/Mac/MacPlatformApplicationMisc.cpp b/Engine/Source/Runtime/ApplicationCore/Private/Mac/MacPlatformApplicationMisc.cpp
--- a/Engine/Source/Runtime/ApplicationCore/Private/Mac/MacPlatformApplicationMisc.cpp
+++ b/Engine/Source/Runtime/ApplicationCore/Private/Mac/MacPlatformApplicationMisc.cpp
@@ -301,10 +301,10 @@ void FMacPlatformApplicationMisc::PostInit()
 		PostInitMacMenuStartupCb();
 	}
 
-	if (!MacApplication)
+	if (IsRunningDedicatedServer() || IsRunningCookCommandlet())
 	{
-		// No MacApplication means that app is a dedicated server, commandline tool or the editor running a commandlet. In these cases we don't want macOS to put our app into App Nap mode.
-		CommandletActivity = [[NSProcessInfo processInfo] beginActivityWithOptions:NSActivityUserInitiated reason:IsRunningCommandlet() ? @"Running commandlet" : @"Running dedicated server"];
+		// During cooking and on dedicated server, we don't want macOS to put our app into App Nap mode.
+		CommandletActivity = [[NSProcessInfo processInfo] beginActivityWithOptions:NSActivityUserInitiated reason:IsRunningCookCommandlet() ? @"Running cook commandlet" : @"Running dedicated server"];
 		[CommandletActivity retain];
 	}
 	
