[UE] Exposing Device Mapping Policy Function

diff --git a/Engine/Source/Runtime/ApplicationCore/Private/GenericPlatform/GenericPlatformInputDeviceMapper.cpp b/Engine/Source/Runtime/ApplicationCore/Private/GenericPlatform/GenericPlatformInputDeviceMapper.cpp
--- a/Engine/Source/Runtime/ApplicationCore/Private/GenericPlatform/GenericPlatformInputDeviceMapper.cpp
+++ b/Engine/Source/Runtime/ApplicationCore/Private/GenericPlatform/GenericPlatformInputDeviceMapper.cpp
@@ -553,6 +553,11 @@ bool IPlatformInputDeviceMapper::ShouldCreateUniqueUserForEachDevice() const
 		Policy == EInputDeviceMappingPolicy::UseManagedPlatformLogin;
 }
 
+const EInputDeviceMappingPolicy IPlatformInputDeviceMapper::GetCurrentDeviceMappingPolicy() const
+{
+	return UE::Input::GetDeviceMappingPolicyFromConfig();
+}
+
 ///////////////////////////////////////////////////////////////////////////
 // FGenericPlatformInputDeviceMapper
 
diff --git a/Engine/Source/Runtime/ApplicationCore/Public/GenericPlatform/GenericPlatformInputDeviceMapper.h b/Engine/Source/Runtime/ApplicationCore/Public/GenericPlatform/GenericPlatformInputDeviceMapper.h
--- a/Engine/Source/Runtime/ApplicationCore/Public/GenericPlatform/GenericPlatformInputDeviceMapper.h
+++ b/Engine/Source/Runtime/ApplicationCore/Public/GenericPlatform/GenericPlatformInputDeviceMapper.h
@@ -254,6 +254,7 @@ public:
 	/** Returns the next available input device id. This ID should be globally unique! */
 	virtual FInputDeviceId AllocateNewInputDeviceId() = 0;
 	
+	APPLICATIONCORE_API virtual const EInputDeviceMappingPolicy GetCurrentDeviceMappingPolicy() const;
 protected:
 
 	/**
