PR #9858: Reduce mouse input overhead on Windows by processing raw mouse moves on a separate FRunnable thread in the WindowsApplication.

This behavior is enabled by default, but you can go back to processing only on the main game thread with this console variable:

[ConsoleVariables]
WindowsApplication.UseWorkerThreadForRawInput=false

diff --git a/Engine/Config/Windows/BaseWindowsEngine.ini b/Engine/Config/Windows/BaseWindowsEngine.ini
--- a/Engine/Config/Windows/BaseWindowsEngine.ini
+++ b/Engine/Config/Windows/BaseWindowsEngine.ini
@@ -62,3 +62,6 @@ BindlessSamplers=Enabled
 [/Script/WindowsRuntimeSettings.WindowsRuntimeSettings]
 MaxNumTlsSlots=2048
 MaxNumThreadsWithTlsSlots=512
+
+[ConsoleVariables]
+WindowsApplication.UseWorkerThreadForRawInput=true
\ No newline at end of file
diff --git a/Engine/Source/Runtime/ApplicationCore/Private/Windows/WindowsApplication.cpp b/Engine/Source/Runtime/ApplicationCore/Private/Windows/WindowsApplication.cpp
--- a/Engine/Source/Runtime/ApplicationCore/Private/Windows/WindowsApplication.cpp
+++ b/Engine/Source/Runtime/ApplicationCore/Private/Windows/WindowsApplication.cpp
@@ -22,6 +22,8 @@
 #include "HAL/IConsoleManager.h"
 #include "Tasks/Task.h"
 #include "Tasks/Pipe.h"
+#include "Containers/LockFreeList.h"
+#include "HAL/RunnableThread.h"
 
 #if WITH_EDITOR
 #include "Modules/ModuleManager.h"
@@ -848,6 +850,393 @@ const FIntPoint FWindowsApplication::MinimizedWindowPosition(-32000,-32000);
 
 FWindowsApplication* WindowsApplication = nullptr;
 
+namespace UE::WindowsApplication::RawInputThread
+{
+	DECLARE_LOG_CATEGORY_EXTERN(LogWindowsRawInputThread, Display, VeryVerbose);
+	DEFINE_LOG_CATEGORY(LogWindowsRawInputThread);
+
+	// Lockless queue to push mouse inputs from the worker thread to the game thread
+	class FMouseQueue
+	{
+	public:
+		FMouseQueue() = default;
+
+		~FMouseQueue()
+		{
+			RAWMOUSE* Node = FreeList.Pop();
+			while (Node)
+			{
+				delete Node;
+				Node = FreeList.Pop();
+			}
+		}
+
+		bool Dequeue(RAWMOUSE& OutItem)
+		{
+			RAWMOUSE* const Node = QueuedList.Pop();
+			const bool Success = Node != nullptr;
+			if (Success)
+			{
+				OutItem = *Node;
+				FreeList.Push(Node);
+			}
+
+			return Success;
+		}
+
+		void Enqueue(const RAWMOUSE& Mouse)
+		{
+			RAWMOUSE* Node = FreeList.Pop();
+			Node = Node ? Node : new RAWMOUSE;
+			*Node = Mouse;
+			QueuedList.Push(Node);
+		}
+
+	private:
+		TLockFreePointerListUnordered<RAWMOUSE, PLATFORM_CACHE_LINE_SIZE> FreeList;
+		TLockFreePointerListFIFO<RAWMOUSE, PLATFORM_CACHE_LINE_SIZE> QueuedList;
+
+		FMouseQueue(const FMouseQueue&) = delete;
+		FMouseQueue& operator=(const FMouseQueue&) = delete;
+	};
+
+	constexpr uint32 WM_UE_RAWINPUT_QUIT = WM_APP;
+	constexpr uint32 WM_UE_RAWINPUT_REGISTER = WM_APP + 1;
+	constexpr uint32 WM_UE_RAWINPUT_UNREGISTER = WM_APP + 2;
+
+	class FWindowsRawInputRunnable : public FRunnable
+	{
+		// Somewhat arbitrary
+		static constexpr int32 RawInputBlockCount = 16;
+
+		static FRunnableThread* Thread;
+		static FWindowsRawInputRunnable Singleton;
+		static bool bInitialized;
+		static HINSTANCE InstanceHandle;
+		static HANDLE CreateWindowEvent;
+		static HANDLE UnregisterEvent;
+
+		HWND WindowHandle = NULL;
+
+		FMouseQueue MouseInputs;
+		std::atomic<uint32> InputCount;
+
+		bool bRegistered = false;
+		TArray<SIZE_T> DataArray; // GetRawInputBuffer wants DWORD/QWORD alignment for 32/64-bit architecture
+		TArray<SIZE_T> BufferArray; // GetRawInputBuffer wants DWORD/QWORD alignment for 32/64-bit architecture
+
+		FWindowsRawInputRunnable() = default;
+
+		virtual uint32 Run() override;
+		virtual void Stop() override;
+
+		static LRESULT CALLBACK WndProc(HWND hwnd, uint32 msg, WPARAM wParam, LPARAM lParam);
+		void HandleRegister();
+		void HandleUnregister();
+		void HandleInputMessage(HRAWINPUT InputHandle);
+
+	public:
+		static void Initialize(HINSTANCE HInstance);
+		static void Shutdown();
+
+		static void Register();
+		static void Unregister();
+		static uint32 TakeMouseInputCount();
+		static void PopMouseInput(RAWMOUSE& Mouse);
+	};
+
+	static bool bUseWorkerThreadForRawInput = true;
+	FAutoConsoleVariableRef	CVarUseWorkerThreadForRawInput(
+		TEXT("WindowsApplication.UseWorkerThreadForRawInput"),
+		bUseWorkerThreadForRawInput,
+		TEXT("Use a worker thread for raw input messages to remove processing from game thread."
+			" Mice with high polling rates can take multiple milliseconds per frame even on fast CPUs."
+			" When the cvar is changed a CVarSinkHandler will attempt to switch the thread used for handling inputs.")
+	);
+
+	static bool ShouldUseWorkerThread()
+	{
+		return bUseWorkerThreadForRawInput && FPlatformProcess::SupportsMultithreading();
+	}
+
+	FRunnableThread* FWindowsRawInputRunnable::Thread;
+	FWindowsRawInputRunnable FWindowsRawInputRunnable::Singleton;
+	bool FWindowsRawInputRunnable::bInitialized = false;
+	HINSTANCE FWindowsRawInputRunnable::InstanceHandle = NULL;
+	HANDLE FWindowsRawInputRunnable::CreateWindowEvent = NULL;
+	HANDLE FWindowsRawInputRunnable::UnregisterEvent = NULL;
+
+	uint32 FWindowsRawInputRunnable::Run()
+	{
+		WNDCLASS wc{};
+		wc.style = 0;
+		wc.lpfnWndProc = &WndProc;
+		wc.cbClsExtra = 0;
+		wc.cbWndExtra = 0;
+		wc.hInstance = InstanceHandle;
+		wc.hIcon = NULL;
+		wc.hCursor = NULL;
+		wc.hbrBackground = NULL;
+		wc.lpszMenuName = NULL;
+		wc.lpszClassName = L"UnrealRawInputWindow";
+		if (RegisterClassW(&wc))
+		{
+			WindowHandle = CreateWindowW(wc.lpszClassName, wc.lpszClassName, 0, 0, 0, 0, 0, HWND_MESSAGE, NULL, InstanceHandle, NULL);
+			if (WindowHandle)
+			{
+				SetEvent(CreateWindowEvent);
+
+				MSG Msg;
+				BOOL bRet;
+				while ((bRet = GetMessageW(&Msg, NULL, 0, 0)) != 0)
+				{
+					if (bRet != -1)
+					{
+						TranslateMessage(&Msg);
+						DispatchMessage(&Msg);
+					}
+				}
+
+				HandleUnregister();
+
+				DestroyWindow(WindowHandle);
+				WindowHandle = NULL;
+			}
+			else
+			{
+				UE_LOG(LogWindowsDesktop, Error, TEXT("Could not CreateWindowW for the raw input thread. Error code: %u"), GetLastError());
+			}
+
+			UnregisterClassW(wc.lpszClassName, InstanceHandle);
+		}
+		else
+		{
+			UE_LOG(LogWindowsDesktop, Error, TEXT("Could not RegisterClassW for the raw input thread. Error code: %u"), GetLastError());
+		}
+
+		// Unblocks game thread on failure, redundant on success
+		SetEvent(CreateWindowEvent);
+
+		return 0;
+	}
+
+	void FWindowsRawInputRunnable::Stop()
+	{
+		SendMessageW(Singleton.WindowHandle, WM_UE_RAWINPUT_QUIT, 0, 0);
+	}
+
+	LRESULT CALLBACK FWindowsRawInputRunnable::WndProc(HWND hWnd, uint32 Msg, WPARAM wParam, LPARAM lParam)
+	{
+		switch (Msg)
+		{
+		case WM_INPUT:
+			Singleton.HandleInputMessage((HRAWINPUT)lParam);
+			// RIM_INPUTSINK should not appear since we don't use RIDEV_INPUTSINK
+			// Call DefWindowProc conditionally for RIM_INPUT if this assert trips somehow though
+			check(GET_RAWINPUT_CODE_WPARAM(wParam) == RIM_INPUT);
+			// THe documentation says DefWindowProc is needed for RIM_INPUT cleanup
+			// The OS may be forgiving but don't rely on that
+			DefWindowProc(hWnd, Msg, wParam, lParam);
+			// Ignore DefWindowProc return value since the documentation says to return 0 if the application handles the message
+			return 0;
+		case WM_UE_RAWINPUT_QUIT:
+			PostQuitMessage(0);
+			return 0;
+		case WM_UE_RAWINPUT_REGISTER:
+			FWindowsRawInputRunnable::Singleton.HandleRegister();
+			return 0;
+		case WM_UE_RAWINPUT_UNREGISTER:
+			FWindowsRawInputRunnable::Singleton.HandleUnregister();
+			return 0;
+		default:;
+		}
+
+		return DefWindowProc(hWnd, Msg, wParam, lParam);
+	}
+
+	void FWindowsRawInputRunnable::HandleRegister()
+	{
+		if (!bRegistered)
+		{
+			RAWINPUTDEVICE RawInputDevice;
+			RawInputDevice.usUsagePage = 0x01; // HID_USAGE_PAGE_GENERIC
+			RawInputDevice.usUsage = 0x02; // HID_USAGE_GENERIC_MOUSE
+			RawInputDevice.dwFlags = RIDEV_NOLEGACY;
+			RawInputDevice.hwndTarget = WindowHandle;
+			bRegistered = ::RegisterRawInputDevices(&RawInputDevice, 1, sizeof(RAWINPUTDEVICE)) == TRUE;
+			UE_LOG(LogWindowsRawInputThread, Verbose, TEXT("WindowsApplication.UseWorkerThreadForRawInput CVarSinkHandler registering worker thread for input: bRegistered=%d"), (int)bRegistered);
+		}
+	}
+
+	void FWindowsRawInputRunnable::HandleUnregister()
+	{
+		if (bRegistered)
+		{
+			RAWINPUTDEVICE RawInputDevice;
+			RawInputDevice.usUsagePage = 0x01; // HID_USAGE_PAGE_GENERIC
+			RawInputDevice.usUsage = 0x02; // HID_USAGE_GENERIC_MOUSE
+			RawInputDevice.dwFlags = RIDEV_REMOVE;
+			RawInputDevice.hwndTarget = NULL;
+			bRegistered = ::RegisterRawInputDevices(&RawInputDevice, 1, sizeof(RAWINPUTDEVICE)) != TRUE;
+			UE_LOG(LogWindowsRawInputThread, Verbose, TEXT("WindowsApplication.UseWorkerThreadForRawInput CVarSinkHandler unregistering worker thread for input: bRegistered=%d"), (int)bRegistered);
+		}
+
+		SetEvent(UnregisterEvent);
+	}
+
+	void FWindowsRawInputRunnable::HandleInputMessage(HRAWINPUT InputHandle)
+	{
+		// Lazy-allocate because HRAWINPUT might be necessary to get the size
+		if (DataArray.IsEmpty())
+		{
+			UINT Size = 0;
+			UINT Result = ::GetRawInputData(InputHandle, RID_INPUT, NULL, &Size, sizeof(RAWINPUTHEADER));
+			if (Result != 0)
+			{
+				UE_LOG(LogWindowsDesktop, Error, TEXT("Could not get the required size for raw input data. Return value: %u"), Result);
+				return;
+			}
+
+			DataArray.AddDefaulted((Size + sizeof(SIZE_T) - 1) / sizeof(SIZE_T));
+			BufferArray.AddDefaulted(((RawInputBlockCount * Size) + sizeof(SIZE_T) - 1) / sizeof(SIZE_T));
+		}
+
+		// Relying only on GetRawInputBuffer seems to lead to missing inputs
+		PRAWINPUT DataRaw = (PRAWINPUT)DataArray.GetData();
+		UINT DataSize = DataArray.Num() * sizeof(SIZE_T);
+		const UINT DataResult = GetRawInputData(InputHandle, RID_INPUT, DataRaw, &DataSize, sizeof(RAWINPUTHEADER));
+		if ((DataResult != UINT_MAX) && (DataResult > 0))
+		{
+			MouseInputs.Enqueue(DataRaw->data.mouse);
+			++InputCount;
+		}
+
+		PRAWINPUT BufferRaw = (PRAWINPUT)BufferArray.GetData();
+		UINT BufferSize = BufferArray.Num() * sizeof(SIZE_T);
+		UINT BufferResult = GetRawInputBuffer(BufferRaw, &BufferSize, sizeof(RAWINPUTHEADER));
+		while ((BufferResult != UINT_MAX) && (BufferResult > 0))
+		{
+			for (UINT Index = 0; Index < BufferResult; ++Index)
+			{
+				if (BufferRaw->header.dwType == RIM_TYPEMOUSE)
+				{
+					MouseInputs.Enqueue(BufferRaw->data.mouse);
+					++InputCount;
+				}
+
+				typedef ULONGLONG QWORD;
+				BufferRaw = NEXTRAWINPUTBLOCK(BufferRaw);
+			}
+
+			if (BufferResult < RawInputBlockCount)
+			{
+				break;
+			}
+
+			BufferResult = GetRawInputBuffer((PRAWINPUT)BufferRaw, &BufferSize, sizeof(RAWINPUTHEADER));
+		}
+	}
+
+	void FWindowsRawInputRunnable::Initialize(HINSTANCE HInstance)
+	{
+		ensure(IsInGameThread());
+
+		if (!bInitialized)
+		{
+			check(Thread == nullptr);
+
+			CreateWindowEvent = CreateEventW(NULL, TRUE, FALSE, NULL);
+			UnregisterEvent = CreateEventW(NULL, FALSE, FALSE, NULL);
+			if (CreateWindowEvent && UnregisterEvent)
+			{
+				InstanceHandle = HInstance;
+
+				// Imitate DirectInput which also uses THREAD_PRIORITY_HIGHEST for its input thread
+				Thread = FRunnableThread::Create(&Singleton, TEXT("WindowsRawInputThread"), 0, TPri_Highest);
+
+				// Don't read WindowHandle until this event clears
+				WaitForSingleObject(CreateWindowEvent, INFINITE);
+
+				if (Singleton.WindowHandle)
+				{
+					bInitialized = true;
+				}
+				else
+				{
+					CloseHandle(UnregisterEvent);
+					UnregisterEvent = NULL;
+					CloseHandle(CreateWindowEvent);
+					CreateWindowEvent = NULL;
+				}
+			}
+			else
+			{
+				UE_LOG(LogWindowsDesktop, Error, TEXT("Could not events for the raw input thread. Error code: %u"), GetLastError());
+
+				if (UnregisterEvent)
+				{
+					CloseHandle(UnregisterEvent);
+					UnregisterEvent = NULL;
+				}
+
+				if (CreateWindowEvent)
+				{
+					CloseHandle(CreateWindowEvent);
+					CreateWindowEvent = NULL;
+				}
+			}
+		}
+	}
+
+	void FWindowsRawInputRunnable::Shutdown()
+	{
+		ensure(IsInGameThread());
+
+		if (bInitialized)
+		{
+			bInitialized = false;
+
+			if (Thread)
+			{
+				Thread->Kill(true);
+				delete Thread;
+				Thread = nullptr;
+			}
+
+			CloseHandle(CreateWindowEvent);
+			CreateWindowEvent = NULL;
+
+			InstanceHandle = NULL;
+		}
+	}
+
+	void FWindowsRawInputRunnable::Register()
+	{
+		ensure(IsInGameThread());
+
+		SendMessageW(Singleton.WindowHandle, WM_UE_RAWINPUT_REGISTER, 0, 0);
+	}
+
+	void FWindowsRawInputRunnable::Unregister()
+	{
+		ensure(IsInGameThread());
+
+		SendMessageW(Singleton.WindowHandle, WM_UE_RAWINPUT_UNREGISTER, 0, 0);
+		WaitForSingleObject(UnregisterEvent, INFINITE);
+	}
+
+	uint32 FWindowsRawInputRunnable::TakeMouseInputCount()
+	{
+		return Thread ? Singleton.InputCount.exchange(0) : 0;
+	}
+
+	void FWindowsRawInputRunnable::PopMouseInput(RAWMOUSE& Mouse)
+	{
+		ensure(IsInGameThread());
+		Singleton.MouseInputs.Dequeue(Mouse);
+	}
+}
+
 namespace UE::WindowsInput
 {
 	static const FName InputClassName = TEXT("WindowsApplication");
@@ -872,6 +1261,8 @@ FWindowsApplication::FWindowsApplication( const HINSTANCE HInstance, const HICON
 	, InstanceHandle( HInstance )
 	, bMinimized( false )
 	, bUsingHighPrecisionMouseInput( false )
+	, bCachedUseWorkerThreadForRawInput( false )
+	, HighPrecisionMouseInputWindowHandle( NULL )
 	, bIsMouseAttached( false )
 	, bForceActivateByMouse( false )
 	, bForceNoGamepads( false )
@@ -892,7 +1283,7 @@ FWindowsApplication::FWindowsApplication( const HINSTANCE HInstance, const HICON
 	, LastCursorPointPreWrap(FIntPoint::ZeroValue)
 	, NumPreWrapMsgsToRespect(0)
 	, ClipCursorRect()
-
+	, RepeatTime{}
 {
 	FMemory::Memzero(ModifierKeyState, EModifierKey::Count);
 
@@ -993,6 +1384,11 @@ FWindowsApplication::FWindowsApplication( const HINSTANCE HInstance, const HICON
 	IConsoleManager::Get().RegisterConsoleCommand(TEXT("WindowsApplication.RemoveLowLevelMouseFilter"), TEXT("Removes Low Level mouse filter that filters out mouse inputs that act like touch inputs"), FConsoleCommandDelegate::CreateRaw(this, &FWindowsApplication::RemoveLowLevelMouseFilter));
 
 	UE::WindowsApplication::TouchpadSensitivity::HandleStartupApp();
+
+	if (FPlatformProcess::SupportsMultithreading())
+	{
+		UE::WindowsApplication::RawInputThread::FWindowsRawInputRunnable::Initialize(HInstance);
+	}
 }
 
 bool FWindowsApplication::SupportsForceMaxTouchpadSensitivity()
@@ -1139,6 +1535,11 @@ bool FWindowsApplication::RegisterClass( const HINSTANCE HInstance, const HICON
 
 FWindowsApplication::~FWindowsApplication()
 {
+	if (FPlatformProcess::SupportsMultithreading())
+	{
+		UE::WindowsApplication::RawInputThread::FWindowsRawInputRunnable::Shutdown();
+	}
+
 	FCoreDelegates::OnShutdownAfterError.RemoveAll(this);
 
 	if (TextInputMethodSystem.IsValid())
@@ -1305,38 +1706,77 @@ void FWindowsApplication::SetHighPrecisionMouseMode( const bool Enable, const TS
 		}
 	}
 
-	{
-		
-		HWND hwnd = NULL;
-		DWORD flags = RIDEV_REMOVE;
-		bUsingHighPrecisionMouseInput = Enable;
-
-		if (Enable)
-		{
-			flags = 0;
+	bUsingHighPrecisionMouseInput = Enable;
+	HighPrecisionMouseInputWindowHandle = (Enable && InWindow.IsValid()) ? (HWND)InWindow->GetOSWindowHandle() : NULL;
 
-			if (InWindow.IsValid())
-			{
-				hwnd = (HWND)InWindow->GetOSWindowHandle();
-			}
-		}
+	// NOTE: Currently has to be created every time due to conflicts with Direct8 Input used by the wx unrealed
+	if (!Enable || (bCachedUseWorkerThreadForRawInput != UE::WindowsApplication::RawInputThread::ShouldUseWorkerThread()))
+	{
+		UnregisterHighPrecisionMouseMode();
+	}
+	if (Enable)
+	{
+		RegisterHighPrecisionMouseMode();
+	}
+}
 
-		// NOTE: Currently has to be created every time due to conflicts with Direct8 Input used by the wx unrealed
+void FWindowsApplication::RegisterHighPrecisionMouseMode()
+{
+	const bool bUseWorkerThreadForRawInputLocal = UE::WindowsApplication::RawInputThread::ShouldUseWorkerThread();
+	if (bUseWorkerThreadForRawInputLocal)
+	{
+		UE::WindowsApplication::RawInputThread::FWindowsRawInputRunnable::Register();
+	}
+	else
+	{
 		RAWINPUTDEVICE RawInputDevice;
+		RawInputDevice.usUsagePage = 0x01; // HID_USAGE_PAGE_GENERIC
+		RawInputDevice.usUsage = 0x02; // HID_USAGE_GENERIC_MOUSE
+		RawInputDevice.dwFlags = 0;
+		RawInputDevice.hwndTarget = HighPrecisionMouseInputWindowHandle;
+		const BOOL bRegistered = ::RegisterRawInputDevices(&RawInputDevice, 1, sizeof(RAWINPUTDEVICE));
+		UE_LOG(LogWindowsDesktop, Verbose, TEXT("Registering game thread for raw input: %d"), bRegistered);
+	}
 
-		//The HID standard for mouse
-		const uint16 StandardMouse = 0x02;
+	bCachedUseWorkerThreadForRawInput = bUseWorkerThreadForRawInputLocal;
+}
 
-		RawInputDevice.usUsagePage = 0x01;
-		RawInputDevice.usUsage = StandardMouse;
-		RawInputDevice.dwFlags = flags;
+void FWindowsApplication::UnregisterHighPrecisionMouseMode()
+{
+	if (bCachedUseWorkerThreadForRawInput)
+	{
+		UE::WindowsApplication::RawInputThread::FWindowsRawInputRunnable::Unregister();
 
-		// Process input for just the window that requested it.  NOTE: If we pass NULL here events are routed to the window with keyboard focus
-		// which is not always known at the HWND level with Slate
-		RawInputDevice.hwndTarget = hwnd;
+		// Flush inputs from worker thread
+		ProcessWorkerInputs();
+	}
+	else
+	{
+		RAWINPUTDEVICE RawInputDevice;
+		RawInputDevice.usUsagePage = 0x01; // HID_USAGE_PAGE_GENERIC
+		RawInputDevice.usUsage = 0x02; // HID_USAGE_GENERIC_MOUSE
+		RawInputDevice.dwFlags = RIDEV_REMOVE;
+		RawInputDevice.hwndTarget = NULL;
+		const BOOL bUnregistered = ::RegisterRawInputDevices(&RawInputDevice, 1, sizeof(RAWINPUTDEVICE));
+		UE_LOG(LogWindowsDesktop, Verbose, TEXT("Unregistering game thread for raw input: %d"), bUnregistered);
+
+		// Flush inputs on window thread (and other Windows messages)
+		MSG Message;
+		while (PeekMessage(&Message, NULL, 0, 0, PM_REMOVE))
+		{
+			TranslateMessage(&Message);
+			DispatchMessage(&Message);
+		}
+	}
+}
 
-		// Register the raw input device
-		::RegisterRawInputDevices(&RawInputDevice, 1, sizeof(RAWINPUTDEVICE));
+void FWindowsApplication::UpdateHighPrecisionMouseMode()
+{
+	if (bUsingHighPrecisionMouseInput &&
+		(bCachedUseWorkerThreadForRawInput != UE::WindowsApplication::RawInputThread::ShouldUseWorkerThread()))
+	{
+		UnregisterHighPrecisionMouseMode();
+		RegisterHighPrecisionMouseMode();
 	}
 }
 
@@ -1931,115 +2371,16 @@ int32 FWindowsApplication::ProcessMessage( HWND hwnd, uint32 msg, WPARAM wParam,
 
 					if (Raw->header.dwType == RIM_TYPEMOUSE) 
 					{
-						const bool IsAbsoluteInput = (Raw->data.mouse.usFlags & MOUSE_MOVE_ABSOLUTE) == MOUSE_MOVE_ABSOLUTE;
-						if (IsAbsoluteInput)
+						RawMouseInputResult Result;
+						ProcessRawMouseInput(Raw->data.mouse, Result);
+						if (Result.bIsAbsoluteInput)
 						{
 							if (bSimulatingHighPrecisionMouseInputForRDP)
-							{	
-								const bool IsVirtualScreen = (Raw->data.mouse.usFlags & MOUSE_VIRTUAL_DESKTOP) == MOUSE_VIRTUAL_DESKTOP;
-
-								// Get the new cursor position
-								POINT CursorPoint;
-								const int32 Left = IsVirtualScreen ? GetSystemMetrics(SM_XVIRTUALSCREEN) : 0;
-								const int32 Top = IsVirtualScreen ? GetSystemMetrics(SM_YVIRTUALSCREEN) : 0;
-								const int32 Width = GetSystemMetrics(IsVirtualScreen ? SM_CXVIRTUALSCREEN : SM_CXSCREEN);
-								const int32 Height = GetSystemMetrics(IsVirtualScreen ? SM_CYVIRTUALSCREEN : SM_CYSCREEN);
-
-								CursorPoint.x = static_cast<int>((float(Raw->data.mouse.lLastX) / 65535.0f) * (float)Width) + Left;
-								CursorPoint.y = static_cast<int>((float(Raw->data.mouse.lLastY) / 65535.0f) * (float)Height) + Top;
-
-								const int32 ClipWidth = ClipCursorRect.right - ClipCursorRect.left;
-								const int32 ClipHeight = ClipCursorRect.bottom - ClipCursorRect.top;
-
-								const int32 DeltaWidthMax = (int32)((float)ClipWidth * 0.4f);
-								const int32 DeltaHeightMax = (int32)((float)ClipHeight * 0.4f);
-
-								const POINT CursorPointNoWrap = CursorPoint;
-
-								const bool bCanAcceptPreWrapMsg = NumPreWrapMsgsToRespect > 0;
-								NumPreWrapMsgsToRespect--;
-
-								// Calculate the cursor delta from the last position
-								// We do this prior to and wrapping for continuous input as that would mean an incorrect delta we'd just ignore
-								int32 DeltaX = CursorPoint.x - LastCursorPoint.X;
-								int32 DeltaY = CursorPoint.y - LastCursorPoint.Y;
-
-								// Always skip wrapping if no movement has occured.
-								bool bAcceptingPreWrapDelta = DeltaX == 0 && DeltaY == 0;
-
-								if (bCanAcceptPreWrapMsg && !bAcceptingPreWrapDelta)
-								{
-									const int32 DeltaXPreWrap = CursorPoint.x - LastCursorPointPreWrap.X;
-									const int32 DeltaYPreWrap = CursorPoint.y - LastCursorPointPreWrap.Y;
-
-									const int32 DeltaLen = (DeltaX * DeltaX) + (DeltaY * DeltaY);
-									const int32 DeltaPreWrapLen = (DeltaXPreWrap * DeltaXPreWrap) + (DeltaYPreWrap * DeltaYPreWrap);
-
-									if (DeltaPreWrapLen < DeltaLen)
-									{
-										bAcceptingPreWrapDelta = true;
-
-										DeltaX = DeltaXPreWrap;
-										DeltaY = DeltaYPreWrap;
-										LastCursorPointPreWrap.X = CursorPoint.x;
-										LastCursorPointPreWrap.Y = CursorPoint.y;
-
-										//UE_LOG(LogWindowsDesktop, Log, TEXT("Accept PreWrap Delta X: %d Y: %d  ---- Last X: %d Y: %d ---- Cur X: %d Y: %d ---- Width: %d Height: %d"), DeltaX, DeltaY, LastCursorPointPreWrap.X, LastCursorPointPreWrap.Y, CursorPoint.x, CursorPoint.y, (int32)Width, (int32)Height);
-									}
-								}
-
-								if (!bAcceptingPreWrapDelta)
-								{
-									// Wrap and set cursor position in necessary
-									const int32 WrapLeeway = 50; // We add some leeway on the wrap so that if the user is doing small movements hear the border we don't wrap back and fourth constantly
-									const int32 TopEdge = ClipCursorRect.top + int32(0.1f * float(ClipHeight));
-									const int32 BottomEdge = ClipCursorRect.top + int32(0.9f * float(ClipHeight));
-									const int32 LeftEdge = ClipCursorRect.left + int32(0.1f * float(ClipWidth));
-									const int32 RightEdge = ClipCursorRect.left + int32(0.9f * float(ClipWidth));
-
-									bool bSet = false;
-									if (CursorPoint.y < TopEdge) { CursorPoint.y = BottomEdge - WrapLeeway;	bSet = true; }
-									else if (CursorPoint.y > BottomEdge) { CursorPoint.y = TopEdge + WrapLeeway; bSet = true; }
-
-									if (CursorPoint.x < LeftEdge) { CursorPoint.x = RightEdge - WrapLeeway;	bSet = true; }
-									else if (CursorPoint.x > RightEdge) { CursorPoint.x = LeftEdge + WrapLeeway; bSet = true; }
-
-									if (bSet)
-									{
-										//UE_LOG(LogWindowsDesktop, Log, TEXT("Wrapping Cursor to X: %d Y: %d ---- TopEdge: %d BottomEdge: %d LeftEdge: %d RightEdge: %d "), CursorPoint.x, CursorPoint.y, TopEdge, BottomEdge, LeftEdge, RightEdge);
-
-										MessageHandler->SetCursorPos(FVector2D(CursorPoint.x,CursorPoint.y));
-										LastCursorPoint.X = CursorPoint.x;
-										LastCursorPoint.Y = CursorPoint.y;
-
-										NumPreWrapMsgsToRespect = 10;
-										LastCursorPointPreWrap.X = CursorPointNoWrap.x;
-										LastCursorPointPreWrap.Y = CursorPointNoWrap.y;
-									}
-									
-									/*
-									if (DeltaX != 0 || DeltaY != 0)
-									{
-										if (FMath::Abs(DeltaX) < DeltaWidthMax && FMath::Abs(DeltaY) < DeltaHeightMax)
-										{
-											UE_LOG(LogWindowsDesktop, Log, TEXT("Accept Delta X: %d Y: %d  ---- Last X: %d Y: %d ---- Cur X: %d Y: %d ---- Width: %d Height: %d"), DeltaX, DeltaY, LastCursorPoint.X, LastCursorPoint.Y, CursorPoint.x, CursorPoint.y, (int32)Width, (int32)Height);
-										}
-										else
-										{
-											UE_LOG(LogWindowsDesktop, Log, TEXT("IGNORE Delta X: %d Y: %d  ---- Last X: %d Y: %d ---- Cur X: %d Y: %d ---- Width: %d Height: %d"), DeltaX, DeltaY, LastCursorPoint.X, LastCursorPoint.Y, CursorPoint.x, CursorPoint.y, (int32)Width, (int32)Height);
-										}
-									}*/
-
-									LastCursorPoint.X = CursorPoint.x;
-									LastCursorPoint.Y = CursorPoint.y;
-								}
-
-								// Send a delta assuming it's not zero or beyond our max delta 
-								if ((DeltaX != 0 || DeltaY != 0) && FMath::Abs(DeltaX) < DeltaWidthMax && FMath::Abs(DeltaY) < DeltaHeightMax)
+							{
+								if (Result.bSimulateRelativeMovement)
 								{
-									DeferMessage(CurrentNativeEventWindowPtr, hwnd, msg, wParam, lParam, DeltaX, DeltaY, MOUSE_MOVE_RELATIVE);
+									DeferMessage(CurrentNativeEventWindowPtr, hwnd, msg, wParam, lParam, Result.DeltaX, Result.DeltaY, MOUSE_MOVE_RELATIVE);
 								}
-								return 1;
 							}
 							else
 							{
@@ -2051,8 +2392,8 @@ int32 FWindowsApplication::ProcessMessage( HWND hwnd, uint32 msg, WPARAM wParam,
 						}
 
 						// Since raw input is coming in as relative it is likely a traditional mouse device
-						const int xPosRelative = Raw->data.mouse.lLastX;
-						const int yPosRelative = Raw->data.mouse.lLastY;
+						const int xPosRelative = Result.DeltaX;
+						const int yPosRelative = Result.DeltaY;
 
 						DeferMessage( CurrentNativeEventWindowPtr, hwnd, msg, wParam, lParam, xPosRelative, yPosRelative, MOUSE_MOVE_RELATIVE );
 						return 1;
@@ -3581,6 +3922,8 @@ void FWindowsApplication::PumpMessages( const float TimeDelta )
 		TranslateMessage(&Message);
 		DispatchMessage(&Message); 
 	}
+
+	PostPumpMessages();
 }
 
 void FWindowsApplication::ProcessDeferredEvents( const float TimeDelta )
@@ -3978,6 +4321,265 @@ TSharedRef<FTaskbarList> FTaskbarList::Create()
 	return TaskbarList;
 }
 
+void FWindowsApplication::StaticPostPumpMessages()
+{
+	if (WindowsApplication)
+	{
+		WindowsApplication->PostPumpMessages();
+	}
+}
+
+void FWindowsApplication::PostPumpMessages()
+{
+	WindowsApplication->UpdateHighPrecisionMouseMode();
+	WindowsApplication->ProcessWorkerInputs();
+}
+
+void FWindowsApplication::ProcessRawMouseInput(const RAWMOUSE& Mouse, RawMouseInputResult& OutResult)
+{
+	OutResult.bIsAbsoluteInput = (Mouse.usFlags & MOUSE_MOVE_ABSOLUTE) == MOUSE_MOVE_ABSOLUTE;
+	if (OutResult.bIsAbsoluteInput)
+	{
+		if (bSimulatingHighPrecisionMouseInputForRDP)
+		{
+			const bool IsVirtualScreen = (Mouse.usFlags & MOUSE_VIRTUAL_DESKTOP) == MOUSE_VIRTUAL_DESKTOP;
+
+			// Get the new cursor position
+			POINT CursorPoint;
+			const int32 Left = IsVirtualScreen ? GetSystemMetrics(SM_XVIRTUALSCREEN) : 0;
+			const int32 Top = IsVirtualScreen ? GetSystemMetrics(SM_YVIRTUALSCREEN) : 0;
+			const int32 Width = GetSystemMetrics(IsVirtualScreen ? SM_CXVIRTUALSCREEN : SM_CXSCREEN);
+			const int32 Height = GetSystemMetrics(IsVirtualScreen ? SM_CYVIRTUALSCREEN : SM_CYSCREEN);
+
+			CursorPoint.x = (LONG)(((LONGLONG)Mouse.lLastX * (LONGLONG)Width) / 65535LL) + Left;
+			CursorPoint.y = (LONG)(((LONGLONG)Mouse.lLastY * (LONGLONG)Height) / 65535LL) + Top;
+
+			const int32 ClipWidth = ClipCursorRect.right - ClipCursorRect.left;
+			const int32 ClipHeight = ClipCursorRect.bottom - ClipCursorRect.top;
+
+			const int32 DeltaWidthMax = (int32)((float)ClipWidth * 0.4f);
+			const int32 DeltaHeightMax = (int32)((float)ClipHeight * 0.4f);
+
+			const POINT CursorPointNoWrap = CursorPoint;
+
+			const bool bCanAcceptPreWrapMsg = NumPreWrapMsgsToRespect > 0;
+			NumPreWrapMsgsToRespect--;
+
+			// Calculate the cursor delta from the last position
+			// We do this prior to and wrapping for continuous input as that would mean an incorrect delta we'd just ignore
+			int32 DeltaX = CursorPoint.x - LastCursorPoint.X;
+			int32 DeltaY = CursorPoint.y - LastCursorPoint.Y;
+
+			// Always skip wrapping if no movement has occured.
+			bool bAcceptingPreWrapDelta = DeltaX == 0 && DeltaY == 0;
+
+			if (bCanAcceptPreWrapMsg && !bAcceptingPreWrapDelta)
+			{
+				const int32 DeltaXPreWrap = CursorPoint.x - LastCursorPointPreWrap.X;
+				const int32 DeltaYPreWrap = CursorPoint.y - LastCursorPointPreWrap.Y;
+
+				const int32 DeltaLen = (DeltaX * DeltaX) + (DeltaY * DeltaY);
+				const int32 DeltaPreWrapLen = (DeltaXPreWrap * DeltaXPreWrap) + (DeltaYPreWrap * DeltaYPreWrap);
+
+				if (DeltaPreWrapLen < DeltaLen)
+				{
+					bAcceptingPreWrapDelta = true;
+
+					DeltaX = DeltaXPreWrap;
+					DeltaY = DeltaYPreWrap;
+					LastCursorPointPreWrap.X = CursorPoint.x;
+					LastCursorPointPreWrap.Y = CursorPoint.y;
+
+					//UE_LOG(LogWindowsDesktop, Log, TEXT("Accept PreWrap Delta X: %d Y: %d  ---- Last X: %d Y: %d ---- Cur X: %d Y: %d ---- Width: %d Height: %d"), DeltaX, DeltaY, LastCursorPointPreWrap.X, LastCursorPointPreWrap.Y, CursorPoint.x, CursorPoint.y, (int32)Width, (int32)Height);
+				}
+			}
+
+			if (!bAcceptingPreWrapDelta)
+			{
+				// Wrap and set cursor position in necessary
+				const int32 WrapLeeway = 50; // We add some leeway on the wrap so that if the user is doing small movements hear the border we don't wrap back and fourth constantly
+				const int32 TopEdge = ClipCursorRect.top + int32(0.1f * float(ClipHeight));
+				const int32 BottomEdge = ClipCursorRect.top + int32(0.9f * float(ClipHeight));
+				const int32 LeftEdge = ClipCursorRect.left + int32(0.1f * float(ClipWidth));
+				const int32 RightEdge = ClipCursorRect.left + int32(0.9f * float(ClipWidth));
+
+				bool bSet = false;
+				if (CursorPoint.y < TopEdge) { CursorPoint.y = BottomEdge - WrapLeeway;	bSet = true; }
+				else if (CursorPoint.y > BottomEdge) { CursorPoint.y = TopEdge + WrapLeeway; bSet = true; }
+
+				if (CursorPoint.x < LeftEdge) { CursorPoint.x = RightEdge - WrapLeeway;	bSet = true; }
+				else if (CursorPoint.x > RightEdge) { CursorPoint.x = LeftEdge + WrapLeeway; bSet = true; }
+
+				if (bSet)
+				{
+					//UE_LOG(LogWindowsDesktop, Log, TEXT("Wrapping Cursor to X: %d Y: %d ---- TopEdge: %d BottomEdge: %d LeftEdge: %d RightEdge: %d "), CursorPoint.x, CursorPoint.y, TopEdge, BottomEdge, LeftEdge, RightEdge);
+
+					MessageHandler->SetCursorPos(FVector2D(CursorPoint.x, CursorPoint.y));
+					LastCursorPoint.X = CursorPoint.x;
+					LastCursorPoint.Y = CursorPoint.y;
+
+					NumPreWrapMsgsToRespect = 10;
+					LastCursorPointPreWrap.X = CursorPointNoWrap.x;
+					LastCursorPointPreWrap.Y = CursorPointNoWrap.y;
+				}
+
+				/*
+				if (DeltaX != 0 || DeltaY != 0)
+				{
+					if (FMath::Abs(DeltaX) < DeltaWidthMax && FMath::Abs(DeltaY) < DeltaHeightMax)
+					{
+						UE_LOG(LogWindowsDesktop, Log, TEXT("Accept Delta X: %d Y: %d  ---- Last X: %d Y: %d ---- Cur X: %d Y: %d ---- Width: %d Height: %d"), DeltaX, DeltaY, LastCursorPoint.X, LastCursorPoint.Y, CursorPoint.x, CursorPoint.y, (int32)Width, (int32)Height);
+					}
+					else
+					{
+						UE_LOG(LogWindowsDesktop, Log, TEXT("IGNORE Delta X: %d Y: %d  ---- Last X: %d Y: %d ---- Cur X: %d Y: %d ---- Width: %d Height: %d"), DeltaX, DeltaY, LastCursorPoint.X, LastCursorPoint.Y, CursorPoint.x, CursorPoint.y, (int32)Width, (int32)Height);
+					}
+				}*/
+
+				LastCursorPoint.X = CursorPoint.x;
+				LastCursorPoint.Y = CursorPoint.y;
+			}
+
+			// Send a delta assuming it's not zero or beyond our max delta
+			OutResult.bSimulateRelativeMovement = (DeltaX != 0 || DeltaY != 0) && FMath::Abs(DeltaX) < DeltaWidthMax && FMath::Abs(DeltaY) < DeltaHeightMax;
+			OutResult.DeltaX = DeltaX;
+			OutResult.DeltaY = DeltaY;
+		}
+	}
+	else
+	{
+		// Since raw input is coming in as relative it is likely a traditional mouse device
+		OutResult.DeltaX = Mouse.lLastX;
+		OutResult.DeltaY = Mouse.lLastY;
+	}
+}
+
+void FWindowsApplication::ProcessWorkerInputs()
+{
+	TRACE_CPUPROFILER_EVENT_SCOPE(ProcessWorkerInputs);
+	FInputDeviceScope InputScope(nullptr, UE::WindowsInput::InputClassName, IPlatformInputDeviceMapper::Get().GetDefaultInputDevice().GetId(), UE::WindowsInput::KBMInputHardwareName);
+
+	int32 xPosRelative = 0;
+	int32 yPosRelative = 0;
+
+	const auto FlushRelativeMovement =
+		[this, &xPosRelative, &yPosRelative]()
+		{
+			if (xPosRelative || yPosRelative)
+			{
+				MessageHandler->OnRawMouseMove(xPosRelative, yPosRelative);
+				xPosRelative = 0;
+				yPosRelative = 0;
+			}
+		};
+
+	for (uint32 InputIndex = 0, InputCount = UE::WindowsApplication::RawInputThread::FWindowsRawInputRunnable::TakeMouseInputCount(); InputIndex < InputCount; ++InputIndex)
+	{
+		RAWMOUSE Mouse;
+		UE::WindowsApplication::RawInputThread::FWindowsRawInputRunnable::PopMouseInput(Mouse);
+		RawMouseInputResult Result;
+		ProcessRawMouseInput(Mouse, Result);
+
+		if (Result.bIsAbsoluteInput)
+		{
+			if (bSimulatingHighPrecisionMouseInputForRDP)
+			{
+				if (Result.bSimulateRelativeMovement)
+				{
+					FlushRelativeMovement();
+
+					MessageHandler->OnRawMouseMove(Result.DeltaX, Result.DeltaY);
+				}
+			}
+			else
+			{
+				FlushRelativeMovement();
+
+				// Since the raw input is coming in as absolute it is likely the user is using a tablet
+				// or perhaps is interacting through a virtual desktop
+				MessageHandler->OnMouseMove();
+			}
+		}
+		else
+		{
+			xPosRelative += Result.DeltaX;
+			yPosRelative += Result.DeltaY;
+		}
+
+		const USHORT ButtonFlags = Mouse.usButtonFlags;
+		if (ButtonFlags)
+		{
+			FlushRelativeMovement();
+
+			TSharedPtr< FWindowsWindow > Window;
+
+			POINT ButtonCursorPoint;
+			BOOL bGotPoint;
+			{
+				bGotPoint = ::GetCursorPos(&ButtonCursorPoint);
+			}
+
+			if (bGotPoint)
+			{
+				HWND hWnd = ::WindowFromPoint(ButtonCursorPoint);
+
+				if (hWnd)
+				{
+					Window = FindWindowByHWND(Windows, hWnd);
+				}
+			}
+			else
+			{
+				UE_LOG(LogWindowsDesktop, Error, TEXT("Could not GetCursorPos for raw input. Error code: %u"), GetLastError());
+				ButtonCursorPoint.x = 0;
+				ButtonCursorPoint.y = 0;
+			}
+
+			const FVector2D ButtonCursorPos = FVector2D(ButtonCursorPoint.x, ButtonCursorPoint.y);
+			const double CurrentTime = FPlatformTime::Seconds();
+
+			ProcessMouseButton(CurrentTime, Window, ButtonFlags, RI_MOUSE_LEFT_BUTTON_DOWN, RI_MOUSE_LEFT_BUTTON_UP, EMouseButtons::Left, ButtonCursorPos);
+			ProcessMouseButton(CurrentTime, Window, ButtonFlags, RI_MOUSE_RIGHT_BUTTON_DOWN, RI_MOUSE_RIGHT_BUTTON_UP, EMouseButtons::Right, ButtonCursorPos);
+			ProcessMouseButton(CurrentTime, Window, ButtonFlags, RI_MOUSE_MIDDLE_BUTTON_DOWN, RI_MOUSE_MIDDLE_BUTTON_UP, EMouseButtons::Middle, ButtonCursorPos);
+			ProcessMouseButton(CurrentTime, Window, ButtonFlags, RI_MOUSE_BUTTON_4_DOWN, RI_MOUSE_BUTTON_4_UP, EMouseButtons::Thumb01, ButtonCursorPos);
+			ProcessMouseButton(CurrentTime, Window, ButtonFlags, RI_MOUSE_BUTTON_5_DOWN, RI_MOUSE_BUTTON_5_UP, EMouseButtons::Thumb02, ButtonCursorPos);
+
+			if (ButtonFlags & RI_MOUSE_WHEEL)
+			{
+				SHORT SignedDistance;
+				memcpy(&SignedDistance, &Mouse.usButtonData, sizeof(SignedDistance));
+				MessageHandler->OnMouseWheel(static_cast<float>(SignedDistance) / WHEEL_DELTA, ButtonCursorPos);
+			}
+		}
+	}
+
+	FlushRelativeMovement();
+}
+
+void FWindowsApplication::ProcessMouseButton(double CurrentTime, const TSharedPtr< FGenericWindow >& Window, USHORT ButtonFlags, USHORT DownFlag, USHORT UpFlag, EMouseButtons::Type MouseButton, FVector2D CursorPos)
+{
+	if (ButtonFlags & DownFlag)
+	{
+		constexpr float MouseDoubleClickDelay = 0.5f;
+
+		const bool bHasDoubleClick = (CurrentTime <= RepeatTime[MouseButton]);
+		if (bHasDoubleClick)
+		{
+			MessageHandler->OnMouseDoubleClick(Window, MouseButton, CursorPos);
+		}
+		else
+		{
+			MessageHandler->OnMouseDown(Window, MouseButton, CursorPos);
+		}
+
+		RepeatTime[MouseButton] = CurrentTime + MouseDoubleClickDelay;
+	}
+	else if (ButtonFlags & UpFlag)
+	{
+		MessageHandler->OnMouseUp(MouseButton);
+	}
+}
+
 // Restore the windowsx.h macro for IsMaximized
 #pragma pop_macro("IsMaximized")
 
diff --git a/Engine/Source/Runtime/ApplicationCore/Private/Windows/WindowsPlatformApplicationMisc.cpp b/Engine/Source/Runtime/ApplicationCore/Private/Windows/WindowsPlatformApplicationMisc.cpp
--- a/Engine/Source/Runtime/ApplicationCore/Private/Windows/WindowsPlatformApplicationMisc.cpp
+++ b/Engine/Source/Runtime/ApplicationCore/Private/Windows/WindowsPlatformApplicationMisc.cpp
@@ -143,6 +143,7 @@ void FWindowsPlatformApplicationMisc::PumpMessages(bool bFromMainLoop)
 
 	GPumpingMessagesOutsideOfMainLoop = false;
 	WinPumpMessages();
+	FWindowsApplication::StaticPostPumpMessages();
 
 	// Determine if application has focus
 	bool bHasFocus = FApp::HasFocus();
diff --git a/Engine/Source/Runtime/ApplicationCore/Public/Windows/WindowsApplication.h b/Engine/Source/Runtime/ApplicationCore/Public/Windows/WindowsApplication.h
--- a/Engine/Source/Runtime/ApplicationCore/Public/Windows/WindowsApplication.h
+++ b/Engine/Source/Runtime/ApplicationCore/Public/Windows/WindowsApplication.h
@@ -294,6 +294,26 @@ public:
 	 */
 	static APPLICATIONCORE_API FWindowsApplication* CreateWindowsApplication( const HINSTANCE InstanceHandle, const HICON IconHandle );
 
+	/** Perform operations after pumping messages on global application. */
+	static void StaticPostPumpMessages();
+
+	/** Perform operations after pumping messages. */
+	void PostPumpMessages();
+
+	struct RawMouseInputResult
+	{
+		bool bIsAbsoluteInput;
+		bool bSimulateRelativeMovement;
+		int32 DeltaX;
+		int32 DeltaY;
+	};
+
+	/** Process raw mouse input. */
+	void ProcessRawMouseInput(const RAWMOUSE& Mouse, RawMouseInputResult& OutResult);
+
+	/** Process raw mouse inputs gathered from a helper thread/window. */
+	void ProcessWorkerInputs();
+
 	/** Virtual destructor. */
 	APPLICATIONCORE_API virtual ~FWindowsApplication();
 
@@ -349,6 +369,7 @@ public:
 	APPLICATIONCORE_API virtual void* GetCapture( void ) const override;
 	virtual bool IsMinimized() const override { return bMinimized; }
 	APPLICATIONCORE_API virtual void SetHighPrecisionMouseMode( const bool Enable, const TSharedPtr< FGenericWindow >& InWindow ) override;
+	void UpdateHighPrecisionMouseMode();
 	virtual bool IsUsingHighPrecisionMouseMode() const override { return bUsingHighPrecisionMouseInput; }
 	virtual bool IsMouseAttached() const override { return bIsMouseAttached; }
 	APPLICATIONCORE_API virtual bool IsGamepadAttached() const override;
@@ -428,6 +449,9 @@ protected:
 	/** Processes a single Windows message. */
 	APPLICATIONCORE_API int32 ProcessMessage( HWND hwnd, uint32 msg, WPARAM wParam, LPARAM lParam );
 
+	/** Processes mouse button. */
+	void ProcessMouseButton(double CurrentTime, const TSharedPtr< FGenericWindow >& Window, USHORT ButtonFlags, USHORT DownFlag, USHORT UpFlag, EMouseButtons::Type MouseButton, FVector2D CursorPos);
+
 	/** Processes a deferred Windows message. */
 	APPLICATIONCORE_API int32 ProcessDeferredMessage( const FDeferredWindowsMessage& DeferredMessage );
 
@@ -488,6 +512,12 @@ private:
 
 	APPLICATIONCORE_API FPlatformRect GetWorkAreaFromOS(const FPlatformRect& CurrentWindow) const;
 
+	/** Registers for raw input */
+	void RegisterHighPrecisionMouseMode();
+
+	/** Unregisters for raw input */
+	void UnregisterHighPrecisionMouseMode();
+
 private:
 
 	static APPLICATIONCORE_API const FIntPoint MinimizedWindowPosition;
@@ -497,6 +527,8 @@ private:
 	bool bMinimized;
 
 	bool bUsingHighPrecisionMouseInput;
+	bool bCachedUseWorkerThreadForRawInput;
+	HWND HighPrecisionMouseInputWindowHandle;
 
 	bool bIsMouseAttached;
 
@@ -580,6 +612,8 @@ private:
 	FIntPoint LastCursorPointPreWrap;
 	int32 NumPreWrapMsgsToRespect;
 	RECT ClipCursorRect;
+
+	double RepeatTime[EMouseButtons::Invalid];
 };
 
 
