- Fixed a crash when background the game super early (which was triggered always on first run after TF install)

diff --git a/Engine/Source/Runtime/ApplicationCore/Private/IOS/IOSAppDelegate.cpp b/Engine/Source/Runtime/ApplicationCore/Private/IOS/IOSAppDelegate.cpp
--- a/Engine/Source/Runtime/ApplicationCore/Private/IOS/IOSAppDelegate.cpp
+++ b/Engine/Source/Runtime/ApplicationCore/Private/IOS/IOSAppDelegate.cpp
@@ -349,7 +349,7 @@ static IOSAppDelegate* CachedDelegate = nil;
 	[self InitIdleTimerSettings];
 #endif
 
-	bEngineInit = true;
+	self.bEngineInit = true;
 
 	// put a render thread job to turn off the splash screen after the first render flip
 	if (GShowSplashScreen)
@@ -1005,7 +1005,7 @@ bool GIsSuspended = 0;
 - (BOOL)application:(UIApplication *)application willFinishLaunchingWithOptions:(NSDictionary *)launchOptions
 {
 	self.bDeviceInPortraitMode = false;
-	bEngineInit = false;
+	self.bEngineInit = false;
 
 	return YES;
 }
@@ -1417,7 +1417,7 @@ static FAutoConsoleVariableRef CVarGEnableThermalsReport(
     UIInterfaceOrientation Orientation = [self.Window.windowScene interfaceOrientation];
 	self.InterfaceOrientation = Orientation;
 	
-    if (bEngineInit)
+    if (self.bEngineInit)
     {
 		FFunctionGraphTask::CreateAndDispatchWhenReady([Orientation]()
 		{
@@ -1475,7 +1475,7 @@ static FAutoConsoleVariableRef CVarGEnableThermalsReport(
 	
 	//    Save openurl infomation before engine initialize.
 	//    When engine is done ready, running like previous. ( if OnOpenUrl is bound on game source. )
-	if (bEngineInit)
+	if (self.bEngineInit)
 	{
 		FIOSCoreDelegates::OnOpenURL.Broadcast(application, url, sourceApplication, annotation);
 	}
@@ -1511,7 +1511,7 @@ FCriticalSection RenderSuspend;
 		game.
 	 */
 	UE::GC::Private::SetSkipDestroy(true);
-    if (bEngineInit)
+    if (self.bEngineInit)
     {
  		FEmbeddedCommunication::KeepAwake(TEXT("Background"), false);
         FGraphEventRef ResignTask = FFunctionGraphTask::CreateAndDispatchWhenReady([]()
@@ -1536,8 +1536,12 @@ FCriticalSection RenderSuspend;
     
 // fix for freeze on tvOS, moving to applicationDidEnterBackground. Not making the changes for iOS platforms as the bug does not happen and could bring some side effets.
 #if !PLATFORM_TVOS
+	
 	bool bSupportsBackgroundAudio = false;
-	GConfig->GetBool(TEXT("/Script/IOSRuntimeSettings.IOSRuntimeSettings"), TEXT("bSupportsBackgroundAudio"), bSupportsBackgroundAudio, GEngineIni);
+	if (GConfig)
+	{
+		GConfig->GetBool(TEXT("/Script/IOSRuntimeSettings.IOSRuntimeSettings"), TEXT("bSupportsBackgroundAudio"), bSupportsBackgroundAudio, GEngineIni);
+	}
     if (!bSupportsBackgroundAudio)
     {
         [self ToggleSuspend:true];
@@ -1548,22 +1552,25 @@ FCriticalSection RenderSuspend;
 #endif
     
 #if WITH_ENGINE
-	FGraphEventRef RTBlockTask = FFunctionGraphTask::CreateAndDispatchWhenReady([]()
+	if (FTaskGraphInterface::IsRunning() && self.bEngineInit)
 	{
-		FRHICommandListImmediate& RHICmdList = FRHICommandListImmediate::Get();
-		if (!RHICmdList.Bypass()) 
-		{
-			RHICmdList.SubmitAndBlockUntilGPUIdle();
-		}
-	}, TStatId(), NULL, ENamedThreads::GetRenderThread());
-
-	RTBlockTask->Wait();
+		FGraphEventRef RTBlockTask = FFunctionGraphTask::CreateAndDispatchWhenReady([]()
+																					{
+			FRHICommandListImmediate& RHICmdList = FRHICommandListImmediate::Get();
+			if (!RHICmdList.Bypass()) 
+			{
+				RHICmdList.SubmitAndBlockUntilGPUIdle();
+			}
+		}, TStatId(), NULL, ENamedThreads::GetRenderThread());
+		
+		RTBlockTask->Wait();
+	}
 #endif
 	
     RenderSuspend.TryLock();
     if (FTaskGraphInterface::IsRunning())
     {
-        if (bEngineInit)
+        if (self.bEngineInit)
         {
             FGraphEventRef ResignTask = FFunctionGraphTask::CreateAndDispatchWhenReady([]()
             {
@@ -1655,7 +1662,7 @@ extern double GCStartTime;
     FAppEntry::Resume(true);
 
 
-    if (bEngineInit)
+    if (self.bEngineInit)
     {
 		FEmbeddedCommunication::KeepAwake(TEXT("Background"), false);
 
@@ -1697,7 +1704,7 @@ extern double GCStartTime;
     // TODO: fix the reason why we are hanging when asked to shutdown
 /*    RequestEngineExit(TEXT("IOS applicationWillTerminate"));
     
-    if (!bEngineInit)*/
+    if (!self.bEngineInit)*/
     {
         // we haven't yet made it to the point where the engine is initialized, so just exit the app
         _Exit(0);
diff --git a/Engine/Source/Runtime/ApplicationCore/Public/IOS/IOSAppDelegate.h b/Engine/Source/Runtime/ApplicationCore/Public/IOS/IOSAppDelegate.h
--- a/Engine/Source/Runtime/ApplicationCore/Public/IOS/IOSAppDelegate.h
+++ b/Engine/Source/Runtime/ApplicationCore/Public/IOS/IOSAppDelegate.h
@@ -193,7 +193,7 @@ APPLICATIONCORE_API
 #endif
 
 /** True if the engine has been initialized */
-@property (readonly) bool bEngineInit;
+@property (atomic, readonly) bool bEngineInit;
 
 /** Delays game initialization slightly in case we have a URL launch to handle */
 @property (retain) NSTimer* CommandLineParseTimer;
