- Fix for readonly, fixes a build break.

diff --git a/Engine/Source/Runtime/ApplicationCore/Public/IOS/IOSAppDelegate.h b/Engine/Source/Runtime/ApplicationCore/Public/IOS/IOSAppDelegate.h
--- a/Engine/Source/Runtime/ApplicationCore/Public/IOS/IOSAppDelegate.h
+++ b/Engine/Source/Runtime/ApplicationCore/Public/IOS/IOSAppDelegate.h
@@ -193,7 +193,7 @@ APPLICATIONCORE_API
 #endif
 
 /** True if the engine has been initialized */
-@property (atomic, readonly) bool bEngineInit;
+@property (atomic) bool bEngineInit;
 
 /** Delays game initialization slightly in case we have a URL launch to handle */
 @property (retain) NSTimer* CommandLineParseTimer;
