Dynamic Trigger Release support on main analog trigger platforms (disabled by default)

Allows analog trigger inputs that are just being used as Pressed/Released inputs to be considered released without releasing the trigger all of the way.

Then it can be enabled and configured for a given controller through either:
- Setting Engine CVar Input.TriggerDynamicReleaseDeadZone to a value between 0 and 1 to set the deadzone on all applicable platforms (I like 0.25); or
- IInputInterface::SetDeviceProperty using FInputDeviceTriggerDynamicReleaseDeadZoneProperty; or
- APlayerController::SetControllerTriggerReleaseThresholds(const float LeftThreshold, const float RightThreshold)

Setting the threshold/deadzone to 1 disables the effect for that trigger (1 = full release required for the trigger to be considered released).

Additional overall configuration through Engine CVars:
- Input.AllowTriggerDynamicReleaseDeadZoneCustomization (default true) - Set to false to prevent SetDeviceProperty or SetControllerTriggerReleaseThresholds from overriding the deadzone set by [REDACTED_PERSON].
- Input.TriggerDynamicReleaseDeadZoneRePressFactor (default 0.1) - After a release is triggered by leaving the deadzone, how far does the trigger need to be pulled to be considered Pressed again? It needs to be pulled by the the DeadZone (as set through SetDeviceProperty or SetControllerTriggerReleaseThresholds) * Input.TriggerDynamicReleaseDeadZoneRePressFactor. With the default value, this is a 10% threshold.
- Input.TriggerDynamicReleaseDeadZoneMinimumRePress (default 0.05) A very small trigger release threshold with a very small Input.TriggerDynamicReleaseDeadZoneRePressFactor could make for a very, very small buffer for considering the trigger Pressed again, which could easily be done accidentally. So Input.TriggerDynamicReleaseDeadZoneMinimumRePress sets a limit on how small that threshold can get. The default value means the threshold for triggering a new trigger press or release* won't be less than 5% of the full trigger movement.

*By design, this behavior does not prevent normal Pressed and Released thresholds from working correctly. So pressing a trigger enough to be considered Pressed by default behavior will always be enough to consider it Pressed with this dynamic deadzone, and fully releasing the trigger will always work just as well with dynamic deadzone as it does without the dynamic deadzone.

#tests PIE with a variety of settings and supported controllers.

diff --git a/Engine/Plugins/Runtime/Windows/XInputDevice/Source/XInputDevice/Private/XInputInterface.cpp b/Engine/Plugins/Runtime/Windows/XInputDevice/Source/XInputDevice/Private/XInputInterface.cpp
--- a/Engine/Plugins/Runtime/Windows/XInputDevice/Source/XInputDevice/Private/XInputInterface.cpp
+++ b/Engine/Plugins/Runtime/Windows/XInputDevice/Source/XInputDevice/Private/XInputInterface.cpp
@@ -150,6 +150,35 @@ void XInputInterface::GetPlatformUserAndDevice(int32 InControllerId, EInputDevic
 	}
 }
 
+void XInputInterface::SetDynamicTriggerThreshold(const int32 InControllerId, const EInputDeviceTriggerMask TriggerMask, const float Threshold)
+{
+	for (int32 ControllerIndex = 0; ControllerIndex < MAX_NUM_XINPUT_CONTROLLERS; ++ControllerIndex)
+	{
+		FControllerState& ControllerState = ControllerStates[ControllerIndex];
+
+		if (ControllerState.ControllerId == InControllerId)
+		{
+			switch (TriggerMask)
+			{
+			case EInputDeviceTriggerMask::Left:
+				ControllerState.LeftTriggerReleaseDeadZone.OverrideDeadZone(Threshold);
+				break;
+			case EInputDeviceTriggerMask::Right:
+				ControllerState.RightTriggerReleaseDeadZone.OverrideDeadZone(Threshold);
+				break;
+			case EInputDeviceTriggerMask::All:
+				ControllerState.LeftTriggerReleaseDeadZone.OverrideDeadZone(Threshold);
+				ControllerState.RightTriggerReleaseDeadZone.OverrideDeadZone(Threshold);
+				break;
+			case EInputDeviceTriggerMask::None:
+			default:
+				break;
+			}
+			break;
+		}
+	}
+}
+
 namespace UE::XInputInterface::Private
 {
 EInputDeviceConnectionState GetInputDeviceConnectionState(bool bWasConnected, bool bControllerStateIsConnected)
@@ -226,6 +255,12 @@ void XInputInterface::SendControllerEvents()
 				return;
 			}
 			
+			bool bLeftTriggerPressed = XInputState.Gamepad.bLeftTrigger > XINPUT_GAMEPAD_TRIGGER_THRESHOLD;
+			bool bRightTriggerPressed = XInputState.Gamepad.bRightTrigger > XINPUT_GAMEPAD_TRIGGER_THRESHOLD;
+
+			bLeftTriggerPressed = ControllerState.LeftTriggerReleaseDeadZone.IsPressed(XInputState.Gamepad.bLeftTrigger, ControllerState.LeftTriggerAnalog, bLeftTriggerPressed, ControllerState.LeftTriggerAnalog > XINPUT_GAMEPAD_TRIGGER_THRESHOLD);
+			bRightTriggerPressed = ControllerState.RightTriggerReleaseDeadZone.IsPressed(XInputState.Gamepad.bRightTrigger, ControllerState.RightTriggerAnalog, bRightTriggerPressed, ControllerState.RightTriggerAnalog > XINPUT_GAMEPAD_TRIGGER_THRESHOLD);
+
 			bool CurrentStates[MAX_NUM_CONTROLLER_BUTTONS] = {0};
 		
 			// Get the current state of all buttons
@@ -239,8 +274,8 @@ void XInputInterface::SendControllerEvents()
 			CurrentStates[X360ToXboxControllerMapping[7]] = !!(XInputState.Gamepad.wButtons & XINPUT_GAMEPAD_START);
 			CurrentStates[X360ToXboxControllerMapping[8]] = !!(XInputState.Gamepad.wButtons & XINPUT_GAMEPAD_LEFT_THUMB);
 			CurrentStates[X360ToXboxControllerMapping[9]] = !!(XInputState.Gamepad.wButtons & XINPUT_GAMEPAD_RIGHT_THUMB);
-			CurrentStates[X360ToXboxControllerMapping[10]] = !!(XInputState.Gamepad.bLeftTrigger > XINPUT_GAMEPAD_TRIGGER_THRESHOLD);
-			CurrentStates[X360ToXboxControllerMapping[11]] = !!(XInputState.Gamepad.bRightTrigger > XINPUT_GAMEPAD_TRIGGER_THRESHOLD);
+			CurrentStates[X360ToXboxControllerMapping[10]] = !!(bLeftTriggerPressed);
+			CurrentStates[X360ToXboxControllerMapping[11]] = !!(bRightTriggerPressed);
 			CurrentStates[X360ToXboxControllerMapping[12]] = !!(XInputState.Gamepad.wButtons & XINPUT_GAMEPAD_DPAD_UP);
 			CurrentStates[X360ToXboxControllerMapping[13]] = !!(XInputState.Gamepad.wButtons & XINPUT_GAMEPAD_DPAD_DOWN);
 			CurrentStates[X360ToXboxControllerMapping[14]] = !!(XInputState.Gamepad.wButtons & XINPUT_GAMEPAD_DPAD_LEFT);
@@ -366,9 +401,17 @@ void XInputInterface::SetDeviceProperty(int32 ControllerId, const FInputDevicePr
 {
 	// This should only get called from WindowsApplication.cpp when Windows detects a device change.
 	static const FName UpdateRequestedName = TEXT("Request_Device_Update");
-	if (Property && Property->Name == UpdateRequestedName)
+	if (Property)
 	{
-		SetNeedsControllerStateUpdate();	
+		if (Property->Name == UpdateRequestedName)
+		{
+			SetNeedsControllerStateUpdate();
+		}
+		else if (Property->Name == FInputDeviceTriggerDynamicReleaseDeadZoneProperty::PropertyName())
+		{
+			const FInputDeviceTriggerDynamicReleaseDeadZoneProperty* TriggerReleaseThreshold = static_cast<const FInputDeviceTriggerDynamicReleaseDeadZoneProperty*>(Property);
+			SetDynamicTriggerThreshold(ControllerId, TriggerReleaseThreshold->AffectedTriggers, TriggerReleaseThreshold->DeadZone);
+		}
 	}
 }
 
diff --git a/Engine/Plugins/Runtime/Windows/XInputDevice/Source/XInputDevice/Private/XInputInterface.h b/Engine/Plugins/Runtime/Windows/XInputDevice/Source/XInputDevice/Private/XInputInterface.h
--- a/Engine/Plugins/Runtime/Windows/XInputDevice/Source/XInputDevice/Private/XInputInterface.h
+++ b/Engine/Plugins/Runtime/Windows/XInputDevice/Source/XInputDevice/Private/XInputInterface.h
@@ -5,6 +5,7 @@
 #include "CoreTypes.h"
 #include "GenericPlatform/IInputInterface.h"
 #include "IInputDevice.h"
+#include "GenericPlatform/GamepadUtils.h"
 #include "GenericPlatform/GenericApplicationMessageHandler.h"
 #include "GenericPlatform/GenericInputDeviceMap.h"
 #include "Misc/CoreMiscDefines.h"
@@ -78,6 +79,9 @@ private:
 	 */
 	void GetPlatformUserAndDevice( int32 InControllerId, EInputDeviceConnectionState InDeviceState, FPlatformUserId& OutPlatformUserId, FInputDeviceId& OutDeviceId );
 
+	/** Sets dynamic trigger release threshold for the given trigger(s) on the given controller. */
+	void SetDynamicTriggerThreshold(const int32 ControllerId, const EInputDeviceTriggerMask TriggerMask, const float Threshold);
+
 	struct FControllerState
 	{
 		/** Last frame's button states, so we only send events on edges */
@@ -120,6 +124,12 @@ private:
 		 * The last valid FPlatformUserId that can be used on the frame where the input device is disconnected
 		 */
 		FPlatformUserId LastUsedValidPlatformUserId;
+
+		/** Dynamic Release DeadZone for Left Trigger */
+		FDynamicReleaseDeadZone LeftTriggerReleaseDeadZone;
+
+		/** Dynamic Release DeadZone for Right Trigger */
+		FDynamicReleaseDeadZone RightTriggerReleaseDeadZone;
 	};
 
 	/** If we've been notified by the system that the controller state may have changed */
diff --git a/Engine/Source/Runtime/ApplicationCore/Private/GenericPlatform/GamepadUtils.cpp b/Engine/Source/Runtime/ApplicationCore/Private/GenericPlatform/GamepadUtils.cpp
new file mode 100644
--- /dev/null
+++ b/Engine/Source/Runtime/ApplicationCore/Private/GenericPlatform/GamepadUtils.cpp
@@ -0,0 +1,109 @@
+// Copyright Epic Games, Inc. All Rights Reserved.
+
+#include "GenericPlatform/GamepadUtils.h"
+#include "HAL/IConsoleManager.h"
+#include "Math/UnrealMathUtility.h"
+
+namespace DynamicTriggerReleaseCVars
+{
+	static float DefaultDeadZone = 1.f;
+	static FAutoConsoleVariableRef CVarTriggerDynamicReleaseDeadZone(TEXT("Input.TriggerDynamicReleaseDeadZone"),
+		DefaultDeadZone,
+		TEXT("Dynamic DeadZone for releasing analog triggers. It's dynamic in that it's relative to how far the trigger was pulled. Default of 1 means no dynamic release, letting triggers be released with their usual fixed threshold."));
+
+	static bool bAllowOverride = true;
+	static FAutoConsoleVariableRef CVarAllowTriggerDynamicReleaseDeadZoneCustomization(TEXT("Input.AllowTriggerDynamicReleaseDeadZoneCustomization"),
+		bAllowOverride,
+		TEXT("If true, the Dynamic Release DeadZone for a given analog trigger can be customized on supported devices by calling SetDeviceProperty with a FInputDeviceTriggerDynamicReleaseDeadZoneProperty. Otherwise, the Dynamic Release DeadZone is always the value of the CVar Input.TriggerDynamicReleaseDeadZone at initialization or when last Refreshed."));
+
+	static float RePressFactor = 0.1f;
+	static FAutoConsoleVariableRef CVarRePressFactor(TEXT("Input.TriggerDynamicReleaseDeadZoneRePressFactor"),
+		RePressFactor,
+		TEXT("When the trigger is considered released due to the Dynamic DeadZone, how far it needs to be pulled again to be considered pressed again. This will be multiplied by the Dynamic DeadZone, so lower values mean more sensitivity to pressing again."));
+
+	static float MinimumRePress = 0.05f;
+	static FAutoConsoleVariableRef CVarMinimumRePress(TEXT("Input.TriggerDynamicReleaseDeadZoneMinimumRePress"),
+		MinimumRePress,
+		TEXT("When the trigger is considered released due to the Dynamic DeadZone, how far it needs to be pulled again at a minimum to be considered pressed again. This is to prevent that a small Dynamic DeadZone and a small RePress Factor (see Input.TriggerDynamicReleaseDeadZoneRePressFactor) make the trigger overly sensitive to being considered pressed again."));
+}
+
+void FDynamicReleaseDeadZone::RefreshSettings()
+{
+	// Main settings:
+	if (!DynamicTriggerReleaseCVars::bAllowOverride)
+	{
+		bHasOverride = false;
+	}
+	if (!bHasOverride)
+	{
+		DeadZone = FMath::Clamp(DynamicTriggerReleaseCVars::DefaultDeadZone, 0.f, 1.f);
+	}
+}
+
+void FDynamicReleaseDeadZone::OverrideDeadZone(const float InDeadZone)
+{
+	if (DynamicTriggerReleaseCVars::bAllowOverride)
+	{
+		bHasOverride = true;
+		DeadZone = FMath::Clamp(InDeadZone, 0.f, 1.f);
+	}
+	else if (bHasOverride)
+	{
+		bHasOverride = false;
+		DeadZone = FMath::Clamp(DynamicTriggerReleaseCVars::DefaultDeadZone, 0.f, 1.f);
+	}
+}
+
+bool FDynamicReleaseDeadZone::IsDynamicReleaseEnabled() const
+{
+	return DeadZone < 1.f;
+}
+
+bool FDynamicReleaseDeadZone::IsPressed(const uint8 TriggerAnalog, const uint8 PreviousTriggerAnalog)
+{
+	return IsPressed(TriggerAnalog, PreviousTriggerAnalog, TriggerAnalog > 0, PreviousTriggerAnalog > 0);
+}
+
+bool FDynamicReleaseDeadZone::IsPressed(const uint8 TriggerAnalog, const uint8 PreviousTriggerAnalog, const bool bIsSimplePressed, const bool bPreviousSimplePressed)
+{
+	bool bIsPressedResult = bIsSimplePressed;
+
+	// If the threshold is 1.0, this trigger's thresholds aren't dynamic -- it's the same as the default behavior
+	if (DeadZone < 1.f)
+	{
+		// Conservatively, changes in "Simple" trigger state (Pressed to Released or Released to Pressed) will take precedence over Dynamic Thresholds
+		const bool bSimpleChanged = bIsSimplePressed != bPreviousSimplePressed;
+		const bool bDynamicPressed = TriggerAnalog > TriggerThreshold;
+		const bool bDynamicChanged = bDynamicPressed != (PreviousTriggerAnalog > TriggerThreshold);
+
+		if (!bSimpleChanged)
+		{
+			bIsPressedResult = bDynamicPressed;
+		}
+
+		if (bSimpleChanged || bDynamicChanged)
+		{
+			// Reset dynamic thresholds so they can be recalculated correctly
+			TriggerThreshold = bIsPressedResult ? 0 : 255;
+		}
+
+		const float RePressFactor = DynamicTriggerReleaseCVars::RePressFactor;
+		const float MinimumRePress = DynamicTriggerReleaseCVars::MinimumRePress;
+
+		// Update current dynamic thresholds
+		if (bIsPressedResult)
+		{
+			const float ReleaseThresholdOffset = FMath::Max(DeadZone, MinimumRePress);
+			const uint8 NewThreshold = (uint8)FMath::Max((int32)TriggerAnalog - (int32)(255 * ReleaseThresholdOffset) - 1, 0);
+			TriggerThreshold = FMath::Max(TriggerThreshold, NewThreshold);
+		}
+		else
+		{
+			const float PressThresholdOffset = FMath::Max(DeadZone * RePressFactor, MinimumRePress);
+			const uint8 NewThreshold = (uint8)FMath::Min((int32)TriggerAnalog + (int32)(255 * PressThresholdOffset), 254);
+			TriggerThreshold = FMath::Min(TriggerThreshold, NewThreshold);
+		}
+	}
+
+	return bIsPressedResult;
+}
diff --git a/Engine/Source/Runtime/ApplicationCore/Public/GenericPlatform/GamepadUtils.h b/Engine/Source/Runtime/ApplicationCore/Public/GenericPlatform/GamepadUtils.h
new file mode 100644
--- /dev/null
+++ b/Engine/Source/Runtime/ApplicationCore/Public/GenericPlatform/GamepadUtils.h
@@ -0,0 +1,33 @@
+// Copyright Epic Games, Inc. All Rights Reserved.
+
+#pragma once
+
+#include "CoreFwd.h"
+#include "CoreTypes.h"
+
+/**
+ * Properties and functions for a trigger's Dyanamic Release DeadZone.
+ */
+struct FDynamicReleaseDeadZone
+{
+public:
+	FDynamicReleaseDeadZone()
+	{
+		RefreshSettings();
+	}
+
+	APPLICATIONCORE_API void RefreshSettings();
+
+	APPLICATIONCORE_API void OverrideDeadZone(const float InDeadZone);
+
+	bool IsDynamicReleaseEnabled() const;
+
+	APPLICATIONCORE_API bool IsPressed(const uint8 TriggerAnalog, const uint8 PreviousTriggerAnalog);
+
+	APPLICATIONCORE_API bool IsPressed(const uint8 TriggerAnalog, const uint8 PreviousTriggerAnalog, const bool bIsSimplePressed, const bool bPreviousSimplePressed);
+
+protected:
+	float DeadZone = 1.0f;
+	bool bHasOverride = false;
+	uint8 TriggerThreshold = 0;
+};
diff --git a/Engine/Source/Runtime/ApplicationCore/Public/GenericPlatform/IInputInterface.h b/Engine/Source/Runtime/ApplicationCore/Public/GenericPlatform/IInputInterface.h
--- a/Engine/Source/Runtime/ApplicationCore/Public/GenericPlatform/IInputInterface.h
+++ b/Engine/Source/Runtime/ApplicationCore/Public/GenericPlatform/IInputInterface.h
@@ -278,6 +278,30 @@ struct FInputDeviceAnalogStickDeadZoneProperty : public FInputDeviceAnalogStickP
 	float DeadZone = 0;
 };
 
+/**
+ * A generic input device property that sets "Dynamic Release DeadZones" for Analog Triggers.
+ * 
+ * NOTE: Dynamic Triggers allow for Release and Press to be triggered even if the user fails to fully release the trigger before beginning another press.
+ * Smaller deadzones increase this effect. A deadzone of 1.0 isn't dynamic at all.
+ */
+struct FInputDeviceTriggerDynamicReleaseDeadZoneProperty : public FInputDeviceTriggerProperty
+{
+	FInputDeviceTriggerDynamicReleaseDeadZoneProperty()
+		: FInputDeviceTriggerProperty(PropertyName())
+	{}
+
+	FInputDeviceTriggerDynamicReleaseDeadZoneProperty(EInputDeviceTriggerMask InAffectedTriggers, float InDeadZone)
+		: FInputDeviceTriggerProperty(PropertyName())
+		, DeadZone(InDeadZone)
+	{
+		AffectedTriggers = InAffectedTriggers;
+	}
+
+	static FName PropertyName() { return FName("InputDeviceTriggerDynamicReleaseDeadZone"); }
+
+	float DeadZone = 1.0f;
+};
+
 /**
  * Interface for the input interface.
  */
diff --git a/Engine/Source/Runtime/Engine/Classes/GameFramework/PlayerController.h b/Engine/Source/Runtime/Engine/Classes/GameFramework/PlayerController.h
--- a/Engine/Source/Runtime/Engine/Classes/GameFramework/PlayerController.h
+++ b/Engine/Source/Runtime/Engine/Classes/GameFramework/PlayerController.h
@@ -1377,6 +1377,20 @@ public:
 	UFUNCTION(BlueprintCallable, Category = "Input")
 	ENGINE_API void ResetControllerDeadZones();
 
+	/**
+	 * Sets the trigger release thresholds of the player's controller
+	 * @param	LeftThreshold 	Release Threshold for the left trigger
+	 * @param	RightThreshold	Release Threshold for the right trigger
+	 */
+	UFUNCTION(BlueprintCallable, Category = "Input")
+	ENGINE_API void SetControllerTriggerReleaseThresholds(const float LeftThreshold, const float RightThreshold);
+
+	/**
+	 * Resets the player's controller trigger release thresholds to default
+	 */
+	UFUNCTION(BlueprintCallable, Category = "Input")
+	ENGINE_API void ResetControllerTriggerReleaseThresholds();
+
 	/**
 	 * Sets whether the player's controller's gyro auto calibration is enabled
 	 * @param	bEnabled	 	Whether gyro auto calibration should be enabled
diff --git a/Engine/Source/Runtime/Engine/Private/PlayerController.cpp b/Engine/Source/Runtime/Engine/Private/PlayerController.cpp
--- a/Engine/Source/Runtime/Engine/Private/PlayerController.cpp
+++ b/Engine/Source/Runtime/Engine/Private/PlayerController.cpp
@@ -4595,6 +4595,46 @@ void APlayerController::ResetControllerDeadZones()
 	}
 }
 
+void APlayerController::SetControllerTriggerReleaseThresholds(const float LeftThreshold, const float RightThreshold)
+{
+	if (Player == nullptr)
+	{
+		return;
+	}
+
+	if (FSlateApplication::IsInitialized())
+	{
+		IInputInterface* InputInterface = FSlateApplication::Get().GetInputInterface();
+		if (InputInterface)
+		{
+			const int32 ControllerId = CastChecked<ULocalPlayer>(Player)->GetControllerId();
+			FInputDeviceTriggerDynamicReleaseDeadZoneProperty ReleaseThresholdLeft(EInputDeviceTriggerMask::Left, LeftThreshold);
+			InputInterface->SetDeviceProperty(ControllerId, &ReleaseThresholdLeft);
+			FInputDeviceTriggerDynamicReleaseDeadZoneProperty ReleaseThresholdRight(EInputDeviceTriggerMask::Right, RightThreshold);
+			InputInterface->SetDeviceProperty(ControllerId, &ReleaseThresholdRight);
+		}
+	}
+}
+
+void APlayerController::ResetControllerTriggerReleaseThresholds()
+{
+	if (Player == nullptr)
+	{
+		return;
+	}
+
+	if (FSlateApplication::IsInitialized())
+	{
+		IInputInterface* InputInterface = FSlateApplication::Get().GetInputInterface();
+		if (InputInterface)
+		{
+			const int32 ControllerId = CastChecked<ULocalPlayer>(Player)->GetControllerId();
+			FInputDeviceTriggerDynamicReleaseDeadZoneProperty ReleaseThreshold(EInputDeviceTriggerMask::All, 1.0f);
+			InputInterface->SetDeviceProperty(ControllerId, &ReleaseThreshold);
+		}
+	}
+}
+
 void APlayerController::SetControllerGyroAutoCalibration(bool bEnabled)
 {
 	if (Player == nullptr)
