Simple IsPressed function in GamepadUtils for dynamic trigger release deadzone

#tests PIE with applicable controllers

diff --git a/Engine/Plugins/Runtime/Windows/XInputDevice/Source/XInputDevice/Private/XInputInterface.cpp b/Engine/Plugins/Runtime/Windows/XInputDevice/Source/XInputDevice/Private/XInputInterface.cpp
--- a/Engine/Plugins/Runtime/Windows/XInputDevice/Source/XInputDevice/Private/XInputInterface.cpp
+++ b/Engine/Plugins/Runtime/Windows/XInputDevice/Source/XInputDevice/Private/XInputInterface.cpp
@@ -42,6 +42,8 @@ XInputInterface::XInputInterface(const TSharedRef<FGenericApplicationMessageHand
 		FMemory::Memzero( &ControllerState, sizeof(FControllerState) );
 
 		ControllerState.ControllerId = ControllerIndex;
+		ControllerState.LeftTriggerReleaseDeadZone = FDynamicReleaseDeadZone();
+		ControllerState.RightTriggerReleaseDeadZone = FDynamicReleaseDeadZone();
 	}
 
 	bIsGamepadAttached = false;
@@ -258,8 +260,8 @@ void XInputInterface::SendControllerEvents()
 			bool bLeftTriggerPressed = XInputState.Gamepad.bLeftTrigger > XINPUT_GAMEPAD_TRIGGER_THRESHOLD;
 			bool bRightTriggerPressed = XInputState.Gamepad.bRightTrigger > XINPUT_GAMEPAD_TRIGGER_THRESHOLD;
 
-			bLeftTriggerPressed = ControllerState.LeftTriggerReleaseDeadZone.IsPressed(XInputState.Gamepad.bLeftTrigger, ControllerState.LeftTriggerAnalog, bLeftTriggerPressed, ControllerState.LeftTriggerAnalog > XINPUT_GAMEPAD_TRIGGER_THRESHOLD);
-			bRightTriggerPressed = ControllerState.RightTriggerReleaseDeadZone.IsPressed(XInputState.Gamepad.bRightTrigger, ControllerState.RightTriggerAnalog, bRightTriggerPressed, ControllerState.RightTriggerAnalog > XINPUT_GAMEPAD_TRIGGER_THRESHOLD);
+			bLeftTriggerPressed = ControllerState.LeftTriggerReleaseDeadZone.IsPressed(XInputState.Gamepad.bLeftTrigger, bLeftTriggerPressed);
+			bRightTriggerPressed = ControllerState.RightTriggerReleaseDeadZone.IsPressed(XInputState.Gamepad.bRightTrigger, bRightTriggerPressed);
 
 			bool CurrentStates[MAX_NUM_CONTROLLER_BUTTONS] = {0};
 		
diff --git a/Engine/Source/Runtime/ApplicationCore/Private/GenericPlatform/GamepadUtils.cpp b/Engine/Source/Runtime/ApplicationCore/Private/GenericPlatform/GamepadUtils.cpp
--- a/Engine/Source/Runtime/ApplicationCore/Private/GenericPlatform/GamepadUtils.cpp
+++ b/Engine/Source/Runtime/ApplicationCore/Private/GenericPlatform/GamepadUtils.cpp
@@ -65,6 +65,14 @@ bool FDynamicReleaseDeadZone::IsPressed(const uint8 TriggerAnalog, const uint8 P
 }
 
 bool FDynamicReleaseDeadZone::IsPressed(const uint8 TriggerAnalog, const uint8 PreviousTriggerAnalog, const bool bIsSimplePressed, const bool bPreviousSimplePressed)
+{
+	bWasSimplePressed = bPreviousSimplePressed;
+	bWasDynamicPressed = PreviousTriggerAnalog > TriggerThreshold;
+
+	return IsPressed(TriggerAnalog, bIsSimplePressed);
+}
+
+bool FDynamicReleaseDeadZone::IsPressed(const uint8 TriggerAnalog, const bool bIsSimplePressed)
 {
 	bool bIsPressedResult = bIsSimplePressed;
 
@@ -72,9 +80,9 @@ bool FDynamicReleaseDeadZone::IsPressed(const uint8 TriggerAnalog, const uint8 P
 	if (DeadZone < 1.f)
 	{
 		// Conservatively, changes in "Simple" trigger state (Pressed to Released or Released to Pressed) will take precedence over Dynamic Thresholds
-		const bool bSimpleChanged = bIsSimplePressed != bPreviousSimplePressed;
+		const bool bSimpleChanged = bIsSimplePressed != bWasSimplePressed;
 		const bool bDynamicPressed = TriggerAnalog > TriggerThreshold;
-		const bool bDynamicChanged = bDynamicPressed != (PreviousTriggerAnalog > TriggerThreshold);
+		const bool bDynamicChanged = bDynamicPressed != bWasDynamicPressed;
 
 		if (!bSimpleChanged)
 		{
@@ -103,7 +111,15 @@ bool FDynamicReleaseDeadZone::IsPressed(const uint8 TriggerAnalog, const uint8 P
 			const uint8 NewThreshold = (uint8)FMath::Min((int32)TriggerAnalog + (int32)(255 * PressThresholdOffset), 254);
 			TriggerThreshold = FMath::Min(TriggerThreshold, NewThreshold);
 		}
+
+		bWasDynamicPressed = bDynamicPressed;
+	}
+	else
+	{
+		bWasDynamicPressed = false;
 	}
 
+	bWasSimplePressed = bIsSimplePressed;
+
 	return bIsPressedResult;
 }
diff --git a/Engine/Source/Runtime/ApplicationCore/Public/GenericPlatform/GamepadUtils.h b/Engine/Source/Runtime/ApplicationCore/Public/GenericPlatform/GamepadUtils.h
--- a/Engine/Source/Runtime/ApplicationCore/Public/GenericPlatform/GamepadUtils.h
+++ b/Engine/Source/Runtime/ApplicationCore/Public/GenericPlatform/GamepadUtils.h
@@ -12,6 +12,11 @@ struct FDynamicReleaseDeadZone
 {
 public:
 	FDynamicReleaseDeadZone()
+		: DeadZone(1.0f)
+		, TriggerThreshold(0)
+		, bHasOverride(false)
+		, bWasSimplePressed(false)
+		, bWasDynamicPressed(false)
 	{
 		RefreshSettings();
 	}
@@ -26,8 +31,12 @@ public:
 
 	APPLICATIONCORE_API bool IsPressed(const uint8 TriggerAnalog, const uint8 PreviousTriggerAnalog, const bool bIsSimplePressed, const bool bPreviousSimplePressed);
 
+	APPLICATIONCORE_API bool IsPressed(const uint8 TriggerAnalog, const bool bIsSimplePressed);
+
 protected:
 	float DeadZone = 1.0f;
-	bool bHasOverride = false;
 	uint8 TriggerThreshold = 0;
+	uint8 bHasOverride : 1;
+	uint8 bWasSimplePressed : 1;
+	uint8 bWasDynamicPressed : 1;
 };
