Engine Changes for SDL3

diff --git a/Engine/Plugins/Runtime/WebMMoviePlayer/Source/WebMMoviePlayer/Private/Audio/Unix/WebMAudioBackend.h b/Engine/Plugins/Runtime/WebMMoviePlayer/Source/WebMMoviePlayer/Private/Audio/Unix/WebMAudioBackend.h
--- a/Engine/Plugins/Runtime/WebMMoviePlayer/Source/WebMMoviePlayer/Private/Audio/Unix/WebMAudioBackend.h
+++ b/Engine/Plugins/Runtime/WebMMoviePlayer/Source/WebMMoviePlayer/Private/Audio/Unix/WebMAudioBackend.h
@@ -5,7 +5,7 @@
 #include "HAL/PlatformMath.h"
 
 THIRD_PARTY_INCLUDES_START
-#include "SDL_audio.h"
+#include <SDL3/SDL_audio.h>
 THIRD_PARTY_INCLUDES_END
 
 class FString;
@@ -38,7 +38,7 @@ public:
 	virtual FString GetDefaultDeviceName();
 
 protected:
-	SDL_AudioDeviceID AudioDevice;
+	SDL_AudioStream* AudioStream;
 	bool bSDLInitialized;
 
 	/** Whether we're pausing playing back sound. */
diff --git a/Engine/Plugins/Runtime/WebMMoviePlayer/Source/WebMMoviePlayer/Private/Audio/Unix/WebMAudioBackendSDL.cpp b/Engine/Plugins/Runtime/WebMMoviePlayer/Source/WebMMoviePlayer/Private/Audio/Unix/WebMAudioBackendSDL.cpp
--- a/Engine/Plugins/Runtime/WebMMoviePlayer/Source/WebMMoviePlayer/Private/Audio/Unix/WebMAudioBackendSDL.cpp
+++ b/Engine/Plugins/Runtime/WebMMoviePlayer/Source/WebMMoviePlayer/Private/Audio/Unix/WebMAudioBackendSDL.cpp
@@ -2,10 +2,10 @@
 
 #include "WebMAudioBackend.h"
 #include "WebMMovieCommon.h"
-#include <SDL.h>
+#include <SDL3/SDL.h>
 
 FWebMAudioBackendSDL::FWebMAudioBackendSDL()
-	: AudioDevice(0)
+	: AudioStream(nullptr)
 	, bSDLInitialized(false)
 	, bPaused(false)
 {
@@ -51,24 +51,11 @@ bool FWebMAudioBackendSDL::StartStreaming(int32 SampleRate, int32 NumOfChannels,
 		return false;
 	}
 
-	SDL_AudioSpec AudioSpec;
-	memset(&AudioSpec, 0, sizeof(SDL_AudioSpec));
-	AudioSpec.freq = SampleRate;
-	AudioSpec.format = AUDIO_S16;
-	AudioSpec.channels = NumOfChannels;
-	AudioSpec.samples = 4096;
+	SDL_AudioSpec AudioSpec = {SDL_AUDIO_S16LE, NumOfChannels, SampleRate };
 
-	FString DefaultDeivceName = GetDefaultDeviceName();
-	if (DefaultDeivceName.Len() > 0)
-	{
-		AudioDevice = SDL_OpenAudioDevice(TCHAR_TO_ANSI(*DefaultDeivceName), 0, &AudioSpec, nullptr, 0);
-	}
-	else
-	{
-		AudioDevice = SDL_OpenAudioDevice(nullptr, 0, &AudioSpec, nullptr, 0);
-	}
+	AudioStream = SDL_OpenAudioDeviceStream(SDL_AUDIO_DEVICE_DEFAULT_PLAYBACK, &AudioSpec, nullptr, nullptr);
 
-	if (!AudioDevice)
+	if (!AudioStream)
 	{
 		UE_LOG(LogWebMMoviePlayer, Error, TEXT("SDL_OpenAudioDevice failed"));
 		return false;
@@ -81,16 +68,17 @@ bool FWebMAudioBackendSDL::StartStreaming(int32 SampleRate, int32 NumOfChannels,
 
 void FWebMAudioBackendSDL::StopStreaming()
 {
-	if (AudioDevice)
+	if (AudioStream)
 	{
-		SDL_CloseAudioDevice(AudioDevice);
-		AudioDevice = 0;
+		SDL_CloseAudioDevice(SDL_GetAudioStreamDevice(AudioStream));
+		AudioStream = nullptr;
 	}
 }
 
 bool FWebMAudioBackendSDL::SendAudio(const FTimespan& Timespan, const uint8* Buffer, size_t BufferSize)
 {
-	int32 Result = SDL_QueueAudio(AudioDevice, Buffer, BufferSize);
+
+	int32 Result = SDL_PutAudioStreamData(AudioStream, Buffer, BufferSize);
 	if (Result < 0)
 	{
 		UE_LOG(LogWebMMoviePlayer, Error, TEXT("SDL_QueueAudio failed: %d"), Result);
@@ -100,7 +88,7 @@ bool FWebMAudioBackendSDL::SendAudio(const FTimespan& Timespan, const uint8* Buf
 	{
 		if (!bPaused)
 		{
-			SDL_PauseAudioDevice(AudioDevice, 0);
+			SDL_ResumeAudioDevice(SDL_GetAudioStreamDevice(AudioStream));
 		}
 		return true;
 	}
@@ -109,7 +97,10 @@ bool FWebMAudioBackendSDL::SendAudio(const FTimespan& Timespan, const uint8* Buf
 void FWebMAudioBackendSDL::Pause(bool bPause)
 {
 	bPaused = bPause;
-	SDL_PauseAudioDevice(AudioDevice, bPaused ? 1 : 0);
+	if(bPause)
+		SDL_PauseAudioDevice(SDL_GetAudioStreamDevice(AudioStream));
+	else
+		SDL_ResumeAudioDevice(SDL_GetAudioStreamDevice(AudioStream));
 }
 
 bool FWebMAudioBackendSDL::IsPaused() const
diff --git a/Engine/Plugins/Runtime/WebMMoviePlayer/Source/WebMMoviePlayer/WebMMoviePlayer.Build.cs b/Engine/Plugins/Runtime/WebMMoviePlayer/Source/WebMMoviePlayer/WebMMoviePlayer.Build.cs
--- a/Engine/Plugins/Runtime/WebMMoviePlayer/Source/WebMMoviePlayer/WebMMoviePlayer.Build.cs
+++ b/Engine/Plugins/Runtime/WebMMoviePlayer/Source/WebMMoviePlayer/WebMMoviePlayer.Build.cs
@@ -23,7 +23,7 @@ public class WebMMoviePlayer : ModuleRules
 
 		if (Target.IsInPlatformGroup(UnrealPlatformGroup.Unix))
 		{
-			PrivateDependencyModuleNames.Add("SDL2");
+			PrivateDependencyModuleNames.Add("SDL3");
 
 			PrivateIncludePaths.Add("WebMMoviePlayer/Private/Audio/Unix");
 		}
diff --git a/Engine/Source/Developer/DesktopPlatform/DesktopPlatform.Build.cs b/Engine/Source/Developer/DesktopPlatform/DesktopPlatform.Build.cs
--- a/Engine/Source/Developer/DesktopPlatform/DesktopPlatform.Build.cs
+++ b/Engine/Source/Developer/DesktopPlatform/DesktopPlatform.Build.cs
@@ -35,7 +35,7 @@ public class DesktopPlatform : ModuleRules
 				DynamicallyLoadedModuleNames.Add("SlateFileDialogs");
 			}
 
-			AddEngineThirdPartyPrivateStaticDependencies(Target, "SDL2");
+			AddEngineThirdPartyPrivateStaticDependencies(Target, "SDL3");
 		}
 
 		CppCompileWarningSettings.UnsafeTypeCastWarningLevel = WarningLevel.Error;
diff --git a/Engine/Source/Developer/ShaderFormatOpenGL/Private/OpenGLShaderCompiler.cpp b/Engine/Source/Developer/ShaderFormatOpenGL/Private/OpenGLShaderCompiler.cpp
--- a/Engine/Source/Developer/ShaderFormatOpenGL/Private/OpenGLShaderCompiler.cpp
+++ b/Engine/Source/Developer/ShaderFormatOpenGL/Private/OpenGLShaderCompiler.cpp
@@ -34,7 +34,7 @@
 #elif PLATFORM_LINUX
 	#include <GL/glcorearb.h>
 	#include <GL/glext.h>
-	#include "SDL.h"
+	#include <SDL3/SDL.h>
 	#include <stdio.h>
 	#include <wchar.h>
 	typedef SDL_Window*		SDL_HWindow;
@@ -294,8 +294,8 @@ static void _PlatformCreateDummyGLWindow(FPlatformOpenGLContext *OutContext)
 
 	// Create a dummy window.
 	OutContext->hWnd = SDL_CreateWindow(NULL,
-		0, 0, 1, 1,
-		SDL_WINDOW_OPENGL | SDL_WINDOW_BORDERLESS | SDL_WINDOW_HIDDEN | SDL_WINDOW_SKIP_TASKBAR );
+		1, 1,
+		SDL_WINDOW_OPENGL | SDL_WINDOW_BORDERLESS | SDL_WINDOW_HIDDEN | SDL_WINDOW_UTILITY );
 }
 
 static void _PlatformCreateOpenGLContextCore(FPlatformOpenGLContext* OutContext)
@@ -327,7 +327,7 @@ static void PlatformInitOpenGL(void*& ContextPtr, void*& PrevContextPtr, int InM
 		}
 		else
 		{
-			Uint32 InitializedSubsystemsMask = SDL_WasInit(SDL_INIT_EVERYTHING);
+			Uint32 InitializedSubsystemsMask = SDL_WasInit(SDL_INIT_VIDEO);
 			if ((InitializedSubsystemsMask & SDL_INIT_VIDEO) == 0)
 			{
 				SDL_InitSubSystem(SDL_INIT_VIDEO);
@@ -401,7 +401,7 @@ static void PlatformInitOpenGL(void*& ContextPtr, void*& PrevContextPtr, int InM
 static void PlatformReleaseOpenGL(void* ContextPtr, void* PrevContextPtr)
 {
 	SDL_GL_MakeCurrent(SDL_GL_GetCurrentWindow(), reinterpret_cast<SDL_HGLContext>(PrevContextPtr));
-	SDL_GL_DeleteContext(reinterpret_cast<SDL_HGLContext>(ContextPtr));
+	SDL_GL_DestroyContext(reinterpret_cast<SDL_HGLContext>(ContextPtr));
 }
 #elif PLATFORM_MAC
 static void PlatformInitOpenGL(void*& ContextPtr, void*& PrevContextPtr, int InMajorVersion, int InMinorVersion)
diff --git a/Engine/Source/Developer/ShaderFormatOpenGL/ShaderFormatOpenGL.Build.cs b/Engine/Source/Developer/ShaderFormatOpenGL/ShaderFormatOpenGL.Build.cs
--- a/Engine/Source/Developer/ShaderFormatOpenGL/ShaderFormatOpenGL.Build.cs
+++ b/Engine/Source/Developer/ShaderFormatOpenGL/ShaderFormatOpenGL.Build.cs
@@ -30,7 +30,7 @@ public class ShaderFormatOpenGL : ModuleRules
 
         if (Target.IsInPlatformGroup(UnrealPlatformGroup.Linux))
         {
-            AddEngineThirdPartyPrivateStaticDependencies(Target, "SDL2");
+            AddEngineThirdPartyPrivateStaticDependencies(Target, "SDL3");
         }
 		if (Target.Platform == UnrealTargetPlatform.Mac || Target.Platform == UnrealTargetPlatform.Win64 || Target.Platform == UnrealTargetPlatform.Linux)
 		{
diff --git a/Engine/Source/Developer/StandaloneRenderer/Private/Linux/OpenGL/SlateOpenGLContext.cpp b/Engine/Source/Developer/StandaloneRenderer/Private/Linux/OpenGL/SlateOpenGLContext.cpp
--- a/Engine/Source/Developer/StandaloneRenderer/Private/Linux/OpenGL/SlateOpenGLContext.cpp
+++ b/Engine/Source/Developer/StandaloneRenderer/Private/Linux/OpenGL/SlateOpenGLContext.cpp
@@ -156,13 +156,13 @@ static SDL_Window* CreateDummyGLWindow()
 	FPlatformApplicationMisc::InitSDL(); //	will not initialize more than once
 
 #if DO_CHECK
-	uint32 InitializedSubsystems = SDL_WasInit(SDL_INIT_EVERYTHING);
+	uint32 InitializedSubsystems = SDL_WasInit(SDL_INIT_VIDEO);
 	check(InitializedSubsystems & SDL_INIT_VIDEO);
 #endif // DO_CHECK
 
 	// Create a dummy window.
 	SDL_Window *h_wnd = SDL_CreateWindow(	NULL,
-		0, 0, 1, 1,
+		1, 1,
 		SDL_WINDOW_OPENGL | SDL_WINDOW_BORDERLESS | SDL_WINDOW_HIDDEN );
 
 	return h_wnd;
@@ -193,7 +193,7 @@ void FSlateOpenGLContext::Initialize( void* InWindow, const FSlateOpenGLContext*
 		bReleaseWindowOnDestroy = true;
 	}
 
-	if (SDL_GL_LoadLibrary(nullptr))
+	if (SDL_GL_LoadLibrary(nullptr) == false)
 	{
 		FString SdlError(ANSI_TO_TCHAR(SDL_GetError()));
 		UE_LOG(LogInit, Fatal, TEXT("FSlateOpenGLContext::Initialize - Unable to dynamically load libGL: %s\n."), *SdlError);
@@ -288,7 +288,7 @@ void FSlateOpenGLContext::Destroy()
 #if LINUX_USE_OPENGL_3_2
 		glDeleteVertexArrays(1, &VertexArrayObject);
 #endif // LINUX_USE_OPENGL_3_2
-		SDL_GL_DeleteContext( Context );
+		SDL_GL_DestroyContext( Context );
 
 		if	( bReleaseWindowOnDestroy )
 		{
diff --git a/Engine/Source/Developer/StandaloneRenderer/Private/StandaloneRendererPlatformHeaders.h b/Engine/Source/Developer/StandaloneRenderer/Private/StandaloneRendererPlatformHeaders.h
--- a/Engine/Source/Developer/StandaloneRenderer/Private/StandaloneRendererPlatformHeaders.h
+++ b/Engine/Source/Developer/StandaloneRenderer/Private/StandaloneRendererPlatformHeaders.h
@@ -54,6 +54,6 @@
 //#define GLCOREARB_PROTOTYPES 
 #include <GL/glcorearb.h>
 //#include <GL/glext.h>
-#include "SDL.h"
+#include <SDL3/SDL.h>
 
 #endif // PLATFORM_LINUX
diff --git a/Engine/Source/Developer/StandaloneRenderer/StandaloneRenderer.Build.cs b/Engine/Source/Developer/StandaloneRenderer/StandaloneRenderer.Build.cs
--- a/Engine/Source/Developer/StandaloneRenderer/StandaloneRenderer.Build.cs
+++ b/Engine/Source/Developer/StandaloneRenderer/StandaloneRenderer.Build.cs
@@ -31,7 +31,7 @@ public class StandaloneRenderer : ModuleRules
 		}
 		else if (Target.IsInPlatformGroup(UnrealPlatformGroup.Linux))
 		{
-			AddEngineThirdPartyPrivateStaticDependencies(Target, "SDL2");
+			AddEngineThirdPartyPrivateStaticDependencies(Target, "SDL3");
 		}
 		else if (Target.IsInPlatformGroup(UnrealPlatformGroup.IOS))
 		{
diff --git a/Engine/Source/Programs/IncludeTool/IncludeTool/Rules.cs b/Engine/Source/Programs/IncludeTool/IncludeTool/Rules.cs
--- a/Engine/Source/Programs/IncludeTool/IncludeTool/Rules.cs
+++ b/Engine/Source/Programs/IncludeTool/IncludeTool/Rules.cs
@@ -90,6 +90,7 @@ namespace IncludeTool
 			"fios2.h",
 			"SDL.h",
 			"SDL_timer.h",
+			"SDL3/",
 			"zlib.h",
 			"sqlite3.h",
 			"emscripten.h",
diff --git a/Engine/Source/Programs/TestPAL/TestPAL.Build.cs b/Engine/Source/Programs/TestPAL/TestPAL.Build.cs
--- a/Engine/Source/Programs/TestPAL/TestPAL.Build.cs
+++ b/Engine/Source/Programs/TestPAL/TestPAL.Build.cs
@@ -30,7 +30,7 @@ public class TestPAL : ModuleRules
 		{
 			PrivateDependencyModuleNames.AddRange(
 				new string[] {
-					"SDL2",
+					"SDL3",
 				}
 			);
 		}
diff --git a/Engine/Source/Runtime/ApplicationCore/ApplicationCore.Build.cs b/Engine/Source/Runtime/ApplicationCore/ApplicationCore.Build.cs
--- a/Engine/Source/Runtime/ApplicationCore/ApplicationCore.Build.cs
+++ b/Engine/Source/Runtime/ApplicationCore/ApplicationCore.Build.cs
@@ -60,7 +60,7 @@ public class ApplicationCore : ModuleRules
 		else if (Target.IsInPlatformGroup(UnrealPlatformGroup.Linux))
 		{
 			AddEngineThirdPartyPrivateStaticDependencies(Target,
-				"SDL2"
+				"SDL3"
 			);
 
 			// We need FreeType2 and GL for the Splash, but only in the Editor
diff --git a/Engine/Source/Runtime/ApplicationCore/Private/Linux/LinuxApplication.cpp b/Engine/Source/Runtime/ApplicationCore/Private/Linux/LinuxApplication.cpp
--- a/Engine/Source/Runtime/ApplicationCore/Private/Linux/LinuxApplication.cpp
+++ b/Engine/Source/Runtime/ApplicationCore/Private/Linux/LinuxApplication.cpp
@@ -78,21 +78,28 @@ FLinuxApplication* FLinuxApplication::CreateLinuxApplication()
 	}
 
 #if DO_CHECK
-	uint32 InitializedSubsystems = SDL_WasInit(SDL_INIT_EVERYTHING);
+	uint32 InitializedSubsystems = SDL_WasInit(0);
 	check(InitializedSubsystems & SDL_INIT_EVENTS);
 	check(InitializedSubsystems & SDL_INIT_JOYSTICK);
-	check(InitializedSubsystems & SDL_INIT_GAMECONTROLLER);
+	check(InitializedSubsystems & SDL_INIT_GAMEPAD);
  	check(InitializedSubsystems & SDL_INIT_HAPTIC);
 #endif // DO_CHECK
 
 	LinuxApplication = new FLinuxApplication();
 
-	for (int i = 0; i < SDL_NumJoysticks(); ++i)
+	int NumJoysticks;
+	SDL_JoystickID* Joysticks = SDL_GetJoysticks(&NumJoysticks);
+	if(Joysticks)
 	{
-		if (SDL_IsGameController(i))
+		for (int i = 0; i < NumJoysticks; ++i)
 		{
- 			LinuxApplication->AddGameController(i);
+			SDL_JoystickID InstanceID = Joysticks[i];
+			if (SDL_IsGamepad(InstanceID))
+			{
+ 				LinuxApplication->AddGameController(InstanceID);
+			}
 		}
+		SDL_free(Joysticks);
 	}
 
 	return LinuxApplication;
@@ -136,12 +143,12 @@ void FLinuxApplication::DestroyApplication()
 	{
 		if(ControllerIt.Value().Controller != nullptr)
 		{
-			SDL_GameControllerClose(ControllerIt.Value().Controller);
+			SDL_CloseGamepad(ControllerIt.Value().Controller);
 		}
 
 		if(ControllerIt.Value().Haptic != nullptr)
 		{
-			SDL_HapticClose(ControllerIt.Value().Haptic);
+			SDL_CloseHaptic(ControllerIt.Value().Haptic);
 		}
 	}
 	ControllerStates.Empty();
@@ -240,8 +247,8 @@ void FLinuxApplication::AddPendingEvent( SDL_Event SDLEvent )
 
 bool FLinuxApplication::GeneratesKeyCharMessage(const SDL_KeyboardEvent & KeyDownEvent)
 {
-	bool bCmdKeyPressed = (KeyDownEvent.keysym.mod & (KMOD_LCTRL | KMOD_RCTRL)) != 0;
-	const SDL_Keycode Sym = KeyDownEvent.keysym.sym;
+	bool bCmdKeyPressed = (KeyDownEvent.mod & (SDL_KMOD_LCTRL | SDL_KMOD_RCTRL)) != 0;
+	const SDL_Keycode Sym = KeyDownEvent.key;
 
 	// filter out command keys, non-ASCI and arrow keycodes that don't generate WM_CHAR under Windows (TODO: find a table?)
 	return !bCmdKeyPressed && Sym < 128 &&
@@ -251,7 +258,7 @@ bool FLinuxApplication::GeneratesKeyCharMessage(const SDL_KeyboardEvent & KeyDow
 // Windows handles translating numpad numbers to arrow keys, but we have to do it manually
 static SDL_Keycode TranslateNumLockKeySyms(const SDL_Keycode KeySym)
 {
-	if ((SDL_GetModState() & KMOD_NUM) == 0)
+	if ((SDL_GetModState() & SDL_KMOD_NUM) == 0)
 	{
 		switch (KeySym)
 		{
@@ -309,10 +316,10 @@ void FLinuxApplication::ProcessDeferredMessage( SDL_Event Event )
 
 	switch(Event.type)
 	{
-	case SDL_KEYDOWN:
+	case SDL_EVENT_KEY_DOWN:
 		{
 			const SDL_KeyboardEvent &KeyEvent = Event.key;
-			SDL_Keycode KeySym = KeyEvent.keysym.sym;
+			SDL_Keycode KeySym = KeyEvent.key;
 			const uint32 CharCode = CharCodeFromSDLKeySym(KeySym);
 			KeySym = TranslateNumLockKeySyms(KeySym);
 			const bool bIsRepeated = KeyEvent.repeat != 0;
@@ -331,17 +338,17 @@ void FLinuxApplication::ProcessDeferredMessage( SDL_Event Event )
 			}
 		}
 		break;
-	case SDL_KEYUP:
+	case SDL_EVENT_KEY_UP:
 		{
 			const SDL_KeyboardEvent &KeyEvent = Event.key;
-			const SDL_Keycode KeySym = KeyEvent.keysym.sym;
+			const SDL_Keycode KeySym = KeyEvent.key;
 			const uint32 CharCode = CharCodeFromSDLKeySym(KeySym);
 			const bool IsRepeat = KeyEvent.repeat != 0;
 
 			MessageHandler->OnKeyUp( KeySym, CharCode, IsRepeat );
 		}
 		break;
-	case SDL_TEXTINPUT:
+	case SDL_EVENT_TEXT_INPUT:
 		{
 			// Slate now gets all its text from here, I hope.
 			const bool bIsRepeated = false;  //Event.key.repeat != 0;
@@ -352,7 +359,7 @@ void FLinuxApplication::ProcessDeferredMessage( SDL_Event Event )
 			}
 		}
 		break;
-	case SDL_MOUSEMOTION:
+	case SDL_EVENT_MOUSE_MOTION:
 		{
 			SDL_MouseMotionEvent motionEvent = Event.motion;
 			FLinuxCursor *LinuxCursor = (FLinuxCursor*)Cursor.Get();
@@ -370,7 +377,7 @@ void FLinuxApplication::ProcessDeferredMessage( SDL_Event Event )
 					int32 BorderSizeX, BorderSizeY;
 					CurrentEventWindow->GetNativeBordersSize(BorderSizeX, BorderSizeY);
 
-					LinuxCursor->SetCachedPosition((int32)(motionEvent.x + Props.Location.X + BorderSizeX), (int32)(motionEvent.y + Props.Location.Y + BorderSizeY));
+					LinuxCursor->SetCachedPosition((float)(motionEvent.x + Props.Location.X + BorderSizeX), (float)(motionEvent.y + Props.Location.Y + BorderSizeY));
 				}
 
 				if( !CurrentEventWindow->GetDefinition().HasOSWindowBorder )
@@ -386,7 +393,7 @@ void FLinuxApplication::ProcessDeferredMessage( SDL_Event Event )
 
 			if(bUsingHighPrecisionMouseInput)
 			{
-				MessageHandler->OnRawMouseMove(motionEvent.xrel, motionEvent.yrel);
+				MessageHandler->OnRawMouseMove((int32)motionEvent.xrel, (int32)motionEvent.yrel);
 			}
 			else
 			{
@@ -394,8 +401,8 @@ void FLinuxApplication::ProcessDeferredMessage( SDL_Event Event )
 			}
 		}
 		break;
-	case SDL_MOUSEBUTTONDOWN:
-	case SDL_MOUSEBUTTONUP:
+	case SDL_EVENT_MOUSE_BUTTON_DOWN:
+	case SDL_EVENT_MOUSE_BUTTON_UP:
 		{
 			SDL_MouseButtonEvent buttonEvent = Event.button;
 
@@ -422,7 +429,7 @@ void FLinuxApplication::ProcessDeferredMessage( SDL_Event Event )
 				break;
 			}
 			
-			if (buttonEvent.type == SDL_MOUSEBUTTONUP)
+			if (buttonEvent.type == SDL_EVENT_MOUSE_BUTTON_UP)
 			{
 				MessageHandler->OnMouseUp(button);
 				
@@ -460,7 +467,7 @@ void FLinuxApplication::ProcessDeferredMessage( SDL_Event Event )
 					// Check if we have to set the focus.
 					if(CurrentFocusWindow != CurrentEventWindow)
 					{
-						SDL_SetWindowInputFocus(CurrentEventWindow->GetHWnd());
+						SDL_RaiseWindow(CurrentEventWindow->GetHWnd());
 					}
 
 					MessageHandler->OnMouseDown(CurrentEventWindow, button);
@@ -473,7 +480,7 @@ void FLinuxApplication::ProcessDeferredMessage( SDL_Event Event )
 			}
 		}
 		break;
-	case SDL_MOUSEWHEEL:
+	case SDL_EVENT_MOUSE_WHEEL:
 		{
 			SDL_MouseWheelEvent *WheelEvent = &Event.wheel;
 			float Amount = (float)WheelEvent->y * fMouseWheelScrollAccel;
@@ -481,21 +488,21 @@ void FLinuxApplication::ProcessDeferredMessage( SDL_Event Event )
 			MessageHandler->OnMouseWheel(Amount);
 		}
 		break;
-	case SDL_CONTROLLERDEVICEADDED:
+	case SDL_EVENT_GAMEPAD_ADDED:
 		{
-			SDL_ControllerDeviceEvent& ControllerEvent = Event.cdevice;
-			AddGameController(ControllerEvent.which);
+			SDL_GamepadDeviceEvent& GamepadEvent = Event.gdevice;
+			AddGameController(GamepadEvent.which);
 		}
 		break;
-	case SDL_CONTROLLERDEVICEREMOVED:
+	case SDL_EVENT_GAMEPAD_REMOVED:
 		{
-			SDL_ControllerDeviceEvent& ControllerEvent = Event.cdevice;
-			RemoveGameController(ControllerEvent.which);
+			SDL_GamepadDeviceEvent& GamepadEvent = Event.gdevice;
+			RemoveGameController(GamepadEvent.which);
 		}
 		break;
-	case SDL_CONTROLLERAXISMOTION:
+	case SDL_EVENT_GAMEPAD_AXIS_MOTION:
 		{
-			SDL_ControllerAxisEvent caxisEvent = Event.caxis;
+			SDL_GamepadAxisEvent caxisEvent = Event.gaxis;
 			FGamepadKeyNames::Type Axis = FGamepadKeyNames::Invalid;
 			float AxisValue = ShortToNormalFloat(caxisEvent.value);
 
@@ -509,7 +516,7 @@ void FLinuxApplication::ProcessDeferredMessage( SDL_Event Event )
 
 			switch (caxisEvent.axis)
 			{
-			case SDL_CONTROLLER_AXIS_LEFTX:
+			case SDL_GAMEPAD_AXIS_LEFTX:
 				Axis = FGamepadKeyNames::LeftAnalogX;
 				if(caxisEvent.value > GAMECONTROLLER_LEFT_THUMB_DEADZONE)
 				{
@@ -538,7 +545,7 @@ void FLinuxApplication::ProcessDeferredMessage( SDL_Event Event )
 					ControllerState.AnalogOverThreshold[1] = false;
 				}
 				break;
-			case SDL_CONTROLLER_AXIS_LEFTY:
+			case SDL_GAMEPAD_AXIS_LEFTY:
 				Axis = FGamepadKeyNames::LeftAnalogY;
 				AxisValue *= -1;
 				if(caxisEvent.value > GAMECONTROLLER_LEFT_THUMB_DEADZONE)
@@ -568,7 +575,7 @@ void FLinuxApplication::ProcessDeferredMessage( SDL_Event Event )
 					ControllerState.AnalogOverThreshold[3] = false;
 				}
 				break;
-			case SDL_CONTROLLER_AXIS_RIGHTX:
+			case SDL_GAMEPAD_AXIS_RIGHTX:
 				Axis = FGamepadKeyNames::RightAnalogX;
 				if(caxisEvent.value > GAMECONTROLLER_RIGHT_THUMB_DEADZONE)
 				{
@@ -597,7 +604,7 @@ void FLinuxApplication::ProcessDeferredMessage( SDL_Event Event )
 					ControllerState.AnalogOverThreshold[5] = false;
 				}
 				break;
-			case SDL_CONTROLLER_AXIS_RIGHTY:
+			case SDL_GAMEPAD_AXIS_RIGHTY:
 				Axis = FGamepadKeyNames::RightAnalogY;
 				AxisValue *= -1;
 				if(caxisEvent.value > GAMECONTROLLER_RIGHT_THUMB_DEADZONE)
@@ -627,7 +634,7 @@ void FLinuxApplication::ProcessDeferredMessage( SDL_Event Event )
 					ControllerState.AnalogOverThreshold[7] = false;
 				}
 				break;
-			case SDL_CONTROLLER_AXIS_TRIGGERLEFT:
+			case SDL_GAMEPAD_AXIS_LEFT_TRIGGER:
 				Axis = FGamepadKeyNames::LeftTriggerAnalog;
 				if(caxisEvent.value > GAMECONTROLLER_TRIGGER_THRESHOLD)
 				{
@@ -643,7 +650,7 @@ void FLinuxApplication::ProcessDeferredMessage( SDL_Event Event )
 					ControllerState.AnalogOverThreshold[8] = false;
 				}
 				break;
-			case SDL_CONTROLLER_AXIS_TRIGGERRIGHT:
+			case SDL_GAMEPAD_AXIS_RIGHT_TRIGGER:
 				Axis = FGamepadKeyNames::RightTriggerAnalog;
 				if(caxisEvent.value > GAMECONTROLLER_TRIGGER_THRESHOLD)
 				{
@@ -670,10 +677,10 @@ void FLinuxApplication::ProcessDeferredMessage( SDL_Event Event )
 			}
 		}
 		break;
-	case SDL_CONTROLLERBUTTONDOWN:
-	case SDL_CONTROLLERBUTTONUP:
+	case SDL_EVENT_GAMEPAD_BUTTON_DOWN:
+	case SDL_EVENT_GAMEPAD_BUTTON_UP:
 		{
-			SDL_ControllerButtonEvent cbuttonEvent = Event.cbutton;
+			SDL_GamepadButtonEvent cbuttonEvent = Event.gbutton;
 			FGamepadKeyNames::Type Button = FGamepadKeyNames::Invalid;
 
 			if (!ControllerStates.Contains(cbuttonEvent.which))
@@ -683,57 +690,57 @@ void FLinuxApplication::ProcessDeferredMessage( SDL_Event Event )
 
 			switch (cbuttonEvent.button)
 			{
-			case SDL_CONTROLLER_BUTTON_A:
-				Button = FGamepadKeyNames::FaceButtonBottom;
-				break;
-			case SDL_CONTROLLER_BUTTON_B:
-				Button = FGamepadKeyNames::FaceButtonRight;
-				break;
-			case SDL_CONTROLLER_BUTTON_X:
-				Button = FGamepadKeyNames::FaceButtonLeft;
-				break;
-			case SDL_CONTROLLER_BUTTON_Y:
-				Button = FGamepadKeyNames::FaceButtonTop;
-				break;
-			case SDL_CONTROLLER_BUTTON_BACK:
-				Button = FGamepadKeyNames::SpecialLeft;
-				break;
-			case SDL_CONTROLLER_BUTTON_START:
-				Button = FGamepadKeyNames::SpecialRight;
-				break;
-			case SDL_CONTROLLER_BUTTON_LEFTSTICK:
-				Button = FGamepadKeyNames::LeftThumb;
-				break;
-			case SDL_CONTROLLER_BUTTON_RIGHTSTICK:
-				Button = FGamepadKeyNames::RightThumb;
-				break;
-			case SDL_CONTROLLER_BUTTON_LEFTSHOULDER:
-				Button = FGamepadKeyNames::LeftShoulder;
-				break;
-			case SDL_CONTROLLER_BUTTON_RIGHTSHOULDER:
-				Button = FGamepadKeyNames::RightShoulder;
-				break;
-			case SDL_CONTROLLER_BUTTON_DPAD_UP:
-				Button = FGamepadKeyNames::DPadUp;
-				break;
-			case SDL_CONTROLLER_BUTTON_DPAD_DOWN:
-				Button = FGamepadKeyNames::DPadDown;
-				break;
-			case SDL_CONTROLLER_BUTTON_DPAD_LEFT:
-				Button = FGamepadKeyNames::DPadLeft;
-				break;
-			case SDL_CONTROLLER_BUTTON_DPAD_RIGHT:
-				Button = FGamepadKeyNames::DPadRight;
-				break;
-			default:
-				break;
+				case SDL_GAMEPAD_BUTTON_SOUTH:
+					Button = FGamepadKeyNames::FaceButtonBottom;
+					break;
+				case SDL_GAMEPAD_BUTTON_EAST:
+					Button = FGamepadKeyNames::FaceButtonRight;
+					break;
+				case SDL_GAMEPAD_BUTTON_WEST:
+					Button = FGamepadKeyNames::FaceButtonLeft;
+					break;
+				case SDL_GAMEPAD_BUTTON_NORTH:
+					Button = FGamepadKeyNames::FaceButtonTop;
+					break;
+				case SDL_GAMEPAD_BUTTON_BACK:
+					Button = FGamepadKeyNames::SpecialLeft;
+					break;
+				case SDL_GAMEPAD_BUTTON_START:
+					Button = FGamepadKeyNames::SpecialRight;
+					break;
+				case SDL_GAMEPAD_BUTTON_LEFT_STICK:
+					Button = FGamepadKeyNames::LeftThumb;
+					break;
+				case SDL_GAMEPAD_BUTTON_RIGHT_STICK:
+					Button = FGamepadKeyNames::RightThumb;
+					break;
+				case SDL_GAMEPAD_BUTTON_LEFT_SHOULDER:
+					Button = FGamepadKeyNames::LeftShoulder;
+					break;
+				case SDL_GAMEPAD_BUTTON_RIGHT_SHOULDER:
+					Button = FGamepadKeyNames::RightShoulder;
+					break;
+				case SDL_GAMEPAD_BUTTON_DPAD_UP:
+					Button = FGamepadKeyNames::DPadUp;
+					break;
+				case SDL_GAMEPAD_BUTTON_DPAD_DOWN:
+					Button = FGamepadKeyNames::DPadDown;
+					break;
+				case SDL_GAMEPAD_BUTTON_DPAD_LEFT:
+					Button = FGamepadKeyNames::DPadLeft;
+					break;
+				case SDL_GAMEPAD_BUTTON_DPAD_RIGHT:
+					Button = FGamepadKeyNames::DPadRight;
+					break;
+				default:
+					break;
 			}
 
 			if (Button != FGamepadKeyNames::Invalid)
 			{
 				FPlatformUserId UserId = IPlatformInputDeviceMapper::Get().GetUserForInputDevice(ControllerStates[cbuttonEvent.which].DeviceId);
 
-				if(cbuttonEvent.type == SDL_CONTROLLERBUTTONDOWN)
+				if(cbuttonEvent.type == SDL_EVENT_GAMEPAD_BUTTON_DOWN)
 				{
 					MessageHandler->OnControllerButtonPressed(Button, UserId, ControllerStates[cbuttonEvent.which].DeviceId, false);
 				}
@@ -745,270 +752,247 @@ void FLinuxApplication::ProcessDeferredMessage( SDL_Event Event )
 		}
 		break;
 
-	case SDL_WINDOWEVENT:
+	case SDL_EVENT_WINDOW_PIXEL_SIZE_CHANGED:
 		{
-			SDL_WindowEvent windowEvent = Event.window;
+			int NewWidth  = Event.window.data1;
+			int NewHeight = Event.window.data2;
+
+			MessageHandler->OnSizeChanged(
+				CurrentEventWindow.ToSharedRef(),
+				NewWidth,
+				NewHeight,
+				//	bWasMinimized
+				false
+			);
+		}
+		break;
 
-			switch (windowEvent.event)
+	case SDL_EVENT_WINDOW_RESIZED:
+		{
+			MessageHandler->OnResizingWindow( CurrentEventWindow.ToSharedRef() );
+		}
+		break;
+
+	case SDL_EVENT_WINDOW_CLOSE_REQUESTED:
+		{
+			MessageHandler->OnWindowClose( CurrentEventWindow.ToSharedRef() );
+		}
+		break;
+
+	case SDL_EVENT_WINDOW_SHOWN:
+		{
+			// (re)cache native properties
+			//CurrentEventWindow->CacheNativeProperties();
+
+			// A window did show up. Is the whole Application active? If not first activate it (ignore tooltips).
+			if (!bActivateApp && !CurrentEventWindow->IsTooltipWindow())
 			{
-				case SDL_WINDOWEVENT_SIZE_CHANGED:
-					{
-						int NewWidth  = windowEvent.data1;
-						int NewHeight = windowEvent.data2;
-
-						MessageHandler->OnSizeChanged(
-							CurrentEventWindow.ToSharedRef(),
-							NewWidth,
-							NewHeight,
-							//	bWasMinimized
-							false
-						);
-					}
-					break;
+				ActivateApplication();
+			}
 
-				case SDL_WINDOWEVENT_RESIZED:
-					{
-						MessageHandler->OnResizingWindow( CurrentEventWindow.ToSharedRef() );
-					}
-					break;
+			// Check if this window is different then the currently active one. If it is another one
+			// activate that window and if necessary deactivate the one which was active.
+			if (CurrentlyActiveWindow != CurrentEventWindow)
+			{
+				ActivateWindow(CurrentEventWindow);
+			}
 
-				case SDL_WINDOWEVENT_CLOSE:
-					{
-						MessageHandler->OnWindowClose( CurrentEventWindow.ToSharedRef() );
-					}
-					break;
+			// Set focus if the window wants to have a focus when first shown.
+			if (CurrentEventWindow->IsFocusWhenFirstShown())
+			{
+				// If we don't surpress the window activation phase of RaiseWindow, X11/Gnome will bring up a notification window saying
+				// the window is ready
+				bool bActivate = SDL_GetHintBoolean(SDL_HINT_WINDOW_ACTIVATE_WHEN_RAISED, true);
+				SDL_SetHint(SDL_HINT_WINDOW_ACTIVATE_WHEN_RAISED, "0");
+				SDL_RaiseWindow(CurrentEventWindow->GetHWnd());
+				SDL_SetHint(SDL_HINT_WINDOW_ACTIVATE_WHEN_RAISED, bActivate ? "1" : "0");
+			}
+		}
+		break;
 
-				case SDL_WINDOWEVENT_SHOWN:
-					{
-						// (re)cache native properties
-						CurrentEventWindow->CacheNativeProperties();
-
-						// A window did show up. Is the whole Application active? If not first activate it (ignore tooltips).
-						if (!bActivateApp && !CurrentEventWindow->IsTooltipWindow())
-						{
-							ActivateApplication();
-						}
-
-						// Check if this window is different then the currently active one. If it is another one
-						// activate that window and if necessary deactivate the one which was active.
-						if (CurrentlyActiveWindow != CurrentEventWindow)
-						{
-							ActivateWindow(CurrentEventWindow);
-						}
-
-						// Set focus if the window wants to have a focus when first shown.
-						if (CurrentEventWindow->IsFocusWhenFirstShown())
-						{
-							SDL_SetWindowInputFocus(CurrentEventWindow->GetHWnd());
-						}
-					}
-					break;
+	case SDL_EVENT_WINDOW_MOVED:
+		{
+			// Mask away the higher bits (but preserve negative flags), as SDL uses those as flags
+			// See: SDL_WINDOWPOS_UNDEFINED_MASK & SDL_WINDOWPOS_CENTERED_MASK for context
+			int32 ClientScreenX = MaskAwayHigherBits(Event.window.data1);
+			int32 ClientScreenY = MaskAwayHigherBits(Event.window.data2);
 
-				case SDL_WINDOWEVENT_MOVED:
-					{
-						// Mask away the higher bits (but preserve negative flags), as SDL uses those as flags
-						// See: SDL_WINDOWPOS_UNDEFINED_MASK & SDL_WINDOWPOS_CENTERED_MASK for context
-						int32 ClientScreenX = MaskAwayHigherBits(windowEvent.data1);
-						int32 ClientScreenY = MaskAwayHigherBits(windowEvent.data2);
-
-						int32 BorderSizeX, BorderSizeY;
-						CurrentEventWindow->GetNativeBordersSize(BorderSizeX, BorderSizeY);
-						ClientScreenX += BorderSizeX;
-						ClientScreenY += BorderSizeY;
-
-						MessageHandler->OnMovedWindow(CurrentEventWindow.ToSharedRef(), ClientScreenX, ClientScreenY);
-
-						if (bFirstFrameOfWindowMove)
-						{
-							bFirstFrameOfWindowMove = false;
-
-							if (NotificationWindows.Num() > 0)
-							{
-								RaiseNotificationWindows(CurrentEventWindow);
-							}
-						}
-
-						// Mouse dragging is the likely cause of this window move, so invalidate the cached cursor position
-						FLinuxCursor *LinuxCursor = static_cast<FLinuxCursor*>(Cursor.Get());
-						LinuxCursor->InvalidateCaches();
-					}
-					break;
+			int32 BorderSizeX, BorderSizeY;
+			CurrentEventWindow->GetNativeBordersSize(BorderSizeX, BorderSizeY);
+			ClientScreenX += BorderSizeX;
+			ClientScreenY += BorderSizeY;
 
-				case SDL_WINDOWEVENT_MAXIMIZED:
-					{
-						MessageHandler->OnWindowAction(CurrentEventWindow.ToSharedRef(), EWindowAction::Maximize);
-						UE_LOG(LogLinuxWindowEvent, Verbose, TEXT("Window: '%d' got maximized"), CurrentEventWindow->GetID());
+			MessageHandler->OnMovedWindow(CurrentEventWindow.ToSharedRef(), ClientScreenX, ClientScreenY);
 
-						if (NotificationWindows.Num() > 0)
-						{
-							RaiseNotificationWindows(CurrentEventWindow);
-						}
-					}
-					break;
+			if (bFirstFrameOfWindowMove)
+			{
+				bFirstFrameOfWindowMove = false;
 
-				case SDL_WINDOWEVENT_RESTORED:
-					{
-						MessageHandler->OnWindowAction(CurrentEventWindow.ToSharedRef(), EWindowAction::Restore);
-						UE_LOG(LogLinuxWindowEvent, Verbose, TEXT("Window: '%d' got restored"), CurrentEventWindow->GetID());
+				if (NotificationWindows.Num() > 0)
+				{
+					RaiseNotificationWindows(CurrentEventWindow);
+				}
+			}
 
-						if (NotificationWindows.Num() > 0)
-						{
-							RaiseNotificationWindows(CurrentEventWindow);
-						}
-					}
-					break;
+			// Mouse dragging is the likely cause of this window move, so invalidate the cached cursor position
+			FLinuxCursor *LinuxCursor = static_cast<FLinuxCursor*>(Cursor.Get());
+			LinuxCursor->InvalidateCaches();
+		}
+		break;
 
-				case SDL_WINDOWEVENT_ENTER:
-					{
-						if (CurrentEventWindow.IsValid())
-						{
-							MessageHandler->OnCursorSet();
+	case SDL_EVENT_WINDOW_MAXIMIZED:
+		{
+			MessageHandler->OnWindowAction(CurrentEventWindow.ToSharedRef(), EWindowAction::Maximize);
+			UE_LOG(LogLinuxWindowEvent, Verbose, TEXT("Window: '%d' got maximized"), CurrentEventWindow->GetID());
 
-							// Currently Tooltip windows will also get enter/leave events. depending on if this causes issues
-							// should avoid setting the window under cursor for tooltips and use the window under
-							CurrentUnderCursorWindow = CurrentEventWindow;
+			if (NotificationWindows.Num() > 0)
+			{
+				RaiseNotificationWindows(CurrentEventWindow);
+			}
+		}
+		break;
 
-							bInsideOwnWindow = true;
+	case SDL_EVENT_WINDOW_RESTORED:
+		{
+			MessageHandler->OnWindowAction(CurrentEventWindow.ToSharedRef(), EWindowAction::Restore);
+			UE_LOG(LogLinuxWindowEvent, Verbose, TEXT("Window: '%d' got restored"), CurrentEventWindow->GetID());
 
-							UE_LOG(LogLinuxWindow, Verbose, TEXT("Entered one of application windows. Cursor under Window: '%d'"), CurrentEventWindow->GetID());
-						}
-					}
-					break;
+			if (NotificationWindows.Num() > 0)
+			{
+				RaiseNotificationWindows(CurrentEventWindow);
+			}
+		}
+		break;
 
-				case SDL_WINDOWEVENT_LEAVE:
-					{
-						if (CurrentEventWindow.IsValid())
-						{
-							if (GetCapture() != nullptr)
-							{
-								UpdateMouseCaptureWindow((SDL_HWindow)GetCapture());
-							}
-
- 							CurrentUnderCursorWindow = nullptr;
-							bInsideOwnWindow = false;
-							UE_LOG(LogLinuxWindow, Verbose, TEXT("Left an application window we were hovering above."));
-						}
-					}
-					break;
+	case SDL_EVENT_WINDOW_MOUSE_ENTER:
+		{
+			if (CurrentEventWindow.IsValid())
+			{
+				MessageHandler->OnCursorSet();
 
-				case SDL_WINDOWEVENT_HIT_TEST:
-					{
-						// The user clicked into the hit test area (Titlebar for example). Is the whole Application active?
-						// If not, first activate (ignore tooltips).
-						if (!bActivateApp && !CurrentEventWindow->IsTooltipWindow())
-						{
-							ActivateApplication();
-						}
-
-						// Check if this window is different then the currently active one. If it is another one activate this 
-						// window and deactivate the other one.
-						if (CurrentlyActiveWindow != CurrentEventWindow)
-						{
-							ActivateWindow(CurrentEventWindow);
-						}
-
-						// Set the input focus.
-						if (CurrentEventWindow.IsValid())
-						{
-							SDL_SetWindowInputFocus(CurrentEventWindow->GetHWnd());
-						}
-
-						if (NotificationWindows.Num() > 0)
-						{
-							RaiseNotificationWindows(CurrentEventWindow);
-						}
-
-						bFirstFrameOfWindowMove = true;
-					}
-					break;
+				// Currently Tooltip windows will also get enter/leave events. depending on if this causes issues
+				// should avoid setting the window under cursor for tooltips and use the window under
+				CurrentUnderCursorWindow = CurrentEventWindow;
 
-				case SDL_WINDOWEVENT_TAKE_FOCUS:
-					{
-						if (!bActivateApp)
-						{
-							ActivateApplication();
-						}
-
-						// Some windows like notification windows may popup without needing the focus. That is handled in the SDL_WINDOWEVENT_SHOWN case.
-						// The problem would be that the WM will send the Take Focus event and wants to set the focus. We don't want it to set it
-						// for notifications because they are already handled in the above mentioned event.
-						if ((CurrentFocusWindow != CurrentEventWindow) && !CurrentEventWindow->IsNotificationWindow())
-						{
-							SDL_SetWindowInputFocus(CurrentEventWindow->GetHWnd());
-						}
-					}
-					break;
+				bInsideOwnWindow = true;
 
-				case SDL_WINDOWEVENT_FOCUS_GAINED:
-					{
-						UE_LOG(LogLinuxWindowEvent, Verbose, TEXT("WM_SETFOCUS                                 : %d"), CurrentEventWindow->GetID());
+				UE_LOG(LogLinuxWindow, Verbose, TEXT("Entered one of application windows. Cursor under Window: '%d'"), CurrentEventWindow->GetID());
+			}
+		}
+		break;
 
-						CurrentFocusWindow = CurrentEventWindow;
+	case SDL_EVENT_WINDOW_MOUSE_LEAVE:
+		{
+			if (CurrentEventWindow.IsValid())
+			{
+				if (GetCapture() != nullptr)
+				{
+					UpdateMouseCaptureWindow((SDL_HWindow)GetCapture());
+				}
 
-						if (NotificationWindows.Num() > 0)
-						{
-							RaiseNotificationWindows(CurrentEventWindow);
-						}
+					CurrentUnderCursorWindow = nullptr;
+				bInsideOwnWindow = false;
+				UE_LOG(LogLinuxWindow, Verbose, TEXT("Left an application window we were hovering above."));
+			}
+		}
+		break;
 
-						// We have gained focus, we can stop trying to check if we need to deactivate
-						FocusOutDeactivationTime = 0.0;
-					}
-					break;
+	case SDL_EVENT_WINDOW_HIT_TEST:
+		{
+			// The user clicked into the hit test area (Titlebar for example). Is the whole Application active?
+			// If not, first activate (ignore tooltips).
+			if (!bActivateApp && !CurrentEventWindow->IsTooltipWindow())
+			{
+				ActivateApplication();
+			}
 
-				case SDL_WINDOWEVENT_FOCUS_LOST:
-					{
-						UE_LOG(LogLinuxWindowEvent, Verbose, TEXT("WM_KILLFOCUS                                : %d"), CurrentEventWindow->GetID());
-
-						// OK, the active window lost focus. This could mean the app went completely out of
-						// focus. That means the app must be deactivated. To make sure that the user did
-						// not click to another window we delay the deactivation.
-						// TODO Figure out if the delay time may cause problems.
-						if(CurrentFocusWindow == CurrentEventWindow)
-						{
-							// Only do if the application is active.
-							if(bActivateApp)
-							{
-								FocusOutDeactivationTime = FPlatformTime::Seconds() + DeactivationThreadshold;
-							}
-						}
-						CurrentFocusWindow = nullptr;
-					}
-					break;
-				case SDL_WINDOWEVENT_HIDDEN:	// intended fall-through
-				case SDL_WINDOWEVENT_EXPOSED:	// intended fall-through
-				case SDL_WINDOWEVENT_MINIMIZED:	// intended fall-through
-				default:
-					break;
+			// Check if this window is different then the currently active one. If it is another one activate this 
+			// window and deactivate the other one.
+			if (CurrentlyActiveWindow != CurrentEventWindow)
+			{
+				ActivateWindow(CurrentEventWindow);
 			}
+
+			// Set the input focus.
+			if (CurrentEventWindow.IsValid())
+			{
+				SDL_RaiseWindow(CurrentEventWindow->GetHWnd());
+			}
+
+			if (NotificationWindows.Num() > 0)
+			{
+				RaiseNotificationWindows(CurrentEventWindow);
+			}
+
+			bFirstFrameOfWindowMove = true;
 		}
 		break;
 
-	case SDL_DROPBEGIN:
+	case SDL_EVENT_WINDOW_FOCUS_GAINED:
+		{
+			UE_LOG(LogLinuxWindowEvent, Verbose, TEXT("WM_SETFOCUS                                 : %d"), CurrentEventWindow->GetID());
+
+			CurrentFocusWindow = CurrentEventWindow;
+
+			if (NotificationWindows.Num() > 0)
+			{
+				RaiseNotificationWindows(CurrentEventWindow);
+			}
+
+			// We have gained focus, we can stop trying to check if we need to deactivate
+			FocusOutDeactivationTime = 0.0;
+		}
+		break;
+
+	case SDL_EVENT_WINDOW_FOCUS_LOST:
+		{
+			UE_LOG(LogLinuxWindowEvent, Verbose, TEXT("WM_KILLFOCUS                                : %d"), CurrentEventWindow->GetID());
+
+			// OK, the active window lost focus. This could mean the app went completely out of
+			// focus. That means the app must be deactivated. To make sure that the user did
+			// not click to another window we delay the deactivation.
+			// TODO Figure out if the delay time may cause problems.
+			if(CurrentFocusWindow == CurrentEventWindow)
+			{
+				// Only do if the application is active.
+				if(bActivateApp)
+				{
+					FocusOutDeactivationTime = FPlatformTime::Seconds() + DeactivationThreadshold;
+				}
+			}
+			CurrentFocusWindow = nullptr;
+		}
+		break;
+
+	case SDL_EVENT_WINDOW_HIDDEN:	// intended fall-through
+	case SDL_EVENT_WINDOW_EXPOSED:	// intended fall-through
+	case SDL_EVENT_WINDOW_MINIMIZED:	// intended fall-through
+		break;
+
+	case SDL_EVENT_DROP_BEGIN:
 		{
 			check(DragAndDropQueue.Num() == 0);  // did we get confused?
 			check(DragAndDropTextQueue.Num() == 0);  // did we get confused?
 		}
 		break;
 
-	case SDL_DROPFILE:
+	case SDL_EVENT_DROP_FILE:
 		{
-			FString tmp = StringUtility::UnescapeURI(UTF8_TO_TCHAR(Event.drop.file));
+			FString tmp = StringUtility::UnescapeURI(UTF8_TO_TCHAR(Event.drop.data));
 			DragAndDropQueue.Add(tmp);
-			SDL_free(Event.drop.file);
 			UE_LOG(LogLinuxWindowEvent, Verbose, TEXT("File dropped: %s"), *tmp);
 		}
 		break;
 
-	case SDL_DROPTEXT:
+	case SDL_EVENT_DROP_TEXT:
 		{
-			FString tmp = UTF8_TO_TCHAR(Event.drop.file);
+			FString tmp = UTF8_TO_TCHAR(Event.drop.data);
 			DragAndDropTextQueue.Add(tmp);
-			SDL_free(Event.drop.file);
 			UE_LOG(LogLinuxWindowEvent, Verbose, TEXT("Text dropped: %s"), *tmp);
 		}
 		break;
 
-	case SDL_DROPCOMPLETE:
+	case SDL_EVENT_DROP_COMPLETE:
 		{
 			if (DragAndDropQueue.Num() > 0)
 			{
@@ -1030,15 +1014,15 @@ void FLinuxApplication::ProcessDeferredMessage( SDL_Event Event )
 		}
 		break;
 
-	case SDL_FINGERDOWN:
+	case SDL_EVENT_FINGER_DOWN:
 		{
-			UE_LOG(LogLinuxWindow, Verbose, TEXT("Finger %llu is down at (%f, %f)"), Event.tfinger.fingerId, Event.tfinger.x, Event.tfinger.y);
+			UE_LOG(LogLinuxWindow, Verbose, TEXT("Finger %llu is down at (%f, %f)"), Event.tfinger.fingerID, Event.tfinger.x, Event.tfinger.y);
 
 			// touch events can have no window associated with them, in that case ignore (with a warning)
 			if (LIKELY(!bWindowlessEvent))
 			{
 				// remove touch context even if it existed
-				uint64 FingerId = static_cast<uint64>(Event.tfinger.fingerId);
+				uint64 FingerId = static_cast<uint64>(Event.tfinger.fingerID);
 
 				FTouchContext* TouchContext = Touches.Find(FingerId);
 				if (UNLIKELY(Touches.Find(FingerId) != nullptr))
@@ -1051,7 +1035,7 @@ void FLinuxApplication::ProcessDeferredMessage( SDL_Event Event )
 				FTouchContext NewTouch;
 				NewTouch.TouchIndex = GetFirstFreeTouchId();
 				NewTouch.Location = GetTouchEventLocation(NativeWindow, Event);
-				NewTouch.DeviceId = Event.tfinger.touchId;
+				NewTouch.DeviceId = Event.tfinger.touchID;
 				Touches.Add(FingerId, NewTouch);
 				TouchIds[NewTouch.TouchIndex] = TOptional<uint64>(FingerId);
 
@@ -1061,18 +1045,18 @@ void FLinuxApplication::ProcessDeferredMessage( SDL_Event Event )
 			else
 			{
 				UE_LOG(LogLinuxWindow, Warning, TEXT("Ignoring touch event SDL_FINGERDOWN (finger: %llu, x=%f, y=%f) that doesn't have a window associated with it"),
-					Event.tfinger.fingerId, Event.tfinger.x, Event.tfinger.y);
+					Event.tfinger.fingerID, Event.tfinger.x, Event.tfinger.y);
 			}
 		}
 		break;
-	case SDL_FINGERUP:
+	case SDL_EVENT_FINGER_UP:
 		{
-			UE_LOG(LogLinuxWindow, Verbose, TEXT("Finger %llu is up at (%f, %f)"), Event.tfinger.fingerId, Event.tfinger.x, Event.tfinger.y);
+			UE_LOG(LogLinuxWindow, Verbose, TEXT("Finger %llu is up at (%f, %f)"), Event.tfinger.fingerID, Event.tfinger.x, Event.tfinger.y);
 
 			// touch events can have no window associated with them, in that case ignore (with a warning)
 			if (LIKELY(!bWindowlessEvent))
 			{
-				uint64 FingerId = static_cast<uint64>(Event.tfinger.fingerId);
+				uint64 FingerId = static_cast<uint64>(Event.tfinger.fingerID);
 				FTouchContext* TouchContext = Touches.Find(FingerId);
 				if (UNLIKELY(TouchContext == nullptr))
 				{
@@ -1095,16 +1079,16 @@ void FLinuxApplication::ProcessDeferredMessage( SDL_Event Event )
 			else
 			{
 				UE_LOG(LogLinuxWindow, Warning, TEXT("Ignoring touch event SDL_FINGERUP (finger: %llu, x=%f, y=%f) that doesn't have a window associated with it"),
-					Event.tfinger.fingerId, Event.tfinger.x, Event.tfinger.y);
+					Event.tfinger.fingerID, Event.tfinger.x, Event.tfinger.y);
 			}
 		}
 		break;
-	case SDL_FINGERMOTION:
+	case SDL_EVENT_FINGER_MOTION:
 		{
 			// touch events can have no window associated with them, in that case ignore (with a warning)
 			if (LIKELY(!bWindowlessEvent))
 			{
-				uint64 FingerId = static_cast<uint64>(Event.tfinger.fingerId);
+				uint64 FingerId = static_cast<uint64>(Event.tfinger.fingerID);
 				FTouchContext* TouchContext = Touches.Find(FingerId);
 				if (UNLIKELY(TouchContext == nullptr))
 				{
@@ -1126,7 +1110,7 @@ void FLinuxApplication::ProcessDeferredMessage( SDL_Event Event )
 			else
 			{
 				UE_LOG(LogLinuxWindow, Warning, TEXT("Ignoring touch event SDL_FINGERMOTION (finger: %llu, x=%f, y=%f) that doesn't have a window associated with it"),
-					Event.tfinger.fingerId, Event.tfinger.x, Event.tfinger.y);
+					Event.tfinger.fingerID, Event.tfinger.x, Event.tfinger.y);
 			}
 		}
 		break;
@@ -1177,7 +1161,7 @@ int FLinuxApplication::GetFirstFreeTouchId()
 
 FVector2D FLinuxApplication::GetTouchEventLocation(SDL_HWindow NativeWindow, SDL_Event TouchEvent)
 {
-	checkf(TouchEvent.type == SDL_FINGERDOWN || TouchEvent.type == SDL_FINGERUP || TouchEvent.type == SDL_FINGERMOTION, TEXT("Wrong touch event."));
+	checkf(TouchEvent.type == SDL_EVENT_FINGER_DOWN || TouchEvent.type == SDL_EVENT_FINGER_UP || TouchEvent.type == SDL_EVENT_FINGER_MOTION, TEXT("Wrong touch event."));
 
 	FWindowProperties Props;
 	GetWindowPropertiesInEventLoop(NativeWindow, Props);
@@ -1248,150 +1232,46 @@ void FLinuxApplication::PollGameDeviceState( const float TimeDelta )
 	}
 }
 
-TCHAR FLinuxApplication::ConvertChar( SDL_Keysym Keysym )
-{
-	if (SDL_GetKeyFromScancode(Keysym.scancode) >= 128)
-	{
-		return 0;
-	}
-
-	TCHAR Char = (TCHAR)SDL_GetKeyFromScancode(Keysym.scancode);
-
-    if (Keysym.mod & (KMOD_LSHIFT | KMOD_RSHIFT))
-    {
-        // Convert to uppercase (FIXME: what about CAPS?)
-		if( SDL_GetKeyFromScancode(Keysym.scancode)  >= 97 && SDL_GetKeyFromScancode(Keysym.scancode)  <= 122)
-        {
-            return (TCHAR)(Keysym.sym - 32);
-        }
-		else if( SDL_GetKeyFromScancode(Keysym.scancode) >= 91 && SDL_GetKeyFromScancode(Keysym.scancode)  <= 93)
-        {
-			return (TCHAR)(SDL_GetKeyFromScancode(Keysym.scancode) + 32); // [ \ ] -> { | }
-        }
-        else
-        {
-			switch(SDL_GetKeyFromScancode(Keysym.scancode) )
-            {
-                case '`': // ` -> ~
-                    Char = TEXT('`');
-                    break;
-
-                case '-': // - -> _
-                    Char = TEXT('_');
-                    break;
-
-                case '=': // - -> _
-                    Char = TEXT('+');
-                    break;
-
-                case ',':
-                    Char = TEXT('<');
-                    break;
-
-                case '.':
-                    Char = TEXT('>');
-                    break;
-
-                case ';':
-                    Char = TEXT(':');
-                    break;
-
-                case '\'':
-                    Char = TEXT('\"');
-                    break;
-
-                case '/':
-                    Char = TEXT('?');
-                    break;
-
-                case '0':
-                    Char = TEXT(')');
-                    break;
-
-                case '9':
-                    Char = TEXT('(');
-                    break;
-
-                case '8':
-                    Char = TEXT('*');
-                    break;
-
-                case '7':
-                    Char = TEXT('&');
-                    break;
-
-                case '6':
-                    Char = TEXT('^');
-                    break;
-
-                case '5':
-                    Char = TEXT('%');
-                    break;
-
-                case '4':
-                    Char = TEXT('$');
-                    break;
-
-                case '3':
-                    Char = TEXT('#');
-                    break;
-
-                case '2':
-                    Char = TEXT('@');
-                    break;
-
-                case '1':
-                    Char = TEXT('!');
-                    break;
-
-                default:
-                    break;
-            }
-        }
-    }
-
-    return Char;
-}
-
 TSharedPtr< FLinuxWindow > FLinuxApplication::FindEventWindow(SDL_Event* Event, bool& bOutWindowlessEvent)
 {
 	uint32 WindowID = 0;
 	bOutWindowlessEvent = false;
 
-	switch (Event->type)
+	if(Event->type >= SDL_EVENT_WINDOW_FIRST && Event->type <= SDL_EVENT_WINDOW_LAST)
+	{
+		WindowID = Event->window.windowID;
+	}
+	else switch (Event->type)
 	{
-		case SDL_TEXTINPUT:
+		case SDL_EVENT_TEXT_INPUT:
 			WindowID = Event->text.windowID;
 			break;
-		case SDL_TEXTEDITING:
+		case SDL_EVENT_TEXT_EDITING:
 			WindowID = Event->edit.windowID;
 			break;
-		case SDL_KEYDOWN:
-		case SDL_KEYUP:
+		case SDL_EVENT_KEY_DOWN:
+		case SDL_EVENT_KEY_UP:
 			WindowID = Event->key.windowID;
 			break;
-		case SDL_MOUSEMOTION:
+		case SDL_EVENT_MOUSE_MOTION:
 			WindowID = Event->motion.windowID;
 			break;
-		case SDL_MOUSEBUTTONDOWN:
-		case SDL_MOUSEBUTTONUP:
+		case SDL_EVENT_MOUSE_BUTTON_DOWN:
+		case SDL_EVENT_MOUSE_BUTTON_UP:
 			WindowID = Event->button.windowID;
 			break;
-		case SDL_MOUSEWHEEL:
+		case SDL_EVENT_MOUSE_WHEEL:
 			WindowID = Event->wheel.windowID;
 			break;
-		case SDL_WINDOWEVENT:
-			WindowID = Event->window.windowID;
-			break;
-		case SDL_DROPBEGIN:
-		case SDL_DROPFILE:
-		case SDL_DROPTEXT:
-		case SDL_DROPCOMPLETE:
+		case SDL_EVENT_DROP_BEGIN:
+		case SDL_EVENT_DROP_FILE:
+		case SDL_EVENT_DROP_TEXT:
+		case SDL_EVENT_DROP_COMPLETE:
 			WindowID = Event->drop.windowID;
 			break;
-		case SDL_FINGERDOWN:
-		case SDL_FINGERUP:
-		case SDL_FINGERMOTION:
+		case SDL_EVENT_FINGER_DOWN:
+		case SDL_EVENT_FINGER_UP:
+		case SDL_EVENT_FINGER_MOTION:
 			// SDL touch events are windowless, but Slate needs to associate touch down with a particular window.
 			// Assume that the current focus window is the one relevant for the touch and if there's none, assume the event windowless
 			if (CurrentFocusWindow.IsValid())
@@ -1439,13 +1319,13 @@ FModifierKeysState FLinuxApplication::GetModifierKeys() const
 {
 	SDL_Keymod modifiers = SDL_GetModState();
 
-	const bool bIsLeftShiftDown		= (modifiers & KMOD_LSHIFT) != 0;
-	const bool bIsRightShiftDown	= (modifiers & KMOD_RSHIFT) != 0;
-	const bool bIsLeftControlDown	= (modifiers & KMOD_LCTRL) != 0;
-	const bool bIsRightControlDown	= (modifiers & KMOD_RCTRL) != 0;
-	const bool bIsLeftAltDown		= (modifiers & KMOD_LALT) != 0;
-	const bool bIsRightAltDown		= (modifiers & KMOD_RALT) != 0;
-	const bool bAreCapsLocked		= (modifiers & KMOD_CAPS) != 0;
+	const bool bIsLeftShiftDown	    = (modifiers & SDL_KMOD_LSHIFT) != 0;
+	const bool bIsRightShiftDown	= (modifiers & SDL_KMOD_RSHIFT) != 0;
+	const bool bIsLeftControlDown	= (modifiers & SDL_KMOD_LCTRL) != 0;
+	const bool bIsRightControlDown	= (modifiers & SDL_KMOD_RCTRL) != 0;
+	const bool bIsLeftAltDown	    = (modifiers & SDL_KMOD_LALT) != 0;
+	const bool bIsRightAltDown	    = (modifiers & SDL_KMOD_RALT) != 0;
+	const bool bAreCapsLocked	    = (modifiers & SDL_KMOD_CAPS) != 0;
 
 	return FModifierKeysState( bIsLeftShiftDown, bIsRightShiftDown, bIsLeftControlDown, bIsRightControlDown, bIsLeftAltDown, bIsRightAltDown, false, false, bAreCapsLocked );
 }
@@ -1465,7 +1345,7 @@ bool FLinuxApplication::IsGamepadAttached() const
 {
 	for (const TPair<SDL_JoystickID, SDLControllerState>& ControllerState : ControllerStates)
 	{
-		if (SDL_GameControllerGetAttached(ControllerState.Value.Controller))
+		if (SDL_GamepadConnected(ControllerState.Value.Controller))
 		{
 			return true;
 		}
@@ -1489,23 +1369,32 @@ void FLinuxApplication::UpdateMouseCaptureWindow(SDL_HWindow TargetWindow)
 		}
 		if (bShouldGrab && MouseCaptureWindow)
 		{
-			SDL_CaptureMouse(SDL_TRUE);
+			SDL_CaptureMouse(true);
 		}
 	}
 	else
 	{
-		SDL_CaptureMouse(SDL_FALSE);
+		SDL_CaptureMouse(false);
 		MouseCaptureWindow = nullptr;
 	}
 }
 
 void FLinuxApplication::SetHighPrecisionMouseMode( const bool Enable, const TSharedPtr< FGenericWindow >& InWindow )
 {
+	static const FString EnvValue = FPlatformMisc::GetEnvironmentVariable(TEXT("UE_NORELATIVEMOUSEMODE")); 
+	static const bool bNoRelativeMouseMode = (FParse::Param(FCommandLine::Get(), TEXT("norelativemousemode")) || !EnvValue.IsEmpty());
+
 	MessageHandler->OnCursorSet();
 	bUsingHighPrecisionMouseInput = Enable;
-	if (!FParse::Param(FCommandLine::Get(), TEXT("norelativemousemode")))
+	
+	if (!bNoRelativeMouseMode)
 	{
-		SDL_SetRelativeMouseMode(Enable ? SDL_TRUE : SDL_FALSE);
+		if(!InWindow.Get())
+		{
+			return;
+		}
+			
+		SDL_SetWindowRelativeMouseMode(((FLinuxWindow*)InWindow.Get())->GetHWnd(), Enable);
 	}
 }
 
@@ -1518,15 +1407,17 @@ void FLinuxApplication::RefreshDisplayCache()
 	{
 		CachedDisplays.Empty();
 
-		int NumDisplays = SDL_GetNumVideoDisplays();
+		int NumDisplays;
+		SDL_DisplayID* Displays = SDL_GetDisplays(&NumDisplays);
 
 		for (int DisplayIdx = 0; DisplayIdx < NumDisplays; ++DisplayIdx)
 		{
 			SDL_Rect UsableBounds;
-			SDL_GetDisplayUsableBounds(DisplayIdx, &UsableBounds);
+			SDL_GetDisplayUsableBounds(Displays[DisplayIdx], &UsableBounds);
 
 			CachedDisplays.Add(UsableBounds);
 		}
+		SDL_free(Displays);
 
 		LastTimeCachedDisplays = CurrentTime;
 	}
@@ -1578,20 +1469,22 @@ void FDisplayMetrics::RebuildDisplayMetrics(FDisplayMetrics& OutDisplayMetrics)
 	}
 	else
 	{
-		int NumDisplays = 0;
-
 		if (LIKELY(FApp::CanEverRender()))
 		{
-			if (FLinuxPlatformApplicationMisc::InitSDL()) //	will not initialize more than once
-			{
-				NumDisplays = SDL_GetNumVideoDisplays();
-			}
-			else
+			if (!FLinuxPlatformApplicationMisc::InitSDL()) //	will not initialize more than once
 			{
 				UE_LOG(LogInit, Warning, TEXT("FDisplayMetrics::GetDisplayMetrics: InitSDL() failed, cannot get display metrics"));
 			}
 		}
 
+		int NumDisplays = 0;
+		SDL_DisplayID* Displays = nullptr;
+
+		if(SDL_WasInit(SDL_INIT_VIDEO))
+		{
+			Displays = SDL_GetDisplays(&NumDisplays);
+		}
+
 		OutDisplayMetrics.MonitorInfo.Empty();
 
 		if (NumDisplays <= 0)
@@ -1604,6 +1497,11 @@ void FDisplayMetrics::RebuildDisplayMetrics(FDisplayMetrics& OutDisplayMetrics)
 				OutDisplayMetrics.PrimaryDisplayWidth = 0;
 				OutDisplayMetrics.PrimaryDisplayHeight = 0;
 
+				if(Displays)
+				{
+					SDL_free(Displays);
+					Displays = nullptr;
+				}
 				return;
 			}
 			else
@@ -1637,10 +1535,10 @@ void FDisplayMetrics::RebuildDisplayMetrics(FDisplayMetrics& OutDisplayMetrics)
 		{
 			SDL_Rect DisplayBounds, UsableBounds;
 			FMonitorInfo Display;
-			SDL_GetDisplayBounds(DisplayIdx, &DisplayBounds);
-			SDL_GetDisplayUsableBounds(DisplayIdx, &UsableBounds);
+			SDL_GetDisplayBounds(Displays[DisplayIdx], &DisplayBounds);
+			SDL_GetDisplayUsableBounds(Displays[DisplayIdx], &UsableBounds);
 
-			Display.Name = UTF8_TO_TCHAR(SDL_GetDisplayName(DisplayIdx));
+			Display.Name = UTF8_TO_TCHAR(SDL_GetDisplayName(Displays[DisplayIdx]));
 			Display.ID = FString::Printf(TEXT("display%d"), DisplayIdx);
 			Display.NativeWidth = DisplayBounds.w;
 			Display.NativeHeight = DisplayBounds.h;
@@ -1651,11 +1549,8 @@ void FDisplayMetrics::RebuildDisplayMetrics(FDisplayMetrics& OutDisplayMetrics)
 
 			if (FPlatformApplicationMisc::IsHighDPIAwarenessEnabled())
 			{
-				float HorzDPI = 0.0f, VertDPI = 0.0f;
-				if (SDL_GetDisplayDPI(DisplayIdx, nullptr, &HorzDPI, &VertDPI) == 0)
-				{
-					Display.DPI = FMath::FloorToInt((HorzDPI + VertDPI) / 2.0f);
-				}
+				// This is really only an approximation because SDL_GetDisplayDPI() was removed for being unreliable across platforms
+				Display.DPI = (int32)(SDL_GetDisplayContentScale(Displays[DisplayIdx]) * 96.0f);
 			}
 
 			OutDisplayMetrics.MonitorInfo.Add(Display);
@@ -1678,6 +1573,8 @@ void FDisplayMetrics::RebuildDisplayMetrics(FDisplayMetrics& OutDisplayMetrics)
 				OutDisplayMetrics.VirtualDisplayRect.Bottom = FMath::Max(OutDisplayMetrics.VirtualDisplayRect.Bottom, DisplayBounds.y + DisplayBounds.h);
 			}
 		}
+
+		SDL_free(Displays);
 	}
 
 	// Apply the debug safe zones
@@ -2085,7 +1982,7 @@ void FLinuxApplication::SDLControllerState::UpdateHapticEffect()
 	{
 		if (EffectId >= 0 && bEffectRunning)
 		{
-			SDL_HapticStopEffect(Haptic, EffectId);
+			SDL_StopHapticEffect(Haptic, EffectId);
 			bEffectRunning = false;
 		}
 		return;
@@ -2094,14 +1991,14 @@ void FLinuxApplication::SDLControllerState::UpdateHapticEffect()
 	SDL_HapticEffect Effect;
 	FMemory::Memzero(Effect);
 
-	if (SDL_HapticQuery(Haptic) & SDL_HAPTIC_LEFTRIGHT)
+	if (SDL_GetHapticFeatures(Haptic) & SDL_HAPTIC_LEFTRIGHT)
 	{
 		Effect.type = SDL_HAPTIC_LEFTRIGHT;
 		Effect.leftright.length = 1000;
 		Effect.leftright.large_magnitude = (uint16)(32767.0f * LargeValue);
 		Effect.leftright.small_magnitude = (uint16)(32767.0f * SmallValue);
 	}
-	else if (SDL_HapticQuery(Haptic) & SDL_HAPTIC_SINE)
+	else if (SDL_GetHapticFeatures(Haptic) & SDL_HAPTIC_SINE)
 	{
 		Effect.type = SDL_HAPTIC_SINE;
 		Effect.periodic.length = 1000;
@@ -2111,52 +2008,51 @@ void FLinuxApplication::SDLControllerState::UpdateHapticEffect()
 	else
 	{
 		UE_LOG(LogLinux, Warning, TEXT("No available haptic effects"));
-		SDL_HapticClose(Haptic);
+		SDL_CloseHaptic(Haptic);
 		Haptic = nullptr;
 		return;
 	}
 
 	if (EffectId < 0)
 	{
-		EffectId = SDL_HapticNewEffect(Haptic, &Effect);
+		EffectId = SDL_CreateHapticEffect(Haptic, &Effect);
 		if (EffectId < 0)
 		{
 			UE_LOG(LogLinux, Warning, TEXT("Failed to create haptic effect: %s"), UTF8_TO_TCHAR(SDL_GetError()) );
-			SDL_HapticClose(Haptic);
+			SDL_CloseHaptic(Haptic);
 			Haptic = nullptr;
 			return;
 		}
 	}
 
-	if (SDL_HapticUpdateEffect(Haptic, EffectId, &Effect) < 0)
+	if (SDL_UpdateHapticEffect(Haptic, EffectId, &Effect) < 0)
 	{
-		SDL_HapticDestroyEffect(Haptic, EffectId);
-		EffectId = SDL_HapticNewEffect(Haptic, &Effect);
+		SDL_DestroyHapticEffect(Haptic, EffectId);
+		EffectId = SDL_CreateHapticEffect(Haptic, &Effect);
 		if (EffectId < 0)
 		{
 			UE_LOG(LogLinux, Warning, TEXT("Failed to update and recreate haptic effect: %s"), UTF8_TO_TCHAR(SDL_GetError()) );
-			SDL_HapticClose(Haptic);
+			SDL_CloseHaptic(Haptic);
 			Haptic = nullptr;
 			return;
 		}
 	}
 
-	SDL_HapticRunEffect(Haptic, EffectId, 1);
+	SDL_RunHapticEffect(Haptic, EffectId, 1);
 	bEffectRunning = true;
 }
 
-void FLinuxApplication::AddGameController(int Index)
+void FLinuxApplication::AddGameController(SDL_JoystickID Id)
 {
-	SDL_JoystickID Id = SDL_JoystickGetDeviceInstanceID(Index);
 	if(ControllerStates.Contains(Id))
 	{
 		return;
 	}
 
-	auto Controller = SDL_GameControllerOpen(Index);
+	auto Controller = SDL_OpenGamepad(Id);
 	if (Controller == nullptr)
 	{
-		UE_LOG(LogLinux, Warning, TEXT("Could not open gamecontroller %i: %s"), Index, UTF8_TO_TCHAR(SDL_GetError()) );
+		UE_LOG(LogLinux, Warning, TEXT("Could not open gamecontroller %i: %s"), Id, UTF8_TO_TCHAR(SDL_GetError()) );
 		return;
 	}
 
@@ -2171,7 +2067,7 @@ void FLinuxApplication::AddGameController(int Index)
 
 	int32 FirstUnusedIndex = FMath::CountTrailingZeros(~UsedBits);
 
-	UE_LOG(LogLinux, Verbose, TEXT("Adding controller %i '%s'"), FirstUnusedIndex, UTF8_TO_TCHAR(SDL_GameControllerName(Controller)));
+	UE_LOG(LogLinux, Verbose, TEXT("Adding controller %i '%s'"), FirstUnusedIndex, UTF8_TO_TCHAR(SDL_GetGamepadName(Controller)));
 	auto& ControllerState = ControllerStates.Add(Id);
 	ControllerState.Controller = Controller;
 
@@ -2179,14 +2075,14 @@ void FLinuxApplication::AddGameController(int Index)
 	DeviceMapper.RemapControllerIdToPlatformUserAndDevice(FirstUnusedIndex, UserId, ControllerState.DeviceId);
 
 	// Check for haptic support.
-	ControllerState.Haptic = SDL_HapticOpenFromJoystick(SDL_GameControllerGetJoystick(Controller));
+	ControllerState.Haptic = SDL_OpenHapticFromJoystick(SDL_GetGamepadJoystick(Controller));
 
 	if (ControllerState.Haptic != nullptr)
 	{
-		if ((SDL_HapticQuery(ControllerState.Haptic) & (SDL_HAPTIC_SINE | SDL_HAPTIC_LEFTRIGHT)) == 0)
+		if ((SDL_GetHapticFeatures(ControllerState.Haptic) & (SDL_HAPTIC_SINE | SDL_HAPTIC_LEFTRIGHT)) == 0)
 		{
 			UE_LOG(LogLinux, Warning, TEXT("No supported haptic effects for controller %i"), ControllerState.DeviceId.GetId());
-			SDL_HapticClose(ControllerState.Haptic);
+			SDL_CloseHaptic(ControllerState.Haptic);
 			ControllerState.Haptic = nullptr;
 		}
 	}
@@ -2202,14 +2098,14 @@ void FLinuxApplication::RemoveGameController(SDL_JoystickID Id)
 	}
 
 	SDLControllerState& ControllerState = ControllerStates[Id];
-	UE_LOG(LogLinux, Verbose, TEXT("Removing controller %i '%s'"), ControllerState.DeviceId.GetId(), UTF8_TO_TCHAR(SDL_GameControllerName(ControllerState.Controller)));
+	UE_LOG(LogLinux, Verbose, TEXT("Removing controller %i '%s'"), ControllerState.DeviceId.GetId(), UTF8_TO_TCHAR(SDL_GetGamepadName(ControllerState.Controller)));
 
 	if (ControllerState.Haptic != nullptr)
 	{
-		SDL_HapticClose(ControllerState.Haptic);
+		SDL_CloseHaptic(ControllerState.Haptic);
 	}
 
-	SDL_GameControllerClose(ControllerState.Controller);
+	SDL_CloseGamepad(ControllerState.Controller);
 	ControllerStates.Remove(Id);
 
 	FPlatformUserId PlatformUserId = PLATFORMUSERID_NONE;
diff --git a/Engine/Source/Runtime/ApplicationCore/Private/Linux/LinuxCursor.cpp b/Engine/Source/Runtime/ApplicationCore/Private/Linux/LinuxCursor.cpp
--- a/Engine/Source/Runtime/ApplicationCore/Private/Linux/LinuxCursor.cpp
+++ b/Engine/Source/Runtime/ApplicationCore/Private/Linux/LinuxCursor.cpp
@@ -9,8 +9,8 @@
 
 FLinuxCursor::FLinuxCursor()
 	: 	bHidden(false)
-	,	CachedGlobalXPosition(0)
-	,	CachedGlobalYPosition(0)
+	,	CachedGlobalXPosition(0.0f)
+	,	CachedGlobalYPosition(0.0f)
 	,	bPositionCacheIsValid(false)
 {
 	if (!FApp::CanEverRender())
@@ -28,7 +28,7 @@ FLinuxCursor::FLinuxCursor()
 	}
 
 #if DO_CHECK
-	uint32 InitializedSubsystems = SDL_WasInit(SDL_INIT_EVERYTHING);
+	uint32 InitializedSubsystems = SDL_WasInit(SDL_INIT_VIDEO);
 	check(InitializedSubsystems & SDL_INIT_VIDEO);
 #endif // DO_CHECK
 
@@ -47,31 +47,31 @@ FLinuxCursor::FLinuxCursor()
 			break;
 
 		case EMouseCursor::Default:
-			CursorHandle = SDL_CreateSystemCursor( SDL_SYSTEM_CURSOR_ARROW );
+			CursorHandle = SDL_CreateSystemCursor( SDL_SYSTEM_CURSOR_DEFAULT );
 			break;
 
 		case EMouseCursor::TextEditBeam:
-			CursorHandle = SDL_CreateSystemCursor( SDL_SYSTEM_CURSOR_IBEAM );
+			CursorHandle = SDL_CreateSystemCursor( SDL_SYSTEM_CURSOR_TEXT );
 			break;
 
 		case EMouseCursor::ResizeLeftRight:
-			CursorHandle = SDL_CreateSystemCursor( SDL_SYSTEM_CURSOR_SIZEWE );
+			CursorHandle = SDL_CreateSystemCursor( SDL_SYSTEM_CURSOR_EW_RESIZE );
 			break;
 
 		case EMouseCursor::ResizeUpDown:
-			CursorHandle = SDL_CreateSystemCursor( SDL_SYSTEM_CURSOR_SIZENS );
+			CursorHandle = SDL_CreateSystemCursor( SDL_SYSTEM_CURSOR_NS_RESIZE );
 			break;
 
 		case EMouseCursor::ResizeSouthEast:
-			CursorHandle = SDL_CreateSystemCursor( SDL_SYSTEM_CURSOR_SIZENWSE );
+			CursorHandle = SDL_CreateSystemCursor( SDL_SYSTEM_CURSOR_NWSE_RESIZE );
 			break;
 
 		case EMouseCursor::ResizeSouthWest:
-			CursorHandle = SDL_CreateSystemCursor( SDL_SYSTEM_CURSOR_SIZENESW );
+			CursorHandle = SDL_CreateSystemCursor( SDL_SYSTEM_CURSOR_NESW_RESIZE );
 			break;
 
 		case EMouseCursor::CardinalCross:
-			CursorHandle = SDL_CreateSystemCursor( SDL_SYSTEM_CURSOR_SIZEALL );
+			CursorHandle = SDL_CreateSystemCursor( SDL_SYSTEM_CURSOR_MOVE );
 			break;
 
 		case EMouseCursor::Crosshairs:
@@ -79,7 +79,7 @@ FLinuxCursor::FLinuxCursor()
 			break;
 
 		case EMouseCursor::Hand:
-			CursorHandle = SDL_CreateSystemCursor( SDL_SYSTEM_CURSOR_HAND );
+			CursorHandle = SDL_CreateSystemCursor( SDL_SYSTEM_CURSOR_POINTER );
 			break;
 
 		case EMouseCursor::GrabHand:
@@ -87,7 +87,7 @@ FLinuxCursor::FLinuxCursor()
 				FString path(FPlatformProcess::BaseDir() / FPaths::EngineContentDir() / FString("Editor/Slate/Cursor/grabhand.bmp") );
 				SDL_Surface *surf = SDL_LoadBMP(TCHAR_TO_ANSI(*path));
 				CursorHandle = SDL_CreateColorCursor(surf, 15, 15);
-				SDL_FreeSurface(surf);
+				SDL_DestroySurface(surf);
 			}
 			break;
 
@@ -96,12 +96,12 @@ FLinuxCursor::FLinuxCursor()
 				FString path(FPlatformProcess::BaseDir() / FPaths::EngineContentDir() / FString("Editor/Slate/Cursor/grabhand_closed.bmp") );
 				SDL_Surface *surf = SDL_LoadBMP(TCHAR_TO_ANSI(*path));
 				CursorHandle = SDL_CreateColorCursor(surf, 8, 8);
-				SDL_FreeSurface(surf);
+				SDL_DestroySurface(surf);
 			}
 			break;
 
 		case EMouseCursor::SlashedCircle:
-			CursorHandle = SDL_CreateSystemCursor( SDL_SYSTEM_CURSOR_NO );
+			CursorHandle = SDL_CreateSystemCursor( SDL_SYSTEM_CURSOR_NOT_ALLOWED );
 			break;
 
 		case EMouseCursor::EyeDropper:
@@ -162,7 +162,7 @@ FLinuxCursor::~FLinuxCursor()
 			case EMouseCursor::GrabHand:
 			case EMouseCursor::GrabHandClosed:
 			case EMouseCursor::EyeDropper:
-				SDL_FreeCursor(CursorHandles[CurCursorIndex]);
+				SDL_DestroyCursor(CursorHandles[CurCursorIndex]);
 				break;
 
 			default:
@@ -180,12 +180,12 @@ void* FLinuxCursor::CreateCursorFromRGBABuffer(const FColor* Pixels, int32 Width
 	uint32 Bmask = 0x00ff0000;
 	uint32 Amask = 0xff000000;
 
-	SDL_Surface* Surface = SDL_CreateRGBSurfaceFrom((void*) Pixels, Width, Height, 32, Width * 4, Rmask, Gmask, Bmask, Amask);
+	SDL_Surface* Surface = SDL_CreateSurfaceFrom(Width, Height, SDL_GetPixelFormatForMasks(32, Rmask, Gmask, Bmask, Amask),(void*) Pixels, Width*4);
 	if (Surface)
 	{
 		SDL_HCursor CursorHandle = SDL_CreateColorCursor(Surface, FMath::RoundToInt(InHotSpot.X * Width), FMath::RoundToInt(InHotSpot.Y * Height));
 
-		SDL_FreeSurface(Surface);
+		SDL_DestroySurface(Surface);
 
 		return CursorHandle;
 	}
@@ -206,8 +206,8 @@ FVector2D FLinuxCursor::GetPosition() const
 		bPositionCacheIsValid = true;
 	}
 
-	int CursorX = CachedGlobalXPosition;
-	int CursorY = CachedGlobalYPosition;
+	float CursorX = CachedGlobalXPosition;
+	float CursorY = CachedGlobalYPosition;
 
 	return FVector2D( CursorX, CursorY );
 }
@@ -217,7 +217,7 @@ void FLinuxCursor::InvalidateCaches()
 	bPositionCacheIsValid = false;
 }
 
-void FLinuxCursor::SetCachedPosition( const int32 X, const int32 Y )
+void FLinuxCursor::SetCachedPosition( const float X, const float Y )
 {
 	CachedGlobalXPosition = X;
 	CachedGlobalYPosition = Y;
@@ -246,14 +246,14 @@ void FLinuxCursor::SetPosition( const int32 X, const int32 Y )
 			SDL_GetWindowPosition(WndFocus, &WndX, &WndY);
 		}
 
-		SDL_WarpMouseInWindow(WndFocus, X - WndX, Y - WndY);
+		SDL_WarpMouseInWindow(WndFocus, (float)(X - WndX), (float)(Y - WndY));
 	}
 	else
 	{
-		SDL_WarpMouseGlobal(X, Y);
+		SDL_WarpMouseGlobal((float)X, (float)Y);
 	}
 
-	SetCachedPosition(X, Y);
+	SetCachedPosition((float)X, (float)Y);
 }
 
 void FLinuxCursor::SetType(const EMouseCursor::Type InNewCursor)
@@ -269,13 +269,13 @@ void FLinuxCursor::SetType(const EMouseCursor::Type InNewCursor)
 		{
 			bHidden = true;
 		}
-		SDL_ShowCursor(SDL_DISABLE);
+		SDL_HideCursor();
 		SDL_SetCursor(CursorHandles[0]);
 	}
 	else
 	{
 		bHidden = false;
-		SDL_ShowCursor(SDL_ENABLE);
+		SDL_ShowCursor();
 		SDL_SetCursor(CurrentCursor);
 	}
 }
@@ -292,13 +292,13 @@ void FLinuxCursor::Show( bool bShow )
 	{
 		// Show mouse cursor.
 		bHidden = false;
-		SDL_ShowCursor(SDL_ENABLE);
+		SDL_ShowCursor();
 	}
 	else
 	{
 		// Disable the cursor.
 		bHidden = true;
-		SDL_ShowCursor(SDL_DISABLE);
+		SDL_HideCursor();
 	}
 }
 
diff --git a/Engine/Source/Runtime/ApplicationCore/Private/Linux/LinuxPlatformApplicationMisc.cpp b/Engine/Source/Runtime/ApplicationCore/Private/Linux/LinuxPlatformApplicationMisc.cpp
--- a/Engine/Source/Runtime/ApplicationCore/Private/Linux/LinuxPlatformApplicationMisc.cpp
+++ b/Engine/Source/Runtime/ApplicationCore/Private/Linux/LinuxPlatformApplicationMisc.cpp
@@ -12,14 +12,15 @@
 #include "Linux/LinuxApplication.h"
 
 THIRD_PARTY_INCLUDES_START
-	#include <SDL.h>
+	#include <SDL3/SDL.h>
+	#include <SDL3/SDL_vulkan.h>
 THIRD_PARTY_INCLUDES_END
 
 bool GInitializedSDL = false;
 
 namespace
 {
-	uint32 GWindowStyleSDL = SDL_WINDOW_VULKAN;
+	uint64 GWindowStyleSDL = SDL_WINDOW_VULKAN;
 
 	FString GetHeadlessMessageBoxMessage(EAppMsgType::Type MsgType, const TCHAR* Text, const TCHAR* Caption, EAppReturnType::Type& Answer)
 	{
@@ -79,7 +80,7 @@ EAppReturnType::Type MessageBoxExtImpl(EAppMsgType::Type MsgType, const TCHAR* T
 
 
 #if DO_CHECK
-	uint32 InitializedSubsystems = SDL_WasInit(SDL_INIT_EVERYTHING);
+	uint32 InitializedSubsystems = SDL_WasInit(SDL_INIT_VIDEO);
 	check(InitializedSubsystems & SDL_INIT_VIDEO);
 #endif // DO_CHECK
 
@@ -92,7 +93,7 @@ EAppReturnType::Type MessageBoxExtImpl(EAppMsgType::Type MsgType, const TCHAR* T
 		Buttons = new SDL_MessageBoxButtonData[NumberOfButtons];
 		Buttons[0].flags = SDL_MESSAGEBOX_BUTTON_RETURNKEY_DEFAULT;
 		Buttons[0].text = "Ok";
-		Buttons[0].buttonid = EAppReturnType::Ok;
+		Buttons[0].buttonID = EAppReturnType::Ok;
 		break;
 
 	case EAppMsgType::YesNo:
@@ -100,10 +101,10 @@ EAppReturnType::Type MessageBoxExtImpl(EAppMsgType::Type MsgType, const TCHAR* T
 		Buttons = new SDL_MessageBoxButtonData[NumberOfButtons];
 		Buttons[0].flags = SDL_MESSAGEBOX_BUTTON_ESCAPEKEY_DEFAULT;
 		Buttons[0].text = "Yes";
-		Buttons[0].buttonid = EAppReturnType::Yes;
+		Buttons[0].buttonID = EAppReturnType::Yes;
 		Buttons[1].flags = SDL_MESSAGEBOX_BUTTON_ESCAPEKEY_DEFAULT;
 		Buttons[1].text = "No";
-		Buttons[1].buttonid = EAppReturnType::No;
+		Buttons[1].buttonID = EAppReturnType::No;
 		break;
 
 	case EAppMsgType::OkCancel:
@@ -111,10 +112,10 @@ EAppReturnType::Type MessageBoxExtImpl(EAppMsgType::Type MsgType, const TCHAR* T
 		Buttons = new SDL_MessageBoxButtonData[NumberOfButtons];
 		Buttons[0].flags = SDL_MESSAGEBOX_BUTTON_ESCAPEKEY_DEFAULT;
 		Buttons[0].text = "Ok";
-		Buttons[0].buttonid = EAppReturnType::Ok;
+		Buttons[0].buttonID = EAppReturnType::Ok;
 		Buttons[1].flags = SDL_MESSAGEBOX_BUTTON_ESCAPEKEY_DEFAULT;
 		Buttons[1].text = "Cancel";
-		Buttons[1].buttonid = EAppReturnType::Cancel;
+		Buttons[1].buttonID = EAppReturnType::Cancel;
 		break;
 
 	case EAppMsgType::YesNoCancel:
@@ -122,13 +123,13 @@ EAppReturnType::Type MessageBoxExtImpl(EAppMsgType::Type MsgType, const TCHAR* T
 		Buttons = new SDL_MessageBoxButtonData[NumberOfButtons];
 		Buttons[0].flags = SDL_MESSAGEBOX_BUTTON_ESCAPEKEY_DEFAULT;
 		Buttons[0].text = "Yes";
-		Buttons[0].buttonid = EAppReturnType::Yes;
+		Buttons[0].buttonID = EAppReturnType::Yes;
 		Buttons[1].flags = SDL_MESSAGEBOX_BUTTON_ESCAPEKEY_DEFAULT;
 		Buttons[1].text = "No";
-		Buttons[1].buttonid = EAppReturnType::No;
+		Buttons[1].buttonID = EAppReturnType::No;
 		Buttons[2].flags = SDL_MESSAGEBOX_BUTTON_ESCAPEKEY_DEFAULT;
 		Buttons[2].text = "Cancel";
-		Buttons[2].buttonid = EAppReturnType::Cancel;
+		Buttons[2].buttonID = EAppReturnType::Cancel;
 		break;
 
 	case EAppMsgType::CancelRetryContinue:
@@ -136,13 +137,13 @@ EAppReturnType::Type MessageBoxExtImpl(EAppMsgType::Type MsgType, const TCHAR* T
 		Buttons = new SDL_MessageBoxButtonData[NumberOfButtons];
 		Buttons[0].flags = SDL_MESSAGEBOX_BUTTON_ESCAPEKEY_DEFAULT;
 		Buttons[0].text = "Continue";
-		Buttons[0].buttonid = EAppReturnType::Continue;
+		Buttons[0].buttonID = EAppReturnType::Continue;
 		Buttons[1].flags = SDL_MESSAGEBOX_BUTTON_ESCAPEKEY_DEFAULT;
 		Buttons[1].text = "Retry";
-		Buttons[1].buttonid = EAppReturnType::Retry;
+		Buttons[1].buttonID = EAppReturnType::Retry;
 		Buttons[2].flags = SDL_MESSAGEBOX_BUTTON_ESCAPEKEY_DEFAULT;
 		Buttons[2].text = "Cancel";
-		Buttons[2].buttonid = EAppReturnType::Cancel;
+		Buttons[2].buttonID = EAppReturnType::Cancel;
 		break;
 
 	case EAppMsgType::YesNoYesAllNoAll:
@@ -150,16 +151,16 @@ EAppReturnType::Type MessageBoxExtImpl(EAppMsgType::Type MsgType, const TCHAR* T
 		Buttons = new SDL_MessageBoxButtonData[NumberOfButtons];
 		Buttons[0].flags = SDL_MESSAGEBOX_BUTTON_ESCAPEKEY_DEFAULT;
 		Buttons[0].text = "Yes";
-		Buttons[0].buttonid = EAppReturnType::Yes;
+		Buttons[0].buttonID = EAppReturnType::Yes;
 		Buttons[1].flags = SDL_MESSAGEBOX_BUTTON_ESCAPEKEY_DEFAULT;
 		Buttons[1].text = "No";
-		Buttons[1].buttonid = EAppReturnType::No;
+		Buttons[1].buttonID = EAppReturnType::No;
 		Buttons[2].flags = SDL_MESSAGEBOX_BUTTON_ESCAPEKEY_DEFAULT;
 		Buttons[2].text = "Yes to all";
-		Buttons[2].buttonid = EAppReturnType::YesAll;
+		Buttons[2].buttonID = EAppReturnType::YesAll;
 		Buttons[3].flags = SDL_MESSAGEBOX_BUTTON_ESCAPEKEY_DEFAULT;
 		Buttons[3].text = "No to all";
-		Buttons[3].buttonid = EAppReturnType::NoAll;
+		Buttons[3].buttonID = EAppReturnType::NoAll;
 		break;
 
 	case EAppMsgType::YesNoYesAllNoAllCancel:
@@ -167,19 +168,19 @@ EAppReturnType::Type MessageBoxExtImpl(EAppMsgType::Type MsgType, const TCHAR* T
 		Buttons = new SDL_MessageBoxButtonData[NumberOfButtons];
 		Buttons[0].flags = SDL_MESSAGEBOX_BUTTON_ESCAPEKEY_DEFAULT;
 		Buttons[0].text = "Yes";
-		Buttons[0].buttonid = EAppReturnType::Yes;
+		Buttons[0].buttonID = EAppReturnType::Yes;
 		Buttons[1].flags = SDL_MESSAGEBOX_BUTTON_ESCAPEKEY_DEFAULT;
 		Buttons[1].text = "No";
-		Buttons[1].buttonid = EAppReturnType::No;
+		Buttons[1].buttonID = EAppReturnType::No;
 		Buttons[2].flags = SDL_MESSAGEBOX_BUTTON_ESCAPEKEY_DEFAULT;
 		Buttons[2].text = "Yes to all";
-		Buttons[2].buttonid = EAppReturnType::YesAll;
+		Buttons[2].buttonID = EAppReturnType::YesAll;
 		Buttons[3].flags = SDL_MESSAGEBOX_BUTTON_ESCAPEKEY_DEFAULT;
 		Buttons[3].text = "No to all";
-		Buttons[3].buttonid = EAppReturnType::NoAll;
+		Buttons[3].buttonID = EAppReturnType::NoAll;
 		Buttons[4].flags = SDL_MESSAGEBOX_BUTTON_ESCAPEKEY_DEFAULT;
 		Buttons[4].text = "Cancel";
-		Buttons[4].buttonid = EAppReturnType::Cancel;
+		Buttons[4].buttonID = EAppReturnType::Cancel;
 		break;
 
 	case EAppMsgType::YesNoYesAll:
@@ -187,13 +188,13 @@ EAppReturnType::Type MessageBoxExtImpl(EAppMsgType::Type MsgType, const TCHAR* T
 		Buttons = new SDL_MessageBoxButtonData[NumberOfButtons];
 		Buttons[0].flags = SDL_MESSAGEBOX_BUTTON_ESCAPEKEY_DEFAULT;
 		Buttons[0].text = "Yes";
-		Buttons[0].buttonid = EAppReturnType::Yes;
+		Buttons[0].buttonID = EAppReturnType::Yes;
 		Buttons[1].flags = SDL_MESSAGEBOX_BUTTON_ESCAPEKEY_DEFAULT;
 		Buttons[1].text = "No";
-		Buttons[1].buttonid = EAppReturnType::No;
+		Buttons[1].buttonID = EAppReturnType::No;
 		Buttons[2].flags = SDL_MESSAGEBOX_BUTTON_ESCAPEKEY_DEFAULT;
 		Buttons[2].text = "Yes to all";
-		Buttons[2].buttonid = EAppReturnType::YesAll;
+		Buttons[2].buttonID = EAppReturnType::YesAll;
 		break;
 	}
 
@@ -236,8 +237,8 @@ void UngrabAllInputImpl()
 		SDL_Window* GrabbedWindow = SDL_GetGrabbedWindow();
 		if (GrabbedWindow)
 		{
-			SDL_SetWindowGrab(GrabbedWindow, SDL_FALSE);
-			SDL_SetWindowKeyboardGrab(GrabbedWindow, SDL_FALSE);
+			SDL_SetWindowMouseGrab(GrabbedWindow, false);
+			SDL_SetWindowKeyboardGrab(GrabbedWindow, false);
 		}
 
 		SDL_Window* MouseFocusedWindow = SDL_GetMouseFocus();
@@ -246,13 +247,13 @@ void UngrabAllInputImpl()
 			SDL_SetWindowMouseRect(MouseFocusedWindow, nullptr);
 		}
 
-		SDL_SetRelativeMouseMode(SDL_FALSE);
-		SDL_ShowCursor(SDL_ENABLE);
-		SDL_CaptureMouse(SDL_FALSE);
+		SDL_SetWindowRelativeMouseMode(GrabbedWindow, false);
+		SDL_ShowCursor();
+		SDL_CaptureMouse(false);
 	}
 }
 
-uint32 FLinuxPlatformApplicationMisc::WindowStyle()
+uint64 FLinuxPlatformApplicationMisc::WindowStyle()
 {
 	return GWindowStyleSDL;
 }
@@ -277,13 +278,31 @@ void FLinuxPlatformApplicationMisc::Init()
 	UngrabAllInputCallback = UngrabAllInputImpl;
 }
 
+DEFINE_LOG_CATEGORY_STATIC( LogSDL3, Log, All );
+
+static void SDLCALL SDLLogOutputFunction(void *userdata, int category, SDL_LogPriority priority, const char *message)
+{
+	UE_LOG(LogSDL3, Log, TEXT("%s"), UTF8_TO_TCHAR(message));
+}
+
 bool FLinuxPlatformApplicationMisc::InitSDL()
 {
 	if (!GInitializedSDL)
 	{
 		UE_LOG(LogInit, Log, TEXT("Initializing SDL."));
 
-		SDL_SetHint("SDL_VIDEO_X11_REQUIRE_XRANDR", "1");  // workaround for misbuilt SDL libraries on X11.
+#if !UE_BUILD_SHIPPING || USE_LOGGING_IN_SHIPPING
+		SDL_SetLogOutputFunction(SDLLogOutputFunction, NULL);
+#else
+		SDL_SetLogOutputFunction(NULL, NULL);
+#endif
+
+		// In debug mode, let's log all the SDL warnings + errors, etc.
+#if UE_BUILD_DEBUG		
+		SDL_SetLogPriorities(SDL_LOG_PRIORITY_VERBOSE);
+#endif		
+
+
 
 		// pass the string as is (SDL will parse)
 		FString EglDeviceHint;
@@ -297,7 +316,6 @@ bool FLinuxPlatformApplicationMisc::InitSDL()
 		// SDL_SetRelativeMouseMode when enabled is warping the mouse in default mode but we don't want that. 
 		// Furthermore SDL hides the mouse which we prevent by setting SDL_HINT_MOUSE_RELATIVE_CURSOR_VISIBLE
 		SDL_SetHint(SDL_HINT_MOUSE_RELATIVE_CURSOR_VISIBLE, "1"); // When relative mouse mode is active, don't hide cursor.
-		SDL_SetHint(SDL_HINT_MOUSE_RELATIVE_MODE_WARP, "0"); // Don't warp the cursor to the center in relative mouse mode.
 
 		// If we're rendering offscreen, use the "dummy" SDL video driver
 		if (FParse::Param(FCommandLine::Get(), TEXT("RenderOffScreen")) && !getenv("SDL_VIDEODRIVER"))
@@ -306,8 +324,8 @@ bool FLinuxPlatformApplicationMisc::InitSDL()
 			setenv("SDL_VIDEODRIVER", "dummy", 1);
 		}
 
-		// we don't use SDL for audio
-		if (SDL_Init((SDL_INIT_EVERYTHING ^ SDL_INIT_AUDIO) | SDL_INIT_NOPARACHUTE) != 0)
+		// do we need SDL_INIT_HAPTIC, SDL_INIT_SENSOR or SDL_INIT_CAMERA here?
+		if (!SDL_Init(SDL_INIT_VIDEO | SDL_INIT_GAMEPAD | SDL_INIT_EVENTS | SDL_INIT_HAPTIC | SDL_INIT_SENSOR | SDL_INIT_CAMERA | SDL_INIT_JOYSTICK ))
 		{
 			FString ErrorMessage = UTF8_TO_TCHAR(SDL_GetError());
 			if(ErrorMessage != FString(TEXT("No message system available")))
@@ -318,16 +336,18 @@ bool FLinuxPlatformApplicationMisc::InitSDL()
 			return false;
 		}
 
+		// Engine will control enabling and disabling the screensaver itself but SDL init always defaults the screen saver to disabled.
+		// We are going to enable the screen saver here just so the logic the rest of engine uses works
+		SDL_EnableScreenSaver();
+
 		// print out version information
-		SDL_version CompileTimeSDLVersion;
-		SDL_version RunTimeSDLVersion;
-		SDL_VERSION(&CompileTimeSDLVersion);
-		SDL_GetVersion(&RunTimeSDLVersion);
+		int CompileTimeSDLVersion = SDL_VERSION;
+		int RunTimeSDLVersion = SDL_GetVersion();
 		FString SdlRevision = UTF8_TO_TCHAR(SDL_GetRevision());
 		UE_LOG(LogInit, Log, TEXT("Initialized SDL %d.%d.%d revision: %s (compiled against %d.%d.%d)"),
-			RunTimeSDLVersion.major, RunTimeSDLVersion.minor, RunTimeSDLVersion.patch,
+			SDL_VERSIONNUM_MAJOR(RunTimeSDLVersion), SDL_VERSIONNUM_MINOR(RunTimeSDLVersion), SDL_VERSIONNUM_MICRO(RunTimeSDLVersion),
 			*SdlRevision,
-			CompileTimeSDLVersion.major, CompileTimeSDLVersion.minor, CompileTimeSDLVersion.patch
+			SDL_VERSIONNUM_MAJOR(CompileTimeSDLVersion), SDL_VERSIONNUM_MINOR(CompileTimeSDLVersion), SDL_VERSIONNUM_MICRO(CompileTimeSDLVersion)
 			);
 
 		char const* SdlVideoDriver = SDL_GetCurrentVideoDriver();
@@ -338,8 +358,6 @@ bool FLinuxPlatformApplicationMisc::InitSDL()
 			);
 		}
 
-		// Used to make SDL push SDL_TEXTINPUT events.
-		SDL_StartTextInput();
 
 		GInitializedSDL = true;
 
@@ -465,7 +483,7 @@ void FLinuxPlatformApplicationMisc::PumpMessages( bool bFromMainLoop )
 
 bool FLinuxPlatformApplicationMisc::IsScreensaverEnabled()
 {
-	return SDL_IsScreenSaverEnabled();
+	return SDL_ScreenSaverEnabled();
 }
 
 bool FLinuxPlatformApplicationMisc::ControlScreensaver(EScreenSaverAction Action)
@@ -495,52 +513,13 @@ namespace LinuxPlatformApplicationMisc
 
 float FLinuxPlatformApplicationMisc::GetDPIScaleFactorAtPoint(float X, float Y)
 {
-	float Scale = 1.0f;
-
-	if ((GIsEditor || IS_PROGRAM) && IsHighDPIAwarenessEnabled())
-	{
-		FDisplayMetrics DisplayMetrics;
-		FDisplayMetrics::RebuildDisplayMetrics(DisplayMetrics);
-		// find the monitor
-		int32 XInt = static_cast<int32>(X);
-		int32 YInt = static_cast<int32>(Y);
-		static bool bHaveShownDPIWarnings = false;
-		int NumWarningsShown = 0;
-		for(int Idx = 0, NumMonitors = DisplayMetrics.MonitorInfo.Num(); Idx < NumMonitors; ++Idx)
-		{
-			const FMonitorInfo & MonitorInfo = DisplayMetrics.MonitorInfo[Idx];
-			float HorzDPI = 1.0f, VertDPI = 1.0f;
-			if (SDL_GetDisplayDPI(Idx, nullptr, &HorzDPI, &VertDPI) == 0)
-			{
-				if (MonitorInfo.DisplayRect.Left <= XInt && MonitorInfo.DisplayRect.Right > XInt &&
-					MonitorInfo.DisplayRect.Top <= YInt && MonitorInfo.DisplayRect.Bottom > YInt)
-				{
-					Scale = LinuxPlatformApplicationMisc::QuantizeScale((HorzDPI + VertDPI) / 192.0f);	// average between two scales (divided by 96.0f)
-					UE_LOG(LogLinux, Verbose, TEXT("Scale at X=%f, Y=%f: %f (monitor=#%d, HDPI=%f (horz scale: %f), VDPI=%f (vert scale: %f))"), X, Y, Scale, Idx, HorzDPI, HorzDPI / 96.0f, VertDPI, VertDPI / 96.0f);
-
-					// There is a bug in the X11 RandR library in some versions of Linux that can return a bonkers scale.  Even on an 8K display,
-					// it's unlikely that we'd need more than a 4x scale, so capping there for now
-					if(Scale > 4.0f)
-						Scale = 1.0f;
-				}
-			}
-			else
-			{
-				if (!bHaveShownDPIWarnings)
-				{
-					UE_LOG(LogLinux, Log, TEXT("Could not get DPI information for monitor #%d, assuming it to have 1.0f scale"), Idx);
-					++NumWarningsShown;
-				}
-			}
-		}
-
-		if(NumWarningsShown > 0)
-		{
-			bHaveShownDPIWarnings = true;
-		}
-	}
+	SDL_Point Point = {(int)X, (int)Y};
+	SDL_DisplayID DisplayID = SDL_GetDisplayForPoint(&Point);
+	float Scale = LinuxPlatformApplicationMisc::QuantizeScale(SDL_GetDisplayContentScale(DisplayID));
 
+	UE_LOG(LogLinux, Verbose, TEXT("Scale at X=%f, Y=%f: %f (monitor=#%d)"), X, Y, Scale, DisplayID);
 	return Scale;
+
 }
 
 void FLinuxPlatformApplicationMisc::ClipboardCopy(const TCHAR* Str)
diff --git a/Engine/Source/Runtime/ApplicationCore/Private/Linux/LinuxPlatformSplash.cpp b/Engine/Source/Runtime/ApplicationCore/Private/Linux/LinuxPlatformSplash.cpp
--- a/Engine/Source/Runtime/ApplicationCore/Private/Linux/LinuxPlatformSplash.cpp
+++ b/Engine/Source/Runtime/ApplicationCore/Private/Linux/LinuxPlatformSplash.cpp
@@ -23,7 +23,7 @@
 #include "Modules/ModuleManager.h"
 #include "HAL/PlatformApplicationMisc.h"
 
-#include "SDL.h"
+#include <SDL3/SDL.h>
 
 #if WITH_EDITOR
 
@@ -110,7 +110,7 @@ FLinuxSplashState::~FLinuxSplashState()
 
 	if (SplashSurface)
 	{
-		SDL_FreeSurface(SplashSurface);
+		SDL_DestroySurface(SplashSurface);
 	}
 
 	if (SplashWindow)
@@ -221,7 +221,8 @@ void FLinuxSplashState::DrawCharacter(int32 PenX, int32 PenY, FT_GlyphSlot Glyph
 	unsigned char *Pixels = Glyph->bitmap.buffer;
 
 	const int SplashWidth = SplashSurface->w;
-	const int SplashBPP = SplashSurface->format->BytesPerPixel;
+	const SDL_PixelFormatDetails *PixelFormat = SDL_GetPixelFormatDetails(SplashSurface->format);
+	const int SplashBPP = PixelFormat->bytes_per_pixel;
 
 	// draw glyph raster to texture
 	for (int GlyphY = 0; GlyphY < GlyphHeight; GlyphY++)
@@ -271,7 +272,8 @@ void FLinuxSplashState::RenderStrings()
 
 	const int SplashWidth = SplashSurface->w;
 	const int SplashHeight =  SplashSurface->h;
-	const int SplashBPP = SplashSurface->format->BytesPerPixel;
+	const SDL_PixelFormatDetails *PixelFormat = SDL_GetPixelFormatDetails(SplashSurface->format);
+	const int SplashBPP = PixelFormat->bytes_per_pixel;
 	unsigned char *SplashPixels = reinterpret_cast<unsigned char *>(SplashSurface->pixels);
 
 	// clear the rendering scratch pad.
@@ -427,12 +429,13 @@ SDL_Surface* FLinuxSplashState::LoadImage(const FString &ImagePath)
 			TArray<uint8> RawData;
 			if (ImageWrapper->GetRaw(ERGBFormat::BGRA, 8, RawData))
 			{
-				SDL_Surface *Surface = SDL_CreateRGBSurfaceWithFormat(0, ImageWrapper->GetWidth(), ImageWrapper->GetHeight(), 32, SDL_PIXELFORMAT_BGRA32);
+				SDL_Surface *Surface = SDL_CreateSurface(ImageWrapper->GetWidth(), ImageWrapper->GetHeight(), SDL_PIXELFORMAT_BGRA32);
 				if (Surface)
 				{
 					const int Width = Surface->w;
 					const int Height =  Surface->h;
-					const int BytesPerPixel = Surface->format->BytesPerPixel;
+					const SDL_PixelFormatDetails *PixelFormat = SDL_GetPixelFormatDetails(Surface->format);
+					const int BytesPerPixel = PixelFormat->bytes_per_pixel;
 					FMemory::Memcpy(Surface->pixels, (void*)RawData.GetData(), Width * Height * BytesPerPixel);
 					return Surface;
 				}
@@ -467,7 +470,8 @@ bool FLinuxSplashState::InitSplashResources(const FText &AppName, const FString
 
 	const int SplashWidth = SplashSurface->w;
 	const int SplashHeight =  SplashSurface->h;
-	const int SplashBPP = SplashSurface->format->BytesPerPixel;
+	const SDL_PixelFormatDetails *PixelFormat = SDL_GetPixelFormatDetails(SplashSurface->format);
+	const int SplashBPP = PixelFormat->bytes_per_pixel;
 
 	if (SplashWidth <= 0 || SplashHeight <= 0)
 	{
@@ -490,11 +494,25 @@ bool FLinuxSplashState::InitSplashResources(const FText &AppName, const FString
 	const char *_OriginalHint = SDL_GetHint(SDL_HINT_VIDEO_X11_NET_WM_PING);  // (SDL_SetHint will free this pointer.)
 	char *OriginalHint = _OriginalHint ? SDL_strdup(_OriginalHint) : NULL;
 	SDL_SetHint(SDL_HINT_VIDEO_X11_NET_WM_PING, "0");
-	SplashWindow = SDL_CreateWindow(TCHAR_TO_UTF8(*AppName.ToString()),
-									SDL_WINDOWPOS_CENTERED, SDL_WINDOWPOS_CENTERED,
-									SplashWidth, SplashHeight,
-									SDL_WINDOW_BORDERLESS | SDL_WINDOW_HIDDEN
-									);
+	SDL_PropertiesID Props = SDL_CreateProperties();
+	if (Props == 0)
+	{
+		if (OriginalHint)
+		{
+			SDL_free(OriginalHint);
+		}
+		UE_LOG(LogHAL, Error, TEXT("Failed to create SDL properties"));
+		return false;
+	}
+	SDL_SetStringProperty(Props, SDL_PROP_WINDOW_CREATE_TITLE_STRING, TCHAR_TO_UTF8(*AppName.ToString()));
+	SDL_SetNumberProperty(Props, SDL_PROP_WINDOW_CREATE_WIDTH_NUMBER, SplashWidth);
+	SDL_SetNumberProperty(Props, SDL_PROP_WINDOW_CREATE_HEIGHT_NUMBER, SplashHeight);
+	SDL_SetNumberProperty(Props, SDL_PROP_WINDOW_CREATE_X_NUMBER, SDL_WINDOWPOS_CENTERED_DISPLAY(0));
+	SDL_SetNumberProperty(Props, SDL_PROP_WINDOW_CREATE_Y_NUMBER, SDL_WINDOWPOS_CENTERED_DISPLAY(0));
+	SDL_SetBooleanProperty(Props, SDL_PROP_WINDOW_CREATE_BORDERLESS_BOOLEAN, true);
+	SDL_SetBooleanProperty(Props, SDL_PROP_WINDOW_CREATE_HIDDEN_BOOLEAN, true);
+	
+	SplashWindow = SDL_CreateWindowWithProperties(Props);
 
 	if (SplashWindow == nullptr)
 	{
@@ -515,12 +533,12 @@ bool FLinuxSplashState::InitSplashResources(const FText &AppName, const FString
 		else
 		{
 			SDL_SetWindowIcon(SplashWindow, SplashIconImage);  // SDL_SetWindowIcon() makes a copy of this.
-			SDL_FreeSurface(SplashIconImage);
+			SDL_DestroySurface(SplashIconImage);
 		}
 	}
 
 #if WITH_EDITOR
-	SplashRenderer = SDL_CreateRenderer(SplashWindow, -1, 0);
+	SplashRenderer = SDL_CreateRenderer(SplashWindow, nullptr);
 #endif
 
 	// it's safe to set the hint back once the Renderer is created (since it might recreate the window to get a GL or whatever visual).
@@ -535,7 +553,7 @@ bool FLinuxSplashState::InitSplashResources(const FText &AppName, const FString
 		return false;
 	}
 
-	SplashTexture = SDL_CreateTexture(SplashRenderer, SplashSurface->format->format, SDL_TEXTUREACCESS_STATIC, SplashSurface->w, SplashSurface->h);
+	SplashTexture = SDL_CreateTexture(SplashRenderer, SDL_GetPixelFormatDetails(SplashSurface->format)->format, SDL_TEXTUREACCESS_STATIC, SplashSurface->w, SplashSurface->h);
 	if (SplashTexture == nullptr)
 	{
 		UE_LOG(LogHAL, Error, TEXT("FLinuxSplashState::InitSplashResources() : Splash screen texture could not be created! SDL_Error: %s"), UTF8_TO_TCHAR(SDL_GetError()));
@@ -600,7 +618,7 @@ bool FLinuxSplashState::InitSplashResources(const FText &AppName, const FString
 	}
 
 	SDL_BlitSurface(SplashSurface, NULL, WindowSurface, NULL);
-	SDL_FreeSurface(WindowSurface);
+	SDL_DestroySurface(WindowSurface);
 
 	SDL_ShowWindow(SplashWindow);
 	Pump();
@@ -640,7 +658,7 @@ void FLinuxSplashState::Redraw()
 		SDL_GLContext CurrentContext = SDL_GL_GetCurrentContext();
 
 		RenderStrings();
-		SDL_RenderCopy(SplashRenderer, SplashTexture, NULL, NULL);
+		SDL_RenderTexture(SplashRenderer, SplashTexture, NULL, NULL);
 		SDL_RenderPresent(SplashRenderer);
 		bNeedsRedraw = false;
 
@@ -666,7 +684,7 @@ void FLinuxSplashState::Pump()
 		SDL_Event Event;
 		while (SDL_PollEvent(&Event))
 		{
-			if ((Event.type == SDL_WINDOWEVENT) && (Event.window.event == SDL_WINDOWEVENT_EXPOSED) && (SDL_GetWindowID(SplashWindow) == Event.window.windowID))
+			if ((Event.type == SDL_EVENT_WINDOW_EXPOSED) && (SDL_GetWindowID(SplashWindow) == Event.window.windowID))
 			{
 				bNeedsRedraw = true;
 			}
diff --git a/Engine/Source/Runtime/ApplicationCore/Private/Linux/LinuxWindow.cpp b/Engine/Source/Runtime/ApplicationCore/Private/Linux/LinuxWindow.cpp
--- a/Engine/Source/Runtime/ApplicationCore/Private/Linux/LinuxWindow.cpp
+++ b/Engine/Source/Runtime/ApplicationCore/Private/Linux/LinuxWindow.cpp
@@ -84,14 +84,12 @@ void FLinuxWindow::Initialize( FLinuxApplication* const Application, const TShar
 	}
 
 #if DO_CHECK
-	uint32 InitializedSubsystems = SDL_WasInit(SDL_INIT_EVERYTHING);
+	uint32 InitializedSubsystems = SDL_WasInit(SDL_INIT_VIDEO);
 	check(InitializedSubsystems & SDL_INIT_VIDEO);
 #endif // DO_CHECK
 
 	// Finally, let's initialize the new native window object.  Calling this function will often cause OS
 	// window messages to be sent! (such as activation messages)
-	uint32 WindowExStyle = 0;
-	uint32 WindowStyle = 0;
 
 	RegionWidth = RegionHeight = INDEX_NONE;
 
@@ -113,29 +111,36 @@ void FLinuxWindow::Initialize( FLinuxApplication* const Application, const TShar
 	int32 WindowWidth = ClientWidth;
 	int32 WindowHeight = ClientHeight;
 
-	WindowStyle |= FLinuxPlatformApplicationMisc::WindowStyle() | SDL_WINDOW_SHOWN;
 
-	if ( !Definition->HasOSWindowBorder )
-	{
-		WindowStyle |= SDL_WINDOW_BORDERLESS;
-
-		if (Definition->IsTopmostWindow)
-		{
-			WindowStyle |= SDL_WINDOW_ALWAYS_ON_TOP;
-		}
+        //      The SDL window doesn't need to be reshaped.
+        //      the size of the window you input is the sizeof the client.
+	bool bIsResizable = false;
+        SDL_PropertiesID Props = SDL_CreateProperties();
+        SDL_SetStringProperty( Props, SDL_PROP_WINDOW_CREATE_TITLE_STRING, TCHAR_TO_ANSI( *Definition->Title ) );
+        SDL_SetNumberProperty( Props, SDL_PROP_WINDOW_CREATE_X_NUMBER, X );
+        SDL_SetNumberProperty( Props, SDL_PROP_WINDOW_CREATE_Y_NUMBER, Y );
+        SDL_SetNumberProperty( Props, SDL_PROP_WINDOW_CREATE_WIDTH_NUMBER, ClientWidth );
+        SDL_SetNumberProperty( Props, SDL_PROP_WINDOW_CREATE_HEIGHT_NUMBER, ClientHeight );
 
-		if (!Definition->AppearsInTaskbar)
-		{
-			WindowStyle |= SDL_WINDOW_SKIP_TASKBAR;
-		}
+	UE_LOG(LogLinuxWindowType, Verbose, TEXT("Linux Window: X=%d Y=%d, Width=%d, Height=%d"), X, Y, ClientWidth, ClientHeight);
+	switch(FLinuxPlatformApplicationMisc::WindowStyle())
+	{
+		case SDL_WINDOW_VULKAN:
+			SDL_SetBooleanProperty( Props, SDL_PROP_WINDOW_CREATE_VULKAN_BOOLEAN, true );
+			break;
+		case SDL_WINDOW_OPENGL:
+			SDL_SetBooleanProperty( Props, SDL_PROP_WINDOW_CREATE_OPENGL_BOOLEAN, true );
+			break;
 	}
 
 	if (Definition->IsRegularWindow && Definition->HasSizingFrame)
 	{
-		WindowStyle |= SDL_WINDOW_RESIZABLE;
+		SDL_SetBooleanProperty( Props, SDL_PROP_WINDOW_CREATE_RESIZABLE_BOOLEAN, true );
+		bIsResizable = true;
 	}
 
 	const bool bShouldActivate = Definition->ActivationPolicy != EWindowActivationPolicy::Never;
+	const char *SDLWindowType = nullptr;			// use this to set the X11 Window Type if necessary
 
 	// This is a tool tip window.
 	if (!InParent.IsValid() && !Definition->HasOSWindowBorder &&
@@ -144,7 +149,9 @@ void FLinuxWindow::Initialize( FLinuxApplication* const Application, const TShar
 		!Definition->HasSizingFrame && !Definition->IsModalWindow &&
 		!Definition->IsRegularWindow && Definition->SizeWillChangeOften)
 	{
-		WindowStyle |= SDL_WINDOW_TOOLTIP;
+		SDL_SetBooleanProperty( Props, SDL_PROP_WINDOW_CREATE_HIDDEN_BOOLEAN, true );
+		// We don't set SDL_PROP_WINDOW_CREATE_TOOLTIP_BOOLEAN to true here because it requires a parent and Slate isn't set up
+		// to provide one for tooltips
 		bIsTooltipWindow = true;
 		UE_LOG(LogLinuxWindowType, Verbose, TEXT("*** New Window is a Tooltip Window ***"));
 	}
@@ -155,7 +162,7 @@ void FLinuxWindow::Initialize( FLinuxApplication* const Application, const TShar
 		!Definition->IsModalWindow && !Definition->IsRegularWindow &&
 		!bShouldActivate && Definition->SizeWillChangeOften)
 	{
-		WindowStyle |= SDL_WINDOW_NOTIFICATION;
+		SDLWindowType = "_NET_WM_WINDOW_TYPE_NOTIFICATION";
 		bIsNotificationWindow = true;
 		UE_LOG(LogLinuxWindowType, Verbose, TEXT("*** New Window is a Notification Window ***"));
 	}
@@ -166,7 +173,9 @@ void FLinuxWindow::Initialize( FLinuxApplication* const Application, const TShar
 		Definition->IsModalWindow && !Definition->IsRegularWindow &&
 		bShouldActivate && !Definition->SizeWillChangeOften)
 	{
-		WindowStyle |= SDL_WINDOW_NOTIFICATION;
+		SDL_SetBooleanProperty( Props, SDL_PROP_WINDOW_CREATE_MODAL_BOOLEAN, true );
+        SDL_SetPointerProperty( Props, SDL_PROP_WINDOW_CREATE_PARENT_POINTER, InParent->GetHWnd() );
+		SDLWindowType = "_NET_WM_WINDOW_TYPE_NOTIFICATION";
 		bIsNotificationWindow = true;
 		UE_LOG(LogLinuxWindowType, Verbose, TEXT("*** New Window is another Notification Window ***"));
 	}
@@ -177,7 +186,9 @@ void FLinuxWindow::Initialize( FLinuxApplication* const Application, const TShar
 		!Definition->IsModalWindow && !Definition->IsRegularWindow &&
 		bShouldActivate && !Definition->SizeWillChangeOften)
 	{
-		WindowStyle |= SDL_WINDOW_POPUP_MENU;
+		SDL_SetBooleanProperty( Props, SDL_PROP_WINDOW_CREATE_HIDDEN_BOOLEAN, true );
+		SDL_SetBooleanProperty( Props, SDL_PROP_WINDOW_CREATE_TRANSPARENT_BOOLEAN, true );
+		SDL_SetBooleanProperty( Props, SDL_PROP_WINDOW_CREATE_BORDERLESS_BOOLEAN, true );
 		bIsPopupWindow = true;
 		UE_LOG(LogLinuxWindowType, Verbose, TEXT("*** New Window is a Popup Menu Window ***"));
 	}
@@ -188,7 +199,9 @@ void FLinuxWindow::Initialize( FLinuxApplication* const Application, const TShar
 		!Definition->IsModalWindow && !Definition->IsRegularWindow &&
 		!bShouldActivate && !Definition->SizeWillChangeOften)
 	{
-		WindowStyle |= SDL_WINDOW_POPUP_MENU;
+		SDL_SetBooleanProperty( Props, SDL_PROP_WINDOW_CREATE_HIDDEN_BOOLEAN, true );
+		SDL_SetBooleanProperty( Props, SDL_PROP_WINDOW_CREATE_TRANSPARENT_BOOLEAN, true );
+		SDL_SetBooleanProperty( Props, SDL_PROP_WINDOW_CREATE_BORDERLESS_BOOLEAN, true );
 		bIsConsoleWindow = true;
 		bIsPopupWindow = true;
 		UE_LOG(LogLinuxWindowType, Verbose, TEXT("*** New Window is a Console Window ***"));
@@ -201,7 +214,7 @@ void FLinuxWindow::Initialize( FLinuxApplication* const Application, const TShar
 		!bShouldActivate && !Definition->SizeWillChangeOften)
 	{
 		// TODO Experimental (The SDL_WINDOW_DND sets focus)
-		WindowStyle |= SDL_WINDOW_DND;
+		SDLWindowType = "_NET_WM_WINDOW_TYPE_DND";
 		bIsDragAndDropWindow = true;
 		UE_LOG(LogLinuxWindowType, Verbose, TEXT("*** New Window is a Drag and Drop Window ***"));
 	}
@@ -212,7 +225,9 @@ void FLinuxWindow::Initialize( FLinuxApplication* const Application, const TShar
 		Definition->IsRegularWindow && bShouldActivate &&
 		!Definition->SizeWillChangeOften)
 	{
-		WindowStyle |= SDL_WINDOW_DIALOG;
+		SDLWindowType = "_NET_WM_WINDOW_TYPE_DIALOG";
+		SDL_SetBooleanProperty( Props, SDL_PROP_WINDOW_CREATE_MODAL_BOOLEAN, true );
+		SDL_SetPointerProperty( Props, SDL_PROP_WINDOW_CREATE_PARENT_POINTER, InParent->GetHWnd() );
 		bIsDialogWindow = true;
 		UE_LOG(LogLinuxWindowType, Verbose, TEXT("*** New Window is a Modal Dialog Window ***"));
 	}
@@ -223,7 +238,6 @@ void FLinuxWindow::Initialize( FLinuxApplication* const Application, const TShar
 		!Definition->IsModalWindow && Definition->IsRegularWindow &&
 		bShouldActivate && !Definition->SizeWillChangeOften)
 	{
-		WindowStyle |= SDL_WINDOW_DIALOG;
 		bIsUtilityWindow = true;
 		UE_LOG(LogLinuxWindowType, Verbose, TEXT("*** New Window is a BP, Cascade, etc. Window ***"));
 	}
@@ -233,43 +247,81 @@ void FLinuxWindow::Initialize( FLinuxApplication* const Application, const TShar
 		bIsTopLevelWindow = true;
 	}
 
-	//	The SDL window doesn't need to be reshaped.
-	//	the size of the window you input is the sizeof the client.
-	HWnd = SDL_CreateWindow( TCHAR_TO_ANSI( *Definition->Title ), X, Y, ClientWidth, ClientHeight, WindowStyle  );
-	// produce a helpful message for common driver errors
-	if (HWnd == nullptr)
+	if ( !Definition->HasOSWindowBorder )
 	{
-		FString ErrorMessage;
+		SDL_SetBooleanProperty( Props, SDL_PROP_WINDOW_CREATE_BORDERLESS_BOOLEAN, true );
 
-		if ((WindowStyle & SDL_WINDOW_VULKAN) != 0)
+		if (Definition->IsTopmostWindow)
 		{
-			ErrorMessage = LOCTEXT("VulkanWindowCreationFailedLinux", "Unable to create a Vulkan window - make sure an up-to-date libvulkan.so.1 is installed.").ToString();
-			FPlatformMisc::MessageBoxExt(EAppMsgType::Ok, *ErrorMessage,
-										 *LOCTEXT("VulkanWindowCreationFailedLinuxTitle", "Unable to create a Vulkan window.").ToString());
+			SDL_SetBooleanProperty( Props, SDL_PROP_WINDOW_CREATE_ALWAYS_ON_TOP_BOOLEAN, true );
 		}
-		else if ((WindowStyle & SDL_WINDOW_OPENGL) != 0)
+
+		// MODAL and NOTIFICATION windows will already not be shown in the task bar
+		if (!Definition->AppearsInTaskbar && !Definition->IsModalWindow && !bIsNotificationWindow)
 		{
-			ErrorMessage = LOCTEXT("OpenGLWindowCreationFailedLinux", "Unable to create an OpenGL window - make sure your drivers support at least OpenGL 4.3.").ToString();
-			FPlatformMisc::MessageBoxExt(EAppMsgType::Ok, *ErrorMessage,
-										 *LOCTEXT("OpenGLWindowCreationFailedLinuxTitle", "Unable to create an OpenGL window.").ToString());
+
+			SDL_SetBooleanProperty( Props, SDL_PROP_WINDOW_CREATE_UTILITY_BOOLEAN, true );
 		}
-		else
+	}
+
+	/* 
+			Do not set for Notification Windows the transient flag because the WM's usually raise the the parent window
+			if the Notificaton Window gets raised. That behaviour is to aggresive and disturbs users doing other things 
+			while UnrealEngine calculates lights and other things and pop ups notifications. Notifications will be handled so that 
+			they are some sort of independend but will be raised if the TopLevel Window gets focused or activated.
+	*/
+	
+	// Set parent and make the Window modal if necessary
+	if (bIsUtilityWindow || bIsDialogWindow || bIsConsoleWindow)
+	{
+		SDL_SetPointerProperty( Props, SDL_PROP_WINDOW_CREATE_PARENT_POINTER, InParent->GetHWnd() );
+		if (Definition->IsModalWindow)
 		{
-			ErrorMessage = FText::Format(LOCTEXT("SDLWindowCreationFailedLinuxFmt", "Window creation failed (SDL error: '{0}'')"), FText::FromString(UTF8_TO_TCHAR(SDL_GetError()))).ToString();
-			FPlatformMisc::MessageBoxExt(EAppMsgType::Ok, *ErrorMessage,
-										 *LOCTEXT("SDLWindowCreationFailedLinuxTitle", "Unable to create an SDL window.").ToString());
+			SDL_SetBooleanProperty( Props, SDL_PROP_WINDOW_CREATE_MODAL_BOOLEAN, true );
 		}
+	}
+
+	// This lets us override the X11 window type
+	if (SDLWindowType != nullptr)
+	{
+		SDL_SetHint(SDL_HINT_X11_WINDOW_TYPE, SDLWindowType);
+	}
+
+	HWnd = SDL_CreateWindowWithProperties( Props );
+	SDL_DestroyProperties( Props );
+
+	// Restore the X11 Window Type
+	if (SDLWindowType != nullptr)
+	{
+		SDL_ResetHint(SDL_HINT_X11_WINDOW_TYPE);
+	}
+	
+	// produce a helpful message for common driver errors
+	if (HWnd == nullptr)
+	{
+		FString ErrorMessage;
+
+		ErrorMessage = FText::Format(LOCTEXT("SDLWindowCreationFailedLinuxFmt", "Window creation failed (SDL error: '{0}'')"), FText::FromString(UTF8_TO_TCHAR(SDL_GetError()))).ToString();
+		FPlatformMisc::MessageBoxExt(EAppMsgType::Ok, *ErrorMessage,
+									 *LOCTEXT("SDLWindowCreationFailedLinuxTitle", "Unable to create an SDL window.").ToString());
 
 		checkf(false, TEXT("%s"), *ErrorMessage);
 		// unreachable
 		return;
 	}
 
+	// SDL_StartTextInput is no longer a global state, so we need to set it on every
+	// window we create to have it work like it used to
+	if (Definition->AcceptsInput)
+	{
+		SDL_StartTextInput(HWnd);
+	}
+
 	// SDL_SetWindow[Minimum/Maximum]Size() requires a valid value for BOTH width and height, we can't do one or the other without magic numbers
 	const bool bIsMinSizeSet = Definition->SizeLimits.GetMinWidth().IsSet() && Definition->SizeLimits.GetMinHeight().IsSet();
 	const bool bIsMaxSizeSet = Definition->SizeLimits.GetMaxWidth().IsSet() && Definition->SizeLimits.GetMaxHeight().IsSet();
 
-	if ((WindowStyle & SDL_WINDOW_RESIZABLE) && !bIsMinSizeSet)
+	if (bIsResizable && !bIsMinSizeSet)
 	{
 		// Set a reasonable minimum size if none are set
 		// Vulkan will fail to create a swap chain if the back buffer is 0x0
@@ -333,25 +385,13 @@ void FLinuxWindow::Initialize( FLinuxApplication* const Application, const TShar
 		{
 			SDL_SetWindowIcon(HWnd, IconImage);
 
-			SDL_FreeSurface(IconImage);
+			SDL_DestroySurface(IconImage);
 			IconImage = nullptr;
 		}
 	}
 
 	SDL_SetWindowHitTest( HWnd, FLinuxWindow::HitTest, this );
 
-	/* 
-		Do not set for Notification Windows the transient flag because the WM's usually raise the the parent window
-		if the Notificaton Window gets raised. That behaviour is to aggresive and disturbs users doing other things 
-		while UnrealEngine calculates lights and other things and pop ups notifications. Notifications will be handled so that 
-		they are some sort of independend but will be raised if the TopLevel Window gets focused or activated.
-	*/
-	// Make the Window modal for it's parent.
-	if (bIsUtilityWindow || bIsDialogWindow || bIsConsoleWindow || bIsNotificationWindow)
-	{
-		SDL_SetWindowModalFor(HWnd, InParent->GetHWnd());
-	}
-
 	VirtualWidth  = ClientWidth;
 	VirtualHeight = ClientHeight;
 
@@ -379,7 +419,14 @@ SDL_HitTestResult FLinuxWindow::HitTest( SDL_Window *SDLwin, const SDL_Point *po
 	TSharedPtr< FLinuxWindow > self = pself->OwningApplication->FindWindowBySDLWindow( SDLwin );
 	if ( !self.IsValid() ) 
 	{
-		UE_LOG(LogLinuxWindow, Warning, TEXT("BAD EVENT: SDL window = %p\n"), SDLwin);
+		// When a window is being destroyed it can receive events before the window is actually removed.
+		// At this point it is no longer in UE's Window list so fails the IsValid test
+		if (pself->GetHWnd() == SDLwin)
+		{
+			return SDL_HITTEST_NORMAL;
+		}
+
+		UE_LOG(LogLinuxWindow, Warning, TEXT("BAD HIT TEST EVENT: SDL window = %p\n"), SDLwin);
 		return SDL_HITTEST_NORMAL;
 	}
 
@@ -509,23 +556,31 @@ static void GetBestFullscreenResolution( SDL_HWindow hWnd, int32 *pWidth, int32
 	uint32 BestHeight = 0;
 	uint32 ModeIndex = 0;
 
-	int32 dsp_idx = SDL_GetWindowDisplayIndex( hWnd );
-	if ( dsp_idx < 0 )
-	{	dsp_idx = 0;	}
-
-	SDL_DisplayMode dsp_mode;
-	FMemory::Memzero( &dsp_mode, sizeof(SDL_DisplayMode) );
+	int32 DisplayIndex = SDL_GetDisplayForWindow( hWnd );
+	if ( DisplayIndex < 0 )
+	{	
+		DisplayIndex = 0;
+	}
 
-	while ( !SDL_GetDisplayMode( dsp_idx, ModeIndex++, &dsp_mode ) )
+	int NumModes = 0;
+	SDL_DisplayMode** Modes = SDL_GetFullscreenDisplayModes(DisplayIndex, &NumModes);
+	if (Modes)
 	{
-		bool IsEqualOrBetterWidth  = FMath::Abs((int32)dsp_mode.w - (int32)(*pWidth))  <= FMath::Abs((int32)BestWidth  - (int32)(*pWidth ));
-		bool IsEqualOrBetterHeight = FMath::Abs((int32)dsp_mode.h - (int32)(*pHeight)) <= FMath::Abs((int32)BestHeight - (int32)(*pHeight));
-		if	(!InitializedMode || (IsEqualOrBetterWidth && IsEqualOrBetterHeight))
+		for (int i = 0; i < NumModes; ++i)
 		{
-			BestWidth = dsp_mode.w;
-			BestHeight = dsp_mode.h;
-			InitializedMode = true;
+			SDL_DisplayMode* Mode = Modes[i];
+			bool IsEqualOrBetterWidth  = FMath::Abs((int32)Mode->w - (int32)(*pWidth))  <= FMath::Abs((int32)BestWidth  - (int32)(*pWidth ));
+			bool IsEqualOrBetterHeight = FMath::Abs((int32)Mode->h - (int32)(*pHeight)) <= FMath::Abs((int32)BestHeight - (int32)(*pHeight));
+			if (!InitializedMode || (IsEqualOrBetterWidth && IsEqualOrBetterHeight))
+			{
+				BestWidth = Mode->w;
+				BestHeight = Mode->h;
+				InitializedMode = true;
+			}
+
+			
 		}
+		SDL_free(Modes);
 	}
 
 	check(InitializedMode);
@@ -563,8 +618,9 @@ void FLinuxWindow::ReshapeWindow( int32 NewX, int32 NewY, int32 NewWidth, int32
 		case EWindowMode::Fullscreen:
 		case EWindowMode::WindowedFullscreen:
 		{
-			SDL_SetWindowFullscreen( HWnd, 0 );
-			SDL_SetWindowFullscreen( HWnd, SDL_WINDOW_FULLSCREEN_DESKTOP );
+			// these are asynchronous calls now, so we need to sync to ensure they've completed
+			SDL_SetWindowFullscreen( HWnd, true );
+			SDL_SyncWindow( HWnd );
 			bWasFullscreen = true;
 		}
 		break;
@@ -580,6 +636,8 @@ void FLinuxWindow::ReshapeWindow( int32 NewX, int32 NewY, int32 NewWidth, int32
 			}
 			SDL_SetWindowPosition( HWnd, NewX, NewY );
 			SDL_SetWindowSize( HWnd, NewWidth, NewHeight );
+			SDL_SetWindowFullscreen(HWnd, false);
+			SDL_SyncWindow( HWnd );
 
 			bWasFullscreen = false;
 
@@ -593,16 +651,8 @@ void FLinuxWindow::ReshapeWindow( int32 NewX, int32 NewY, int32 NewWidth, int32
 	VirtualHeight = NewHeight;
 
 	// Avoid broadcasting we have set a zero size as it will attempt to resize the backbuffer which on some RHI is invalid per the spec (ie. Vulkan)
-	if (LinuxWindow )
-	{
-		OwningApplication->GetMessageHandler()->OnSizeChanged(
-			LinuxWindow.ToSharedRef(),
-			VirtualWidth,
-			VirtualHeight,
-			//  bWasMinimized
-			false
-		);
-	}
+	// OnSizeChanged will be called already since SDL_EVENT_WINDOW_RESIZED is now always sent even when SDL_SetWindowSize() is called
+	// on SDL3
 }
 
 /** Toggle native window between fullscreen and normal mode */
@@ -627,7 +677,8 @@ void FLinuxWindow::SetWindowMode( EWindowMode::Type NewWindowMode )
 						OwningApplication->GetMessageHandler()->OnResizingWindow( LinuxWindow.ToSharedRef() );
 					}
 
-					SDL_SetWindowFullscreen( HWnd, SDL_WINDOW_FULLSCREEN_DESKTOP );
+					SDL_SetWindowFullscreen( HWnd, true );
+					SDL_SyncWindow( HWnd );
 					bWasFullscreen = true;
 
 					if ( LinuxWindow )
@@ -653,10 +704,12 @@ void FLinuxWindow::SetWindowMode( EWindowMode::Type NewWindowMode )
 				int SmallerHeight = FMath::Max(100, VirtualHeight - 100);
 				SDL_SetWindowSize(HWnd, SmallerWidth, SmallerHeight);
 
-				SDL_SetWindowFullscreen( HWnd, 0 );
-				SDL_SetWindowBordered( HWnd, SDL_TRUE );
+				SDL_SetWindowFullscreen( HWnd, false );
+				SDL_SetWindowBordered( HWnd, true );
+
+				SDL_SetWindowMouseGrab( HWnd, false );
 
-				SDL_SetWindowGrab( HWnd, SDL_FALSE );
+				SDL_SyncWindow( HWnd );
 
 				bWasFullscreen = false;
 			}
@@ -708,8 +761,8 @@ bool FLinuxWindow::GetFullScreenInfo( int32& X, int32& Y, int32& Width, int32& H
 {
 	SDL_Rect DisplayRect;
 
-	int DisplayIdx = SDL_GetWindowDisplayIndex(HWnd);
-	if (DisplayIdx >= 0 && SDL_GetDisplayBounds(DisplayIdx, &DisplayRect) == 0)
+	int DisplayIdx = SDL_GetDisplayForWindow(HWnd);
+	if (DisplayIdx >= 0 && SDL_GetDisplayBounds(DisplayIdx, &DisplayRect))
 	{
 		X = DisplayRect.x;
 		Y = DisplayRect.y;
diff --git a/Engine/Source/Runtime/ApplicationCore/Public/Linux/LinuxApplication.h b/Engine/Source/Runtime/ApplicationCore/Public/Linux/LinuxApplication.h
--- a/Engine/Source/Runtime/ApplicationCore/Public/Linux/LinuxApplication.h
+++ b/Engine/Source/Runtime/ApplicationCore/Public/Linux/LinuxApplication.h
@@ -161,8 +161,6 @@ private:
 
 	FLinuxApplication();
 
-	TCHAR ConvertChar( SDL_Keysym Keysym );
-
 	/**
 	 * Finds a window associated with the event (if there is such an association).
 	 *
@@ -231,7 +229,7 @@ public:
 	virtual void ResetLightColor(int32 ControllerId) override { }
 
 private:
-	void AddGameController(int Index);
+	void AddGameController(SDL_JoystickID Id);
 	void RemoveGameController(SDL_JoystickID Id);
 
 	/** Stores context information about a currently active touch. */
@@ -256,7 +254,7 @@ private:
 	struct SDLControllerState
 	{
 		/** SDL controller */
-		SDL_GameController* Controller;
+		SDL_Gamepad* Controller;
 
 		/** Tracks whether the "button" was previously pressed so we don't generate extra events */
 		bool AnalogOverThreshold[10];
diff --git a/Engine/Source/Runtime/ApplicationCore/Public/Linux/LinuxCursor.h b/Engine/Source/Runtime/ApplicationCore/Public/Linux/LinuxCursor.h
--- a/Engine/Source/Runtime/ApplicationCore/Public/Linux/LinuxCursor.h
+++ b/Engine/Source/Runtime/ApplicationCore/Public/Linux/LinuxCursor.h
@@ -5,7 +5,7 @@
 #include "CoreTypes.h"
 #include "Math/Vector2D.h"
 #include "GenericPlatform/ICursor.h"
-#include "SDL.h"
+#include <SDL3/SDL.h>
 
 typedef SDL_Cursor*		SDL_HCursor;
 
@@ -59,7 +59,7 @@ public:
 	void InvalidateCaches();
 
 	/** Set the internal cached position, setting the cache to valid */
-	void SetCachedPosition( const int32 X, const int32 Y );
+	void SetCachedPosition( const float X, const float Y );
 
 private:
 
@@ -75,10 +75,10 @@ private:
 	SDL_Rect CursorClipRect;
 
 	/** Cached global X position */
-	mutable int32 CachedGlobalXPosition;
+	mutable float CachedGlobalXPosition;
 
 	/** Cached global Y position */
-	mutable int32 CachedGlobalYPosition;
+	mutable float CachedGlobalYPosition;
 
 	/** Whether mouse position cache is valid */
 	mutable bool bPositionCacheIsValid;
diff --git a/Engine/Source/Runtime/ApplicationCore/Public/Linux/LinuxPlatformApplicationMisc.h b/Engine/Source/Runtime/ApplicationCore/Public/Linux/LinuxPlatformApplicationMisc.h
--- a/Engine/Source/Runtime/ApplicationCore/Public/Linux/LinuxPlatformApplicationMisc.h
+++ b/Engine/Source/Runtime/ApplicationCore/Public/Linux/LinuxPlatformApplicationMisc.h
@@ -12,7 +12,7 @@ struct FLinuxPlatformApplicationMisc : public FGenericPlatformApplicationMisc
 	static APPLICATIONCORE_API void TearDown();
 	static APPLICATIONCORE_API void LoadPreInitModules();
 	static APPLICATIONCORE_API void LoadStartupModules();
-	static APPLICATIONCORE_API uint32 WindowStyle();
+	static APPLICATIONCORE_API uint64 WindowStyle();
 	static APPLICATIONCORE_API class FOutputDeviceConsole* CreateConsoleOutputDevice();
 	static APPLICATIONCORE_API class FOutputDeviceError* GetErrorOutputDevice();
 	static APPLICATIONCORE_API class FFeedbackContext* GetFeedbackContext();
diff --git a/Engine/Source/Runtime/ApplicationCore/Public/Linux/LinuxWindow.h b/Engine/Source/Runtime/ApplicationCore/Public/Linux/LinuxWindow.h
--- a/Engine/Source/Runtime/ApplicationCore/Public/Linux/LinuxWindow.h
+++ b/Engine/Source/Runtime/ApplicationCore/Public/Linux/LinuxWindow.h
@@ -7,7 +7,7 @@
 #include "Templates/SharedPointer.h"
 #include "GenericPlatform/GenericWindow.h"
 
-#include "SDL.h"
+#include <SDL3/SDL.h>
 
 class FLinuxApplication;
 struct FGenericWindowDefinition;
diff --git a/Engine/Source/Runtime/InputCore/InputCore.Build.cs b/Engine/Source/Runtime/InputCore/InputCore.Build.cs
--- a/Engine/Source/Runtime/InputCore/InputCore.Build.cs
+++ b/Engine/Source/Runtime/InputCore/InputCore.Build.cs
@@ -14,7 +14,7 @@ public class InputCore : ModuleRules
 		}
 		if(Target.IsInPlatformGroup(UnrealPlatformGroup.Linux))
 		{
-			PrivateIncludePathModuleNames.Add("SDL2");
+			PrivateIncludePathModuleNames.Add("SDL3");
 		}
 		CppCompileWarningSettings.UnsafeTypeCastWarningLevel = WarningLevel.Error;
 	}
diff --git a/Engine/Source/Runtime/InputCore/Private/Linux/LinuxPlatformInput.cpp b/Engine/Source/Runtime/InputCore/Private/Linux/LinuxPlatformInput.cpp
--- a/Engine/Source/Runtime/InputCore/Private/Linux/LinuxPlatformInput.cpp
+++ b/Engine/Source/Runtime/InputCore/Private/Linux/LinuxPlatformInput.cpp
@@ -3,7 +3,7 @@
 #include "Linux/LinuxPlatformInput.h"
 
 THIRD_PARTY_INCLUDES_START
-	#include <SDL.h>
+	#include <SDL3/SDL.h>
 THIRD_PARTY_INCLUDES_END
 
 uint32 FLinuxPlatformInput::GetKeyMap( uint32* KeyCodes, FString* KeyNames, uint32 MaxMappings )
diff --git a/Engine/Source/Runtime/Linux/AudioMixerSDL/AudioMixerSDL.Build.cs b/Engine/Source/Runtime/Linux/AudioMixerSDL/AudioMixerSDL.Build.cs
--- a/Engine/Source/Runtime/Linux/AudioMixerSDL/AudioMixerSDL.Build.cs
+++ b/Engine/Source/Runtime/Linux/AudioMixerSDL/AudioMixerSDL.Build.cs
@@ -34,7 +34,7 @@ public class AudioMixerSDL : ModuleRules
 			"UEOgg",
 			"Vorbis",
 			"VorbisFile",
-			"SDL2"
+			"SDL3"
 			);
 	}
 }
diff --git a/Engine/Source/Runtime/Linux/AudioMixerSDL/Private/AudioMixerPlatformSDL.cpp b/Engine/Source/Runtime/Linux/AudioMixerSDL/Private/AudioMixerPlatformSDL.cpp
--- a/Engine/Source/Runtime/Linux/AudioMixerSDL/Private/AudioMixerPlatformSDL.cpp
+++ b/Engine/Source/Runtime/Linux/AudioMixerSDL/Private/AudioMixerPlatformSDL.cpp
@@ -18,15 +18,27 @@ DEFINE_LOG_CATEGORY(LogAudioMixerSDL);
 namespace Audio
 {
 	// Static callback function used in SDL
-	static void OnBufferEnd(void* BufferContext, uint8* OutputBuffer, int32 OutputBufferLength)
+
+	static void SDLCALL OnBufferEndWrapper(void* UserData, SDL_AudioStream* Stream, int AdditionalAmount, int TotalAmount)
 	{
-		check(BufferContext);
-		FMixerPlatformSDL* MixerPlatform = (FMixerPlatformSDL*)BufferContext;
-		MixerPlatform->HandleOnBufferEnd(OutputBuffer, OutputBufferLength);
+		if (AdditionalAmount > 0)
+		{
+			FMixerPlatformSDL* MixerPlatform = (FMixerPlatformSDL*)UserData;
+			int OutputBufferByteLength = MixerPlatform->GetOutputBufferByteLength();
+
+			Uint8* Data = SDL_stack_alloc(Uint8, OutputBufferByteLength);
+			if (Data)
+			{
+				MixerPlatform->HandleOnBufferEnd(Data, OutputBufferByteLength);
+				SDL_PutAudioStreamData(Stream, Data, OutputBufferByteLength);
+				SDL_stack_free(data);
+			}
+		}
 	}
 
 	FMixerPlatformSDL::FMixerPlatformSDL()
 		: AudioDeviceID(INDEX_NONE)
+		, AudioStream(nullptr)		
 		, OutputBuffer(nullptr)
 		, OutputBufferByteLength(0)
 		, bSuspended(false)
@@ -97,64 +109,58 @@ namespace Audio
 	{
 		if (!bInitialized)
 		{
-			UE_LOG(LogAudioMixerSDL, Error, TEXT("SDL2 Audio is not initialized."));
+			UE_LOG(LogAudioMixerSDL, Error, TEXT("SDL3 Audio is not initialized."));
 			return false;
 		}
 
-		SDL_bool IsCapture = SDL_FALSE;
-		OutNumOutputDevices = SDL_GetNumAudioDevices(IsCapture);
+		int NumDevices = 0;
+		SDL_AudioDeviceID* Devices = SDL_GetAudioPlaybackDevices(&NumDevices);
+		OutNumOutputDevices = NumDevices;
+		SDL_free(Devices);
 		return true;
 	}
 
-	bool FMixerPlatformSDL::GetOutputDeviceInfo(const uint32 InDeviceIndex, FAudioPlatformDeviceInfo& OutInfo)
+	bool FMixerPlatformSDL::GetOutputDeviceInfo(const uint32 InDeviceID, FAudioPlatformDeviceInfo& OutInfo)
 	{
-		// To figure out the output device info, attempt to init at 7.1, and 48k.
-		// SDL_OpenAudioDevice will attempt open the audio device with that spec but return what it actually used. We'll report that in OutInfo
 		FAudioPlatformSettings PlatformSettings = GetPlatformSettings();
-		SDL_AudioSpec DesiredSpec;
-		DesiredSpec.freq = PlatformSettings.SampleRate;
+		AudioSpec  = { GetPlatformAudioFormat(), GetPlatformChannels(), PlatformSettings.SampleRate };
 
-		DesiredSpec.format = GetPlatformAudioFormat();
-		DesiredSpec.channels = GetPlatformChannels();
+		FString DeviceName;
+		int NumDevices = 0;
+		SDL_AudioDeviceID* Devices = SDL_GetAudioPlaybackDevices(&NumDevices);
 
-		DesiredSpec.samples = PlatformSettings.CallbackBufferFrameSize;
-		DesiredSpec.callback = OnBufferEnd;
-		DesiredSpec.userdata = (void*)this;
+		if (NumDevices)
+		{
+			SDL_AudioDeviceID DevId;
 
-		// It's not possible with SDL to tell whether a given index is default. It only supports directly opening a device handle by passing in a nullptr to SDL_OpenAudioDevice.
-		OutInfo.bIsSystemDefault = false;
+			if (InDeviceID != AUDIO_MIXER_DEFAULT_DEVICE_INDEX)
+			{
+				DevId = Devices[InDeviceID];
+			}
+			else
+			{
+				DevId = SDL_AUDIO_DEVICE_DEFAULT_PLAYBACK;
+			}
 
-		const char* AudioDeviceName = nullptr;
-		FString DeviceName;
+			const char* AudioDeviceName = nullptr;
 
-		if (InDeviceIndex != AUDIO_MIXER_DEFAULT_DEVICE_INDEX)
-		{
-			AudioDeviceName = SDL_GetAudioDeviceName(InDeviceIndex, SDL_FALSE);
+			AudioDeviceName = SDL_GetAudioDeviceName(DevId);
 			DeviceName = ANSI_TO_TCHAR(AudioDeviceName);
-		}
-		else
-		{
-			DeviceName = TEXT("Default Audio Device");
-		}
 
-		SDL_AudioSpec ActualSpec;
-		SDL_AudioDeviceID TempAudioDeviceID = SDL_OpenAudioDevice(AudioDeviceName, SDL_FALSE, &DesiredSpec, &ActualSpec, SDL_AUDIO_ALLOW_CHANNELS_CHANGE);
-		if (!TempAudioDeviceID)
-		{
-			const char* ErrorText = SDL_GetError();
-			UE_LOG(LogAudioMixerSDL, Error, TEXT("%s"), ANSI_TO_TCHAR(ErrorText));
+			SDL_free(Devices);
+			Devices = nullptr;
 		}
 
-		// Name and Id are the same for SDL
+
+		// opening a physical device no longer changes the spec we pass in, so there's no reason to bother with any of the previous code.
+		// Just take the device instance we have and return exactly what we asked for.  The backend automatically takes care of making sure
+		// the physical device is actually set to a capable sample rate, etc to accommodate our requested format
 		OutInfo.DeviceId = DeviceName;
-		OutInfo.Name = OutInfo.DeviceId;
-		OutInfo.SampleRate = ActualSpec.freq;
-		
+		OutInfo.Name = DeviceName;
+		OutInfo.SampleRate = PlatformSettings.SampleRate;	// sample rate conversion now happens in the SDL audio code, so our code can send whatever we like
 		OutInfo.Format = GetAudioStreamFormat();
-
-		ensure(ActualSpec.channels <= AUDIO_MIXER_MAX_OUTPUT_CHANNELS);
-		OutInfo.NumChannels = FMath::Min<int32>(static_cast<int32>(ActualSpec.channels), AUDIO_MIXER_MAX_OUTPUT_CHANNELS);
-
+		OutInfo.NumChannels = FMath::Min<int32>(static_cast<int32>(GetPlatformChannels()), AUDIO_MIXER_MAX_OUTPUT_CHANNELS);
+		
 		// Assume default channel map order, SDL doesn't support us querying it directly
 		OutInfo.OutputChannelArray.Reset();
 		for (int32 i = 0; i < OutInfo.NumChannels; ++i)
@@ -162,8 +168,6 @@ namespace Audio
 			OutInfo.OutputChannelArray.Add(EAudioMixerChannel::Type(i));
 		}
 
-		SDL_CloseAudioDevice(TempAudioDeviceID);
-
 		return true;
 	}
 
@@ -181,14 +185,26 @@ namespace Audio
 			return false;
 		}
 
-		uint32 NumOutputDevices = 0;
-		if (!GetNumOutputDevices(NumOutputDevices))
+		int NumOutputDevices = 0;
+		SDL_AudioDeviceID* Devices = SDL_GetAudioPlaybackDevices(&NumOutputDevices);
+		if (!Devices || NumOutputDevices == 0)
 		{
 			return false;
 		}
-		
+
 		OpenStreamParams = Params;
 
+		if (OpenStreamParams.OutputDeviceIndex!= AUDIO_MIXER_DEFAULT_DEVICE_INDEX)
+		{
+			AudioDeviceID =  Devices[OpenStreamParams.OutputDeviceIndex];
+		}
+		else
+		{
+			AudioDeviceID =  SDL_AUDIO_DEVICE_DEFAULT_PLAYBACK;
+		}
+		SDL_free(Devices);
+		Devices = nullptr;
+
 		AudioStreamInfo.Reset();
 		AudioStreamInfo.AudioMixer = OpenStreamParams.AudioMixer;
 		AudioStreamInfo.NumOutputFrames = OpenStreamParams.NumFrames;
@@ -205,46 +221,35 @@ namespace Audio
 
 			AudioStreamInfo.OutputDeviceIndex = OpenStreamParams.OutputDeviceIndex;
 
-			AudioSpecPrefered.format = GetPlatformAudioFormat();
-			AudioSpecPrefered.freq = Params.SampleRate;
-			AudioSpecPrefered.channels = AudioStreamInfo.DeviceInfo.NumChannels;
-			AudioSpecPrefered.samples = OpenStreamParams.NumFrames;
-			AudioSpecPrefered.callback = OnBufferEnd;
-			AudioSpecPrefered.userdata = (void*)this;
+			AudioSpec.format = GetPlatformAudioFormat();
+			AudioSpec.freq = Params.SampleRate;
+			AudioSpec.channels = AudioStreamInfo.DeviceInfo.NumChannels;
 
 			const char* DeviceName = nullptr;
-			if (OpenStreamParams.OutputDeviceIndex != AUDIO_MIXER_DEFAULT_DEVICE_INDEX && OpenStreamParams.OutputDeviceIndex < (uint32)SDL_GetNumAudioDevices(0))
-			{
-				DeviceName = SDL_GetAudioDeviceName(OpenStreamParams.OutputDeviceIndex, 0);
-			}
+			DeviceName = SDL_GetAudioDeviceName(AudioDeviceID);
 
 			// only the default device can be overriden
 			FString CurrentDeviceNameString = GetCurrentDeviceName();
 			if (OpenStreamParams.OutputDeviceIndex != AUDIO_MIXER_DEFAULT_DEVICE_INDEX || CurrentDeviceNameString.Len() <= 0)
 			{
 				UE_LOG(LogAudioMixerSDL, Log, TEXT("Opening %s audio device (device index %d)"), DeviceName ? ANSI_TO_TCHAR(DeviceName) : TEXT("default"), OpenStreamParams.OutputDeviceIndex);
-				AudioDeviceID = SDL_OpenAudioDevice(DeviceName, 0, &AudioSpecPrefered, &AudioSpecReceived, 0);
 			}
 			else
 			{
 				UE_LOG(LogAudioMixerSDL, Log, TEXT("Opening overridden '%s' audio device (device index %d)"), *CurrentDeviceNameString, OpenStreamParams.OutputDeviceIndex);
-				AudioDeviceID = SDL_OpenAudioDevice(TCHAR_TO_ANSI(*CurrentDeviceNameString), 0, &AudioSpecPrefered, &AudioSpecReceived, 0);
 			}
+	
+			AudioStream = SDL_OpenAudioDeviceStream(AudioDeviceID, &AudioSpec, OnBufferEndWrapper, (void*)this);
 
-			if (!AudioDeviceID)
+			if (!AudioStream)
 			{
 				const char* ErrorText = SDL_GetError();
 				UE_LOG(LogAudioMixerSDL, Error, TEXT("%s"), ANSI_TO_TCHAR(ErrorText));
 				return false;
 			}
 
-			// Make sure our device initialized as expected, we should have already filtered this out before this point.
-			check(AudioSpecReceived.channels == AudioSpecPrefered.channels);
-			check(AudioSpecReceived.samples == OpenStreamParams.NumFrames);
-
 			// Compute the expected output byte length
 			OutputBufferByteLength = OpenStreamParams.NumFrames * AudioStreamInfo.DeviceInfo.NumChannels * GetAudioStreamChannelSize();
-			check(OutputBufferByteLength == AudioSpecReceived.size);
 		}
 		else
 		{
@@ -253,7 +258,8 @@ namespace Audio
 			AudioStreamInfo.DeviceInfo.NumChannels = 2;			
 			AudioStreamInfo.DeviceInfo.Format = EAudioMixerStreamDataFormat::Float;
 		}
-		
+
+		bInitialized = true;
 		AudioStreamInfo.StreamState = EAudioOutputStreamState::Open;
 
 		return true;
@@ -271,12 +277,12 @@ namespace Audio
 			return false;
 		}
 
-		if (AudioDeviceID != INDEX_NONE)
+		if (AudioStream != nullptr)
 		{
 			FScopeLock ScopedLock(&OutputBufferMutex);
 
-			SDL_CloseAudioDevice(AudioDeviceID);
-
+			SDL_DestroyAudioStream(AudioStream);
+			AudioStream = nullptr;
 			OutputBuffer = nullptr;
 			OutputBufferByteLength = 0;
 		}
@@ -295,11 +301,10 @@ namespace Audio
 		// Start generating audio
 		BeginGeneratingAudio();
 
-		// If we're not using the Null renderer AudioDeviceID will be set to something other than INDEX_NONE.
-		if (AudioDeviceID != INDEX_NONE)
+		if (AudioStream != nullptr)
 		{
 			// Unpause audio device to start it rendering audio
-			SDL_PauseAudioDevice(AudioDeviceID, 0);
+			SDL_ResumeAudioStreamDevice(AudioStream);
 		}
 		else
 		{
@@ -314,12 +319,17 @@ namespace Audio
 
 	bool FMixerPlatformSDL::StopAudioStream()
 	{
+		if (!bInitialized)
+		{
+			return false;
+		}
+
 		if (AudioStreamInfo.StreamState != EAudioOutputStreamState::Stopped && AudioStreamInfo.StreamState != EAudioOutputStreamState::Closed)
 		{
-			if (AudioDeviceID != INDEX_NONE)
+			if (AudioStream != nullptr)
 			{
 				// Pause the audio device
-				SDL_PauseAudioDevice(AudioDeviceID, 1);
+				SDL_PauseAudioStreamDevice(AudioStream);
 			}
 			else
 			{
@@ -377,9 +387,13 @@ namespace Audio
 	{
 		if (bSuspended)
 		{
-			SDL_UnlockAudioDevice(AudioDeviceID);
-			UE_LOG(LogAudioMixerSDL, Display, TEXT("Resuming Audio"));
-			bSuspended = false;
+			if (AudioStream != nullptr)
+			{
+				SDL_ResumeAudioStreamDevice(AudioStream);
+				UE_LOG(LogAudioMixerSDL, Display, TEXT("Resuming Audio"));
+				bSuspended = false;
+	
+			}
 		}
 	}
 
@@ -387,9 +401,12 @@ namespace Audio
 	{
 		if (!bSuspended)
 		{
-			SDL_LockAudioDevice(AudioDeviceID);
-			UE_LOG(LogAudioMixerSDL, Display, TEXT("Suspending Audio"));
-			bSuspended = true;
+			if (AudioStream != nullptr)
+			{
+				SDL_PauseAudioStreamDevice(AudioStream);
+				UE_LOG(LogAudioMixerSDL, Display, TEXT("Suspending Audio"));
+				bSuspended = true;
+			}
 		}
 	}
 
diff --git a/Engine/Source/Runtime/Linux/AudioMixerSDL/Private/AudioMixerPlatformSDL.h b/Engine/Source/Runtime/Linux/AudioMixerSDL/Private/AudioMixerPlatformSDL.h
--- a/Engine/Source/Runtime/Linux/AudioMixerSDL/Private/AudioMixerPlatformSDL.h
+++ b/Engine/Source/Runtime/Linux/AudioMixerSDL/Private/AudioMixerPlatformSDL.h
@@ -5,8 +5,8 @@
 #include "AudioMixer.h"
 
 THIRD_PARTY_INCLUDES_START
-#include "SDL.h"
-#include "SDL_audio.h"
+#include <SDL3/SDL.h>
+#include <SDL3/SDL_audio.h>
 THIRD_PARTY_INCLUDES_END
 
 namespace Audio
@@ -21,7 +21,7 @@ namespace Audio
 		virtual ~FMixerPlatformSDL();
 
 		//~ Begin IAudioMixerPlatformInterface Interface
-		FString GetPlatformApi() const override { return TEXT("SDL2"); }
+		FString GetPlatformApi() const override { return TEXT("SDL3"); }
 		bool InitializeHardware() override;
 		bool TeardownHardware() override;
 		bool IsInitialized() const override;
@@ -42,19 +42,20 @@ namespace Audio
 
 		void HandleOnBufferEnd(uint8* InOutputBuffer, int32 InOutputBufferLength);
 
-		SDL_AudioFormat GetPlatformAudioFormat() { return AUDIO_F32; }
+		SDL_AudioFormat GetPlatformAudioFormat() { return SDL_AUDIO_F32LE; }
 		Uint8 GetPlatformChannels() { return 6; }
 		EAudioMixerStreamDataFormat::Type GetAudioStreamFormat() { return EAudioMixerStreamDataFormat::Float; }
 		int32 GetAudioStreamChannelSize() { return sizeof(float); }
-
+		int32 GetOutputBufferByteLength() { return OutputBufferByteLength; }
+		
 	protected:
 		FCriticalSection OutputBufferMutex;
 
 	private:
 
 		SDL_AudioDeviceID AudioDeviceID;
-		SDL_AudioSpec AudioSpecPrefered;
-		SDL_AudioSpec AudioSpecReceived;
+		SDL_AudioSpec AudioSpec;
+		SDL_AudioStream *AudioStream;
 
 		uint8* OutputBuffer;
 		int32 OutputBufferByteLength;
diff --git a/Engine/Source/Runtime/Online/Voice/Private/Linux/VoiceModuleLinux.cpp b/Engine/Source/Runtime/Online/Voice/Private/Linux/VoiceModuleLinux.cpp
--- a/Engine/Source/Runtime/Online/Voice/Private/Linux/VoiceModuleLinux.cpp
+++ b/Engine/Source/Runtime/Online/Voice/Private/Linux/VoiceModuleLinux.cpp
@@ -25,8 +25,8 @@ IVoiceCapture* CreateVoiceCaptureObject(const FString& DeviceName, int32 SampleR
 #include "VoiceCodecOpus.h"
 #include "Voice.h"
 
-#include "SDL.h"
-#include "SDL_audio.h"
+#include <SDL3/SDL.h>
+#include <SDL3/SDL_audio.h>
 
 class FVoiceCaptureSDL : public IVoiceCapture
 {
@@ -50,8 +50,7 @@ public:
 	virtual void DumpState() const override;
 
 private:
-	SDL_AudioDeviceID AudioDeviceID;
-	SDL_AudioSpec AudioSpec;
+	SDL_AudioStream* AudioStream;
 	EVoiceCaptureState::Type VoiceCaptureState;
 
 	/** Array to be used as a ring buffer */
@@ -61,7 +60,7 @@ private:
 };
 
 FVoiceCaptureSDL::FVoiceCaptureSDL()
-	: AudioDeviceID(0)
+	: AudioStream(nullptr)
 	, VoiceCaptureState(EVoiceCaptureState::UnInitialized)
 	, AudioBufferWritePos(0)
 	, AudioBufferAvailableData(0)
@@ -109,6 +108,20 @@ void FVoiceCaptureSDL::ReadData(Uint8* Stream, int Length)
 	AudioBufferAvailableData = FMath::Min(AudioBufferAvailableData + Length, AudioBuffer.Num());
 }
 
+static void SDLCALL VoiceCallbackWrapper(void *UserData, SDL_AudioStream *AudioStream, int AdditionalAmount, int TotalAmount)
+{
+	if (AdditionalAmount > 0)
+	{
+		Uint8 *Data = SDL_stack_alloc(Uint8, AdditionalAmount);
+		if(Data)
+		{
+			FVoiceCaptureSDL::AudioCallback(UserData, Data, AdditionalAmount);
+			SDL_PutAudioStreamData(AudioStream, Data, AdditionalAmount);
+			SDL_stack_free(Data);
+		}
+	}
+}
+
 bool FVoiceCaptureSDL::Init(const FString& DeviceName, int32 SampleRate, int32 NumChannels)
 {
 	check(VoiceCaptureState == EVoiceCaptureState::UnInitialized);
@@ -125,24 +138,37 @@ bool FVoiceCaptureSDL::Init(const FString& DeviceName, int32 SampleRate, int32 N
 		return false;
 	}
 
-	SDL_AudioSpec DesiredAudioSpec;
-	SDL_memset(&DesiredAudioSpec, 0, sizeof(DesiredAudioSpec));
-	DesiredAudioSpec.freq = SampleRate;
-	DesiredAudioSpec.format = AUDIO_S16;
-	DesiredAudioSpec.channels = NumChannels;
-	DesiredAudioSpec.samples = static_cast<Uint16>(SampleRate / 50); // 20ms worth of samples, may be worth making this configurable
-	DesiredAudioSpec.callback = &FVoiceCaptureSDL::AudioCallback;
-	DesiredAudioSpec.userdata = this;
-
-	AudioDeviceID = SDL_OpenAudioDevice(nullptr, 1, &DesiredAudioSpec, &AudioSpec, 0);
-	if (AudioDeviceID == 0)
+
+	const SDL_AudioSpec DesiredAudioSpec = { SDL_AUDIO_S16LE, NumChannels, SampleRate };
+	AudioStream = SDL_OpenAudioDeviceStream(SDL_AUDIO_DEVICE_DEFAULT_RECORDING, &DesiredAudioSpec, &VoiceCallbackWrapper, this);
+	if (AudioStream == nullptr)
 	{
-		UE_LOG(LogVoiceCapture, Error, TEXT("Unable to open Audio Device: %s"), ANSI_TO_TCHAR(SDL_GetError()));
+		UE_LOG(LogVoiceCapture, Error, TEXT("Unable to open Audio Stream: %s"), ANSI_TO_TCHAR(SDL_GetError()));
 		return false;
 	}
 
-	int32 BytesPerSample = AudioSpec.size / AudioSpec.samples;
-	AudioBuffer.SetNum(SampleRate * BytesPerSample); // allow up to 1s to be buffered, may be worth making this configurable
+	int SampleSizeInBytes;
+	
+	switch(DesiredAudioSpec.format)
+	{
+		case SDL_AUDIO_U8:	// intentional fallthrough
+		case SDL_AUDIO_S8:
+			SampleSizeInBytes = 1;
+			break;
+		case SDL_AUDIO_S16LE:	// intentional fallthrough
+		case SDL_AUDIO_S16BE:
+			SampleSizeInBytes = 2;
+			break;
+		case SDL_AUDIO_S32LE:	// intentional fallthrough
+		case SDL_AUDIO_S32BE:
+		case SDL_AUDIO_F32LE:
+		case SDL_AUDIO_F32BE:
+			SampleSizeInBytes = 4;
+			break;
+		default: SampleSizeInBytes = -1;
+			 break;
+	}
+	AudioBuffer.SetNum(DesiredAudioSpec.channels * DesiredAudioSpec.freq * SampleSizeInBytes);
 
 	return true;
 }
@@ -159,10 +185,10 @@ void FVoiceCaptureSDL::Shutdown()
 			Stop();
 			// intentional fall-through
 		case EVoiceCaptureState::NotCapturing:
-			if (AudioDeviceID != 0)
+			if (AudioStream != nullptr)
 			{
-				SDL_CloseAudioDevice(AudioDeviceID);
-				AudioDeviceID = 0;
+				SDL_DestroyAudioStream(AudioStream);
+				AudioStream = nullptr;
 			}
 			VoiceCaptureState = EVoiceCaptureState::UnInitialized;
 			break;
@@ -180,8 +206,8 @@ bool FVoiceCaptureSDL::Start()
 	AudioBufferWritePos = 0;
 	AudioBufferAvailableData = 0;
 
-	SDL_PauseAudioDevice(AudioDeviceID, 0);
-	if (SDL_GetAudioDeviceStatus(AudioDeviceID) == SDL_AUDIO_PLAYING)
+	SDL_ResumeAudioDevice(SDL_GetAudioStreamDevice(AudioStream));
+	if (SDL_AudioDevicePaused(SDL_GetAudioStreamDevice(AudioStream)) == false)
 	{
 		VoiceCaptureState = EVoiceCaptureState::Ok;
 		return true;
@@ -196,7 +222,7 @@ bool FVoiceCaptureSDL::Start()
 void FVoiceCaptureSDL::Stop()
 {
 	check(IsCapturing());
-	SDL_PauseAudioDevice(AudioDeviceID, 1);
+	SDL_PauseAudioDevice(SDL_GetAudioStreamDevice(AudioStream));
 
 	VoiceCaptureState = EVoiceCaptureState::NotCapturing;
 	AudioBufferWritePos = 0;
@@ -211,7 +237,7 @@ bool FVoiceCaptureSDL::ChangeDevice(const FString& DeviceName, int32 SampleRate,
 
 bool FVoiceCaptureSDL::IsCapturing()
 {
-	return AudioDeviceID > 0 && VoiceCaptureState > EVoiceCaptureState::NotCapturing;
+	return AudioStream != nullptr && VoiceCaptureState > EVoiceCaptureState::NotCapturing;
 }
 
 EVoiceCaptureState::Type FVoiceCaptureSDL::GetCaptureState(uint32& OutAvailableVoiceData) const
@@ -219,9 +245,9 @@ EVoiceCaptureState::Type FVoiceCaptureSDL::GetCaptureState(uint32& OutAvailableV
 	if (VoiceCaptureState == EVoiceCaptureState::Ok)
 	{
 		// may be better to use atomics to avoid locking here
-		SDL_LockAudioDevice(AudioDeviceID);
+		SDL_LockAudioStream(AudioStream);
 		OutAvailableVoiceData = AudioBufferAvailableData;
-		SDL_UnlockAudioDevice(AudioDeviceID);
+		SDL_UnlockAudioStream(AudioStream);
 	}
 	else
 	{
@@ -238,13 +264,7 @@ EVoiceCaptureState::Type FVoiceCaptureSDL::GetVoiceData(uint8* OutVoiceBuffer, u
 
 	if (VoiceCaptureState == EVoiceCaptureState::Ok)
 	{
-		if (SDL_GetAudioDeviceStatus(AudioDeviceID) == SDL_AUDIO_STOPPED)
-		{
-			VoiceCaptureState = EVoiceCaptureState::Error;
-			return EVoiceCaptureState::Error;
-		}
-
-		SDL_LockAudioDevice(AudioDeviceID);
+		SDL_LockAudioStream(AudioStream);
 
 		uint32 AvailableVoiceData = AudioBufferAvailableData;
 
@@ -275,7 +295,7 @@ EVoiceCaptureState::Type FVoiceCaptureSDL::GetVoiceData(uint8* OutVoiceBuffer, u
 			AudioBufferAvailableData -= OutAvailableVoiceData;
 		}
 
-		SDL_UnlockAudioDevice(AudioDeviceID);
+		SDL_UnlockAudioStream(AudioStream);
 	}
 	return ReturnState;
 }
diff --git a/Engine/Source/Runtime/Online/Voice/Voice.Build.cs b/Engine/Source/Runtime/Online/Voice/Voice.Build.cs
--- a/Engine/Source/Runtime/Online/Voice/Voice.Build.cs
+++ b/Engine/Source/Runtime/Online/Voice/Voice.Build.cs
@@ -59,7 +59,7 @@ public class Voice : ModuleRules
 		}
 		else if (Target.IsInPlatformGroup(UnrealPlatformGroup.Linux) && !bDontNeedCapture)
 		{
-			AddEngineThirdPartyPrivateStaticDependencies(Target, "SDL2");
+			AddEngineThirdPartyPrivateStaticDependencies(Target, "SDL3");
 		}
 
 		AddEngineThirdPartyPrivateStaticDependencies(Target, "libOpus");
diff --git a/Engine/Source/Runtime/OpenGLDrv/OpenGLDrv.Build.cs b/Engine/Source/Runtime/OpenGLDrv/OpenGLDrv.Build.cs
--- a/Engine/Source/Runtime/OpenGLDrv/OpenGLDrv.Build.cs
+++ b/Engine/Source/Runtime/OpenGLDrv/OpenGLDrv.Build.cs
@@ -31,7 +31,7 @@ public class OpenGLDrv : ModuleRules
 
 		if (Target.IsInPlatformGroup(UnrealPlatformGroup.Linux))
 		{
-			AddEngineThirdPartyPrivateStaticDependencies(Target, "SDL2");
+			AddEngineThirdPartyPrivateStaticDependencies(Target, "SDL3");
 		}
 
 		if (Target.Platform == UnrealTargetPlatform.Android)
diff --git a/Engine/Source/Runtime/OpenGLDrv/Private/Linux/LinuxOpenGLPlatform.cpp b/Engine/Source/Runtime/OpenGLDrv/Private/Linux/LinuxOpenGLPlatform.cpp
--- a/Engine/Source/Runtime/OpenGLDrv/Private/Linux/LinuxOpenGLPlatform.cpp
+++ b/Engine/Source/Runtime/OpenGLDrv/Private/Linux/LinuxOpenGLPlatform.cpp
@@ -4,7 +4,7 @@
 
 #include "Misc/ScopeLock.h"
 #include "OpenGLDrv.h"
-#include "SDL.h"
+#include <SDL3/SDL.h>
 #include "OpenGLDrvPrivate.h"
 #include "ComponentReregisterContext.h"
 #include "Linux/LinuxPlatformApplicationMisc.h"
@@ -102,8 +102,8 @@ void Linux_PlatformCreateDummyGLWindow(FPlatformOpenGLContext* OutContext)
 
 	// Create a dummy window.
 	SDL_HWindow DummyWindow = SDL_CreateWindow(NULL,
-		0, 0, 1, 1,
-		SDL_WINDOW_OPENGL | SDL_WINDOW_BORDERLESS | SDL_WINDOW_HIDDEN | SDL_WINDOW_SKIP_TASKBAR);
+		1, 1,
+		SDL_WINDOW_OPENGL | SDL_WINDOW_BORDERLESS | SDL_WINDOW_HIDDEN | SDL_WINDOW_UTILITY );
 	if (DummyWindow == nullptr)
 	{
 		FString SdlError(UTF8_TO_TCHAR(SDL_GetError()));
@@ -294,7 +294,7 @@ void PlatformReleaseOpenGLContext(FPlatformOpenGLDevice* Device, FPlatformOpenGL
 			}
 		}
 
-		SDL_GL_DeleteContext(Context->hGLContext);
+		SDL_GL_DestroyContext(Context->hGLContext);
 		Context->hGLContext = NULL;
 	}
 
@@ -386,7 +386,7 @@ bool PlatformBlitToViewport(IRHICommandContext& RHICmdContext,
 
 		int WinW = 0;
 		int WinH = 0;
-		SDL_GL_GetDrawableSize(Context->hWnd, &WinW, &WinH);
+		SDL_GetWindowSizeInPixels(Context->hWnd, &WinW, &WinH);
 		GLenum BlitFilter;
 		GLint DestX0, DestY0, DestX1, DestY1;
 
@@ -581,18 +581,23 @@ void PlatformGetSupportedResolution(uint32& Width, uint32& Height)
 	uint32 BestHeight = 0;
 	uint32 ModeIndex = 0;
 
-	SDL_DisplayMode DisplayMode;
-	FMemory::Memzero(&DisplayMode, sizeof(DisplayMode));
-
-	while (!SDL_GetDisplayMode(0, ModeIndex++, &DisplayMode))
+	SDL_DisplayID DisplayID = SDL_GetPrimaryDisplay();
+	int NumModes = 0;
+	SDL_DisplayMode** Modes = SDL_GetFullscreenDisplayModes(DisplayID, &NumModes);
+	if (Modes)
 	{
-		bool IsEqualOrBetterWidth = FMath::Abs((int32)DisplayMode.w - (int32)Width) <= FMath::Abs((int32)BestWidth - (int32)Width);
-		bool IsEqualOrBetterHeight = FMath::Abs((int32)DisplayMode.h - (int32)Height) <= FMath::Abs((int32)BestHeight - (int32)Height);
-		if (!InitializedMode || (IsEqualOrBetterWidth && IsEqualOrBetterHeight))
+		for (int i=0; i < NumModes; ++i)
 		{
-			BestWidth = DisplayMode.w;
-			BestHeight = DisplayMode.h;
-			InitializedMode = true;
+			SDL_DisplayMode* Mode = Modes[i];
+			bool IsEqualOrBetterWidth = FMath::Abs((int32)Mode->w - (int32)Width) <= FMath::Abs((int32)BestWidth - (int32)Width);
+			bool IsEqualOrBetterHeight = FMath::Abs((int32)Mode->h - (int32)Height) <= FMath::Abs((int32)BestHeight - (int32)Height);
+			if (!InitializedMode || (IsEqualOrBetterWidth && IsEqualOrBetterHeight))
+			{
+				BestWidth = Mode->w;
+				BestHeight = Mode->h;
+				InitializedMode = true;
+			}
+
 		}
 	}
 	check(InitializedMode);
@@ -623,54 +628,59 @@ bool PlatformGetAvailableResolutions(FScreenResolutionArray& Resolutions, bool b
 		MaxAllowableRefreshRate = 10480;
 	}
 
-	uint32 ModeIndex = 0;
-	SDL_DisplayMode DisplayMode;
-	FMemory::Memzero(&DisplayMode, sizeof(SDL_DisplayMode));
 
-	while (!SDL_GetDisplayMode(0, ModeIndex++, &DisplayMode))
+	SDL_DisplayID Display = SDL_GetPrimaryDisplay();
+	int NumModes = 0;
+	SDL_DisplayMode** Modes = SDL_GetFullscreenDisplayModes(Display, &NumModes);
+	if (Modes)
 	{
-		if (((int32)DisplayMode.w >= MinAllowableResolutionX) &&
-			((int32)DisplayMode.w <= MaxAllowableResolutionX) &&
-			((int32)DisplayMode.h >= MinAllowableResolutionY) &&
-			((int32)DisplayMode.h <= MaxAllowableResolutionY)
-			)
+		for (int32 i = 0; i < NumModes; i++)
 		{
-			bool bAddIt = true;
-			if (bIgnoreRefreshRate == false)
+			SDL_DisplayMode* Mode = Modes[i];
+			if (((int32)Mode->w >= MinAllowableResolutionX) &&
+				((int32)Mode->w <= MaxAllowableResolutionX) &&
+				((int32)Mode->h >= MinAllowableResolutionY) &&
+				((int32)Mode->h <= MaxAllowableResolutionY)
+			   	)
 			{
-				if (((int32)DisplayMode.refresh_rate < MinAllowableRefreshRate) ||
-					((int32)DisplayMode.refresh_rate > MaxAllowableRefreshRate))
+				bool bAddIt = true;
+				if (bIgnoreRefreshRate == false)
 				{
-					continue;
+					if (((int32)Mode->refresh_rate < MinAllowableRefreshRate) ||
+						((int32)Mode->refresh_rate > MaxAllowableRefreshRate))
+					{
+						continue;
+					}
 				}
-			}
-			else
-			{
-				// See if it is in the list already
-				for (int32 CheckIndex = 0; CheckIndex < Resolutions.Num(); CheckIndex++)
+				else
 				{
-					FScreenResolutionRHI& CheckResolution = Resolutions[CheckIndex];
-					if ((CheckResolution.Width == DisplayMode.w) &&
-						(CheckResolution.Height == DisplayMode.h))
+					// See if it is in the list already
+					for (int32 CheckIndex = 0; CheckIndex < Resolutions.Num(); CheckIndex++)
 					{
-						// Already in the list...
-						bAddIt = false;
-						break;
+						FScreenResolutionRHI& CheckResolution = Resolutions[CheckIndex];
+						if ((CheckResolution.Width == Mode->w) &&
+							(CheckResolution.Height == Mode->h))
+						{
+							// Already in the list...
+							// bAddIt = false;
+							break;
+						}
 					}
 				}
-			}
 
-			if (bAddIt)
-			{
-				// Add the mode to the list
-				int32 Temp2Index = Resolutions.AddZeroed();
-				FScreenResolutionRHI& ScreenResolution = Resolutions[Temp2Index];
+				if (bAddIt)
+				{
+					// Add the mode to the list
+					int32 Temp2Index = Resolutions.AddZeroed();
+					FScreenResolutionRHI& ScreenResolution = Resolutions[Temp2Index];
 
-				ScreenResolution.Width = DisplayMode.w;
-				ScreenResolution.Height = DisplayMode.h;
-				ScreenResolution.RefreshRate = DisplayMode.refresh_rate;
+					ScreenResolution.Width = Mode->w;
+					ScreenResolution.Height = Mode->h;
+					ScreenResolution.RefreshRate = Mode->refresh_rate;
+				}
 			}
 		}
+		SDL_free(Modes);
 	}
 
 	return true;
@@ -694,7 +704,7 @@ bool PlatformInitOpenGL()
 	}
 
 #if DO_CHECK
-	uint32 InitializedSubsystems = SDL_WasInit(SDL_INIT_EVERYTHING);
+	uint32 InitializedSubsystems = SDL_WasInit(SDL_INIT_VIDEO);
 	check(InitializedSubsystems & SDL_INIT_VIDEO);
 #endif // DO_CHECK
 
@@ -783,7 +793,7 @@ bool PlatformInitOpenGL()
 		if (DummyContext.hGLContext)
 		{
 			Linux_ContextMakeCurrent(NULL, NULL);
-			SDL_GL_DeleteContext(DummyContext.hGLContext);
+			SDL_GL_DestroyContext(DummyContext.hGLContext);
 		}
 		check(DummyContext.bReleaseWindowOnDestroy);
 		SDL_DestroyWindow(DummyContext.hWnd);
diff --git a/Engine/Source/Runtime/Slate/Private/Framework/Application/SlateApplication.cpp b/Engine/Source/Runtime/Slate/Private/Framework/Application/SlateApplication.cpp
--- a/Engine/Source/Runtime/Slate/Private/Framework/Application/SlateApplication.cpp
+++ b/Engine/Source/Runtime/Slate/Private/Framework/Application/SlateApplication.cpp
@@ -3503,7 +3503,8 @@ void FSlateApplication::ProcessReply( const FWidgetPath& CurrentEventPath, const
 		{
 			if ( PlatformApplication->IsUsingHighPrecisionMouseMode() )
 			{
-				PlatformApplication->SetHighPrecisionMouseMode(false, nullptr);
+				const TSharedRef< SWindow> Window = CurrentEventPath.GetWindow();
+				PlatformApplication->SetHighPrecisionMouseMode(false, Window->GetNativeWindow());
 				PlatformApplication->SetCapture(nullptr);
 
 #if WITH_SLATE_DEBUGGING
diff --git a/Engine/Source/Runtime/Slate/Slate.Build.cs b/Engine/Source/Runtime/Slate/Slate.Build.cs
--- a/Engine/Source/Runtime/Slate/Slate.Build.cs
+++ b/Engine/Source/Runtime/Slate/Slate.Build.cs
@@ -30,7 +30,7 @@ public class Slate : ModuleRules
 		}
 		else if (Target.IsInPlatformGroup(UnrealPlatformGroup.Linux))
 		{
-			AddEngineThirdPartyPrivateStaticDependencies(Target, "SDL2");
+			AddEngineThirdPartyPrivateStaticDependencies(Target, "SDL3");
 		}
 
 		// Add slate runtime dependencies
diff --git a/Engine/Source/Runtime/VulkanRHI/Private/Linux/VulkanLinuxPlatform.cpp b/Engine/Source/Runtime/VulkanRHI/Private/Linux/VulkanLinuxPlatform.cpp
--- a/Engine/Source/Runtime/VulkanRHI/Private/Linux/VulkanLinuxPlatform.cpp
+++ b/Engine/Source/Runtime/VulkanRHI/Private/Linux/VulkanLinuxPlatform.cpp
@@ -6,8 +6,8 @@
 #include "VulkanRayTracing.h"
 #include "VulkanExtensions.h"
 #include <dlfcn.h>
-#include <SDL.h>
-#include <SDL_vulkan.h>
+#include <SDL3/SDL.h>
+#include <SDL3/SDL_vulkan.h>
 #include "Linux/LinuxPlatformApplicationMisc.h"
 #include "Linux/LinuxPlatformProperties.h"
 
@@ -19,6 +19,7 @@ static bool GRenderOffScreen = false;
 
 void* FVulkanLinuxPlatform::VulkanLib = nullptr;
 bool FVulkanLinuxPlatform::bAttemptedLoad = false;
+bool FVulkanLinuxPlatform::bLoadedSDLVulkanLibrary = false;
 
 bool FVulkanLinuxPlatform::IsSupported()
 {
@@ -152,12 +153,15 @@ bool FVulkanLinuxPlatform::LoadVulkanLibrary()
 	UE_LOG(LogVulkanRHI, Display, TEXT("Updated LD_LIBRARY_PATH=%s"), *PreviousLibPath);
 #endif // VULKAN_HAS_DEBUGGING_ENABLED
 
+	FAnsiString FinalVulkanLoaderPath;
+		
 	if (!FParse::Param(FCommandLine::Get(), TEXT("UseLocalVulkanLoader")))
 	{
 		const ANSICHAR* LoaderFilename = "libvulkan.so.1";
 		
 		// Try to load the default libvulkan.so
 		void* GlobalVulkanLib = dlopen(LoaderFilename, RTLD_NOW | RTLD_LOCAL);
+		FinalVulkanLoaderPath = LoaderFilename;
 		
 		if (FParse::Param(FCommandLine::Get(), TEXT("UseGlobalVulkanLoader")))
 		{
@@ -191,6 +195,7 @@ bool FVulkanLinuxPlatform::LoadVulkanLibrary()
 		FAnsiString VulkanLoaderPath(UEVulkanBinariesPath);
 		VulkanLoaderPath.Append("/libvulkan.so");
 		VulkanLib = dlopen(*VulkanLoaderPath, RTLD_NOW | RTLD_LOCAL);
+		FinalVulkanLoaderPath = *VulkanLoaderPath;
 	}
 
 
@@ -204,6 +209,8 @@ bool FVulkanLinuxPlatform::LoadVulkanLibrary()
 		return false;
 	}
 	
+	bLoadedSDLVulkanLibrary = SDL_Vulkan_LoadLibrary(*FinalVulkanLoaderPath);	// SDL needs to use the same Vulkan Engine is so load it with the path we used
+	
 	bool bFoundAllEntryPoints = true;
 #define CHECK_VK_ENTRYPOINTS(Type,Func) if (VulkanDynamicAPI::Func == NULL) { bFoundAllEntryPoints = false; UE_LOG(LogRHI, Warning, TEXT("Failed to find entry point for %s"), TEXT(#Func)); }
 
@@ -215,6 +222,11 @@ bool FVulkanLinuxPlatform::LoadVulkanLibrary()
 	{
 		dlclose(VulkanLib);
 		VulkanLib = nullptr;
+		if (bLoadedSDLVulkanLibrary)
+		{
+			SDL_Vulkan_UnloadLibrary();
+			bLoadedSDLVulkanLibrary = false;
+		}
 		return false;
 	}
 
@@ -276,6 +288,12 @@ void FVulkanLinuxPlatform::FreeVulkanLibrary()
 
 		dlclose(VulkanLib);
 		VulkanLib = nullptr;
+
+		if (bLoadedSDLVulkanLibrary)
+		{
+			SDL_Vulkan_UnloadLibrary();
+			bLoadedSDLVulkanLibrary = false;
+		}
 	}
 	bAttemptedLoad = false;
 }
@@ -340,7 +358,7 @@ void FVulkanLinuxPlatform::CreateSurface(FVulkanPlatformWindowContext& WindowCon
 {
 	EnsureSDLIsInited();
 
-	if (SDL_Vulkan_CreateSurface((SDL_Window*)WindowContext.GetWindowHandle(), Instance, OutSurface) == SDL_FALSE)
+	if (SDL_Vulkan_CreateSurface((SDL_Window*)WindowContext.GetWindowHandle(), Instance, VulkanRHI::GetMemoryAllocator(nullptr), OutSurface) == false)
 	{
 		UE_LOG(LogInit, Error, TEXT("Error initializing SDL Vulkan Surface: %hs"), SDL_GetError());
 		check(0);
diff --git a/Engine/Source/Runtime/VulkanRHI/Private/Linux/VulkanLinuxPlatform.h b/Engine/Source/Runtime/VulkanRHI/Private/Linux/VulkanLinuxPlatform.h
--- a/Engine/Source/Runtime/VulkanRHI/Private/Linux/VulkanLinuxPlatform.h
+++ b/Engine/Source/Runtime/VulkanRHI/Private/Linux/VulkanLinuxPlatform.h
@@ -55,6 +55,7 @@ public:
 
 protected:
 	static void* VulkanLib;
+	static bool bLoadedSDLVulkanLibrary;
 	static bool bAttemptedLoad;
 };
 
diff --git a/Engine/Source/Runtime/VulkanRHI/VulkanRHI.Build.cs b/Engine/Source/Runtime/VulkanRHI/VulkanRHI.Build.cs
--- a/Engine/Source/Runtime/VulkanRHI/VulkanRHI.Build.cs
+++ b/Engine/Source/Runtime/VulkanRHI/VulkanRHI.Build.cs
@@ -81,7 +81,7 @@ public class VulkanRHI : ModuleRules
 		if (Target.IsInPlatformGroup(UnrealPlatformGroup.Linux))
 		{
 			PrivateDependencyModuleNames.Add("ApplicationCore");
-			AddEngineThirdPartyPrivateStaticDependencies(Target, "SDL2");
+			AddEngineThirdPartyPrivateStaticDependencies(Target, "SDL3");
 
 			// Always add the SDK's loader in the ThirdParty binaries, it might be used if the local system's loader is not supported
 			RuntimeDependencies.Add(Path.Combine(Unreal.EngineDirectory.ToString(), "Binaries/ThirdParty/Vulkan/Linux/libvulkan.so"), StagedFileType.NonUFS);
