- Add Show/Hide Verbosity to new console commands and right click menu

#rnx

diff --git a/Engine/Source/Runtime/ApplicationCore/Private/Windows/WindowsConsoleOutputDevice2.cpp b/Engine/Source/Runtime/ApplicationCore/Private/Windows/WindowsConsoleOutputDevice2.cpp
--- a/Engine/Source/Runtime/ApplicationCore/Private/Windows/WindowsConsoleOutputDevice2.cpp
+++ b/Engine/Source/Runtime/ApplicationCore/Private/Windows/WindowsConsoleOutputDevice2.cpp
@@ -50,6 +50,8 @@ public:
 		"  showdate\t - Show time in log\r\n"
 		"  showdatetime\t - Show date/time in log\r\n"
 		"  hidedatetime\t - Hide date/time from log\r\n"
+		"  showverbosity\t - Show verbosity from log\r\n"
+		"  hideverbosity\t - Hide verbosity from log\r\n"
 		"  <color>\t - Change color of console\r\n"
 		"\r\n"
 		"Supported colors: red,gray,darkgray,lightblue,darkblue\r\n"
@@ -1221,6 +1223,18 @@ public:
 		else
 			OutString.Append(Category).Append(CategorySpace + FCString::Strlen(Category));
 		OutString.Append(": ");
+
+		if (bConsoleShowVerbosity)
+		{
+			const TCHAR* VerbosityStr = ToString(E.Verbosity);
+			const TCHAR VerbositySpace[] = TEXT("           ");
+			if (FCString::Strlen(VerbosityStr) >= UE_ARRAY_COUNT(VerbositySpace))
+				OutString.Append(VerbosityStr, UE_ARRAY_COUNT(VerbositySpace) - 1);
+			else
+				OutString.Append(VerbosityStr).Append(VerbositySpace + FCString::Strlen(VerbosityStr));
+			OutString.Append(": ");
+		}
+
 		for (TCHAR C : E.String)
 		{
 			if (C == '\t')
@@ -1439,6 +1453,10 @@ public:
 				AddItem(TEXT("Show Date and Time"), 132);
 			if (ConsoleShowDateTime != 0)
 				AddItem(TEXT("Hide Date and Time"), 133);
+			if (bConsoleShowVerbosity)
+				AddItem(TEXT("Hide Verbosity"), 134);
+			if (!bConsoleShowVerbosity)
+				AddItem(TEXT("Show Verbosity"), 135);
 
 			MapWindowPoints(hWnd, HWND_DESKTOP, &MousePos, 1);
 			TrackPopupMenu(Menu, 0, MousePos.x, MousePos.y, 0, hWnd, NULL);
@@ -1534,6 +1552,16 @@ public:
 				ConsoleShowDateTime = 0;
 				RefreshLogHwnd();
 				break;
+
+			case 134:
+				bConsoleShowVerbosity = 0;
+				RefreshLogHwnd();
+				break;
+
+			case 135:
+				bConsoleShowVerbosity = 1;
+				RefreshLogHwnd();
+				break;
 			}
 			break;
 		}
@@ -2440,6 +2468,7 @@ public:
 				SetWindowText(CommandHwnd, TEXT(""));
 				EConsoleColor NewConsoleColor = EConsoleColor(-1);
 				int32 OldConsoleShowDateTime = ConsoleShowDateTime;
+				bool bOldConsoleShowVerbosity = bConsoleShowVerbosity;
 				if (FCString::Stricmp(Command, TEXT("red")) == 0)
 					NewConsoleColor = EConsoleColor::Red;
 				else if (FCString::Stricmp(Command, TEXT("darkgray")) == 0)
@@ -2458,8 +2487,13 @@ public:
 					ConsoleShowDateTime = 2;
 				else if (FCString::Stricmp(Command, TEXT("hidedatetime")) == 0)
 					ConsoleShowDateTime = 0;
+				else if (FCString::Stricmp(Command, TEXT("showverbosity")) == 0)
+					bConsoleShowVerbosity = true;
+				else if (FCString::Stricmp(Command, TEXT("hideverbosity")) == 0)
+					bConsoleShowVerbosity = false;
 
-				if (ConsoleShowDateTime != OldConsoleShowDateTime)
+				if (ConsoleShowDateTime != OldConsoleShowDateTime
+					|| bConsoleShowVerbosity != bOldConsoleShowVerbosity)
 					RefreshLogHwnd();
 
 				if (NewConsoleColor != EConsoleColor(-1))
@@ -2747,6 +2781,7 @@ public:
 
 	FString ConsoleTitle;
 	bool bConsoleExpanded = false;
+	bool bConsoleShowVerbosity = false;
 };
 
 FName AppStatusTrackerName(TEXT("AppStatusTracker"));
@@ -2909,6 +2944,7 @@ void FWindowsConsoleOutputDevice2::SaveToINI()
 	FString ConsoleIncludeFilterStr = Window->ConsoleIncludeFilterStr;
 	FString ConsoleExcludeFilterStr = Window->ConsoleExcludeFilterStr;
 	bool bConsoleExpanded = Window->bConsoleExpanded;
+	bool bConsoleShowVerbosity = Window->bConsoleShowVerbosity;
 	WindowRWLock.WriteUnlock();
 
 	FString Filename = GetConfigFilename();
@@ -2937,6 +2973,7 @@ void FWindowsConsoleOutputDevice2::SaveToINI()
 	Config.SetInt(Selection, TEXT("ConsoleColor"), ConsoleColor, Filename);
 
 	Config.SetBool(Selection, TEXT("ConsoleExpanded"), bConsoleExpanded, Filename);
+	Config.SetBool(Selection, TEXT("ConsoleShowVerbosity"), bConsoleShowVerbosity, Filename);
 
 	Config.SetInt(Selection, TEXT("ConsoleShowDateTime"), ConsoleShowDateTime, Filename);
 
@@ -2970,6 +3007,7 @@ void FWindowsConsoleOutputDevice2::Show( bool ShowWindow )
 		FString ConsoleExcludeFilterStr;
 
 		bool bConsoleExpanded = false;
+		bool bConsoleShowVerbosity = false;
 		bool bHasX = false;
 		bool bHasY = false;
 
@@ -2989,6 +3027,7 @@ void FWindowsConsoleOutputDevice2::Show( bool ShowWindow )
 		bHasY = Config.GetInt(Selection, TEXT("ConsoleY"), ConsolePosY, Filename);
 		Config.GetInt(Selection, TEXT("ConsoleColor"), ConsoleColor, Filename);
 		Config.GetBool(Selection, TEXT("ConsoleExpanded"), bConsoleExpanded, Filename);
+		Config.GetBool(Selection, TEXT("ConsoleShowVerbosity"), bConsoleShowVerbosity, Filename);
 		Config.GetInt(Selection, TEXT("ConsoleShowDateTime"), ConsoleShowDateTime, Filename);
 		Config.GetString(Selection, TEXT("IncludeFilter"), ConsoleIncludeFilterStr, Filename);
 		Config.GetString(Selection, TEXT("ExcludeFilter"), ConsoleExcludeFilterStr, Filename);
@@ -3057,6 +3096,7 @@ void FWindowsConsoleOutputDevice2::Show( bool ShowWindow )
 
 		Window->ConsoleTitle = ConsoleTitle;
 		Window->bConsoleExpanded = bConsoleExpanded;
+		Window->bConsoleShowVerbosity = bConsoleShowVerbosity;
 		Window->bIsVisible = true;
 
 		if (bFirstCall)
