Restore WindowsApplication.UseWorkerThreadForRawInput=True after fixing support for using the right mouse button as primary

diff --git a/Engine/Config/Windows/BaseWindowsEngine.ini b/Engine/Config/Windows/BaseWindowsEngine.ini
--- a/Engine/Config/Windows/BaseWindowsEngine.ini
+++ b/Engine/Config/Windows/BaseWindowsEngine.ini
@@ -60,4 +60,4 @@ MaxNumTlsSlots=2048
 MaxNumThreadsWithTlsSlots=512
 
 [ConsoleVariables]
-WindowsApplication.UseWorkerThreadForRawInput=False
+WindowsApplication.UseWorkerThreadForRawInput=True
diff --git a/Engine/Source/Runtime/ApplicationCore/Private/Windows/WindowsApplication.cpp b/Engine/Source/Runtime/ApplicationCore/Private/Windows/WindowsApplication.cpp
--- a/Engine/Source/Runtime/ApplicationCore/Private/Windows/WindowsApplication.cpp
+++ b/Engine/Source/Runtime/ApplicationCore/Private/Windows/WindowsApplication.cpp
@@ -918,7 +918,7 @@ namespace UE::WindowsApplication::RawInputThread
 		HWND WindowHandle = NULL;
 
 		FMouseQueue MouseInputs;
-		std::atomic<uint32> InputCount;
+		std::atomic<uint32> InputCount = 0;
 
 		bool bRegistered = false;
 		TArray<SIZE_T> DataArray; // GetRawInputBuffer wants DWORD/QWORD alignment for 32/64-bit architecture
@@ -4479,6 +4479,13 @@ void FWindowsApplication::ProcessWorkerInputs()
 			}
 		};
 
+	EMouseButtons::Type PrimaryMouseButton = EMouseButtons::Left;
+	EMouseButtons::Type SecondaryMouseButton = EMouseButtons::Right;
+	if (GetSystemMetrics(SM_SWAPBUTTON))
+	{
+		Swap(PrimaryMouseButton, SecondaryMouseButton);
+	}
+
 	for (uint32 InputIndex = 0, InputCount = UE::WindowsApplication::RawInputThread::FWindowsRawInputRunnable::TakeMouseInputCount(); InputIndex < InputCount; ++InputIndex)
 	{
 		RAWMOUSE Mouse;
@@ -4544,8 +4551,8 @@ void FWindowsApplication::ProcessWorkerInputs()
 			const FVector2D ButtonCursorPos = FVector2D(ButtonCursorPoint.x, ButtonCursorPoint.y);
 			const double CurrentTime = FPlatformTime::Seconds();
 
-			ProcessMouseButton(CurrentTime, Window, ButtonFlags, RI_MOUSE_LEFT_BUTTON_DOWN, RI_MOUSE_LEFT_BUTTON_UP, EMouseButtons::Left, ButtonCursorPos);
-			ProcessMouseButton(CurrentTime, Window, ButtonFlags, RI_MOUSE_RIGHT_BUTTON_DOWN, RI_MOUSE_RIGHT_BUTTON_UP, EMouseButtons::Right, ButtonCursorPos);
+			ProcessMouseButton(CurrentTime, Window, ButtonFlags, RI_MOUSE_LEFT_BUTTON_DOWN, RI_MOUSE_LEFT_BUTTON_UP, PrimaryMouseButton, ButtonCursorPos);
+			ProcessMouseButton(CurrentTime, Window, ButtonFlags, RI_MOUSE_RIGHT_BUTTON_DOWN, RI_MOUSE_RIGHT_BUTTON_UP, SecondaryMouseButton, ButtonCursorPos);
 			ProcessMouseButton(CurrentTime, Window, ButtonFlags, RI_MOUSE_MIDDLE_BUTTON_DOWN, RI_MOUSE_MIDDLE_BUTTON_UP, EMouseButtons::Middle, ButtonCursorPos);
 			ProcessMouseButton(CurrentTime, Window, ButtonFlags, RI_MOUSE_BUTTON_4_DOWN, RI_MOUSE_BUTTON_4_UP, EMouseButtons::Thumb01, ButtonCursorPos);
 			ProcessMouseButton(CurrentTime, Window, ButtonFlags, RI_MOUSE_BUTTON_5_DOWN, RI_MOUSE_BUTTON_5_UP, EMouseButtons::Thumb02, ButtonCursorPos);
