Bind core delegates to ApplicationHasReactivatedDelegate and ApplicationWillDeactivateDelegate to properly notify EOS SDK when a phone call happens

diff --git a/Engine/Plugins/Online/EOSShared/Source/EOSShared/Private/IOS/IOSEOSSDKManager.cpp b/Engine/Plugins/Online/EOSShared/Source/EOSShared/Private/IOS/IOSEOSSDKManager.cpp
--- a/Engine/Plugins/Online/EOSShared/Source/EOSShared/Private/IOS/IOSEOSSDKManager.cpp
+++ b/Engine/Plugins/Online/EOSShared/Source/EOSShared/Private/IOS/IOSEOSSDKManager.cpp
@@ -16,6 +16,8 @@ FIOSEOSSDKManager::FIOSEOSSDKManager()
 {
 	FCoreDelegates::ApplicationHasEnteredForegroundDelegate.AddRaw(this, &FIOSEOSSDKManager::OnApplicationStatusChanged, EOS_EApplicationStatus::EOS_AS_Foreground);
 	FCoreDelegates::ApplicationWillEnterBackgroundDelegate.AddRaw(this, &FIOSEOSSDKManager::OnApplicationStatusChanged, EOS_EApplicationStatus::EOS_AS_BackgroundSuspended);
+	FCoreDelegates::AudioInterruptionDelegate.AddRaw(this, &FIOSEOSSDKManager::OnAudioInterruptedNotification);
+	FCoreDelegates::ApplicationHasReactivatedDelegate.AddRaw(this, &FIOSEOSSDKManager::OnApplicationStatusChanged, EOS_EApplicationStatus::EOS_AS_Foreground);
 	FIOSCoreDelegates::OnOpenURL.AddStatic(&OnUrlOpened);
 }
 
@@ -23,6 +25,8 @@ FIOSEOSSDKManager::~FIOSEOSSDKManager()
 {	
 	FCoreDelegates::ApplicationHasEnteredForegroundDelegate.RemoveAll(this);
 	FCoreDelegates::ApplicationWillEnterBackgroundDelegate.RemoveAll(this);
+	FCoreDelegates::AudioInterruptionDelegate.RemoveAll(this);
+	FCoreDelegates::ApplicationHasReactivatedDelegate.RemoveAll(this);
 }
 
 FString FIOSEOSSDKManager::GetCacheDirBase() const
@@ -35,4 +39,12 @@ FString FIOSEOSSDKManager::GetCacheDirBase() const
 	return FString(UTF8_TO_TCHAR(CStrCacheDirectory));
 };
 
+void FIOSEOSSDKManager::OnAudioInterruptedNotification(bool bInterrupted)
+{
+	if (bInterrupted)
+	{
+		OnApplicationStatusChanged(EOS_EApplicationStatus::EOS_AS_BackgroundSuspended);
+	}
+}
+
 #endif // WITH_EOS_SDK
diff --git a/Engine/Plugins/Online/EOSShared/Source/EOSShared/Private/IOS/IOSEOSSDKManager.h b/Engine/Plugins/Online/EOSShared/Source/EOSShared/Private/IOS/IOSEOSSDKManager.h
--- a/Engine/Plugins/Online/EOSShared/Source/EOSShared/Private/IOS/IOSEOSSDKManager.h
+++ b/Engine/Plugins/Online/EOSShared/Source/EOSShared/Private/IOS/IOSEOSSDKManager.h
@@ -13,8 +13,10 @@ public:
 	~FIOSEOSSDKManager();
 
 	virtual FString GetCacheDirBase() const override;
+private:
+	void OnAudioInterruptedNotification(bool bInterrupted);
 };
 
 using FPlatformEOSSDKManager = FIOSEOSSDKManager;
 
-#endif // WITH_EOS_SDK
\ No newline at end of file
+#endif // WITH_EOS_SDK
diff --git a/Engine/Source/Runtime/ApplicationCore/Private/IOS/IOSAppDelegate.cpp b/Engine/Source/Runtime/ApplicationCore/Private/IOS/IOSAppDelegate.cpp
--- a/Engine/Source/Runtime/ApplicationCore/Private/IOS/IOSAppDelegate.cpp
+++ b/Engine/Source/Runtime/ApplicationCore/Private/IOS/IOSAppDelegate.cpp
@@ -559,26 +559,35 @@ static IOSAppDelegate* CachedDelegate = nil;
 {
 	[[NSNotificationCenter defaultCenter] addObserverForName:AVAudioSessionInterruptionNotification object:nil queue:nil usingBlock:^(NSNotification *notification)
 	{
-		// the audio context should resume immediately after interrupt, if suspended
-		FAppEntry::ResetAudioContextResumeTime();
+		NSUInteger InterruptionType = [[[notification userInfo] objectForKey:AVAudioSessionInterruptionTypeKey] unsignedIntegerValue];
+		NSNumber* InterruptionOption = [[notification userInfo] objectForKey:AVAudioSessionInterruptionOptionKey];
+		UE_CLOG(InterruptionType == AVAudioSessionInterruptionTypeBegan, LogIOSAudioSession, Log, TEXT("AudioSession interruption notification: Began"));
+		UE_CLOG(InterruptionType == AVAudioSessionInterruptionTypeEnded, LogIOSAudioSession, Log, TEXT("AudioSession interruption notification: Ended"));
 
-		switch ([[[notification userInfo] objectForKey:AVAudioSessionInterruptionTypeKey] unsignedIntegerValue])
+		[FIOSAsyncTask CreateTaskWithBlock : ^ bool(void)
 		{
-			case AVAudioSessionInterruptionTypeBegan:
-				// Notification that our audio was stopped by the OS. No action needed.
-				FCoreDelegates::AudioInterruptionDelegate.Broadcast(true);
-				break;
-
-			case AVAudioSessionInterruptionTypeEnded:
-				FCoreDelegates::AudioInterruptionDelegate.Broadcast(false);
-				NSNumber * interruptionOption = [[notification userInfo] objectForKey:AVAudioSessionInterruptionOptionKey];
-				if (interruptionOption != nil && interruptionOption.unsignedIntegerValue > 0)
-				{
-                    FAppEntry::RestartAudio();
-                }
-                [self ToggleAudioSession:true];
-				break;
-		}
+			// the audio context should resume immediately after interrupt, if suspended
+			FAppEntry::ResetAudioContextResumeTime();
+			
+			switch (InterruptionType)
+			{
+				case AVAudioSessionInterruptionTypeBegan:
+					// Notification that our audio was stopped by the OS. No action needed.
+					FCoreDelegates::AudioInterruptionDelegate.Broadcast(true);
+					break;
+					
+				case AVAudioSessionInterruptionTypeEnded:
+					FCoreDelegates::AudioInterruptionDelegate.Broadcast(false);
+					
+					if (InterruptionOption != nil && InterruptionOption.unsignedIntegerValue > 0)
+					{
+						FAppEntry::RestartAudio();
+					}
+					[self ToggleAudioSession:true];
+					break;
+			}
+			return true;
+		}];
 	}];
 
 	[[NSNotificationCenter defaultCenter] addObserverForName:AVAudioSessionRouteChangeNotification object : nil queue : nil usingBlock : ^ (NSNotification *notification)
@@ -732,7 +741,11 @@ static IOSAppDelegate* CachedDelegate = nil;
 		[[AVAudioSession sharedInstance] setCategory:Category mode:Mode options:Options error:&ActiveError];
 		if (ActiveError)
 		{
-			UE_LOG(LogIOSAudioSession, Error, TEXT("Failed to set AVAudioSession category to Category:%s Mode:%s Options:%x! [Error = %s]"), *FString(Category), *FString(Mode), Options, *FString([ActiveError description]));
+			UE_LOG(LogIOSAudioSession, Error, TEXT("Failed to set AVAudioSession category to Category:%s (current %s) Mode:%s (current %s) Options:%x (current %x)! [Error = %s]."), 
+				   *FString(Category), *FString([[AVAudioSession sharedInstance] category]), 
+				   *FString(Mode), *FString([[AVAudioSession sharedInstance] mode]), 
+				   Options, [[AVAudioSession sharedInstance] categoryOptions] & ~AVAudioSessionCategoryOptionMixWithOthers, 
+				   *FString([ActiveError description]));
 		}
 	}
 	
