Camera stutters when feathering controller right joystick to pan and tilt

Changes
- Preventing FAppleControllerInterface to send empty (zero) axisValue when the Previous Input is equal to Current Input and the Axis value is below the Deadzone.

diff --git a/Engine/Source/Runtime/ApplicationCore/Private/Apple/AppleControllerInterface.cpp b/Engine/Source/Runtime/ApplicationCore/Private/Apple/AppleControllerInterface.cpp
--- a/Engine/Source/Runtime/ApplicationCore/Private/Apple/AppleControllerInterface.cpp
+++ b/Engine/Source/Runtime/ApplicationCore/Private/Apple/AppleControllerInterface.cpp
@@ -556,6 +556,7 @@ void FAppleControllerInterface::HandleAnalogGamepad(const FGamepadKeyNames::Type
     bool bWasNegativePressed = false;
     bool bNegativePressed = false;
     float axisValue = 0;
+	bool bSendAnalogChangedMessage = false;
 
     GCExtendedGamepad *ExtendedGamepad = Cont.extendedGamepad;
     GCExtendedGamepad *ExtendedPreviousGamepad = Controllers[ControllerIndex].PreviousExtendedGamepad;;
@@ -570,30 +571,89 @@ void FAppleControllerInterface::HandleAnalogGamepad(const FGamepadKeyNames::Type
         case ControllerType::XboxGamepad:
         case ControllerType::DualSenseGamepad:
             
-            if (UEAxis == FGamepadKeyNames::LeftAnalogX){if ((ExtendedPreviousGamepad != nil && ExtendedGamepad.leftThumbstick.xAxis.value != ExtendedPreviousGamepad.leftThumbstick.xAxis.value) ||
-                                                             (ExtendedGamepad.leftThumbstick.xAxis.value < -RepeatDeadzone || ExtendedGamepad.leftThumbstick.xAxis.value > RepeatDeadzone)){axisValue = ExtendedGamepad.leftThumbstick.xAxis.value;}}
-            else if (UEAxis == FGamepadKeyNames::LeftAnalogY){if ((ExtendedPreviousGamepad != nil && ExtendedGamepad.leftThumbstick.yAxis.value != ExtendedPreviousGamepad.leftThumbstick.yAxis.value) ||
-                                                                  (ExtendedGamepad.leftThumbstick.yAxis.value < -RepeatDeadzone || ExtendedGamepad.leftThumbstick.yAxis.value > RepeatDeadzone)){axisValue = ExtendedGamepad.leftThumbstick.yAxis.value;}}
-            else if (UEAxis == FGamepadKeyNames::RightAnalogX){if ((ExtendedPreviousGamepad != nil && ExtendedGamepad.rightThumbstick.xAxis.value != ExtendedPreviousGamepad.rightThumbstick.xAxis.value) ||
-                                                                   (ExtendedGamepad.rightThumbstick.xAxis.value < -RepeatDeadzone || ExtendedGamepad.rightThumbstick.xAxis.value > RepeatDeadzone)){axisValue = ExtendedGamepad.rightThumbstick.xAxis.value;}}
-            else if (UEAxis == FGamepadKeyNames::RightAnalogY){if ((ExtendedPreviousGamepad != nil && ExtendedGamepad.rightThumbstick.yAxis.value != ExtendedPreviousGamepad.rightThumbstick.yAxis.value) ||
-                                                                   (ExtendedGamepad.rightThumbstick.yAxis.value < -RepeatDeadzone || ExtendedGamepad.rightThumbstick.yAxis.value > RepeatDeadzone)){axisValue = ExtendedGamepad.rightThumbstick.yAxis.value;}}
-            else if (UEAxis == FGamepadKeyNames::LeftTriggerAnalog){if ((ExtendedPreviousGamepad != nil && ExtendedGamepad.leftTrigger.value != ExtendedPreviousGamepad.leftTrigger.value) ||
-                                                                        (ExtendedGamepad.leftTrigger.value < -RepeatDeadzone || ExtendedGamepad.leftTrigger.value > RepeatDeadzone)){axisValue = ExtendedGamepad.leftTrigger.value;}}
-            else if (UEAxis == FGamepadKeyNames::RightTriggerAnalog){if ((ExtendedPreviousGamepad != nil && ExtendedGamepad.rightTrigger.value != ExtendedPreviousGamepad.rightTrigger.value) ||
-                                                                         (ExtendedGamepad.rightTrigger.value < -RepeatDeadzone || ExtendedGamepad.rightTrigger.value > RepeatDeadzone)){axisValue = ExtendedGamepad.rightTrigger.value;}}
+            if (UEAxis == FGamepadKeyNames::LeftAnalogX)
+			{
+				if ((ExtendedPreviousGamepad != nil && ExtendedGamepad.leftThumbstick.xAxis.value != ExtendedPreviousGamepad.leftThumbstick.xAxis.value) ||
+                	(ExtendedGamepad.leftThumbstick.xAxis.value < -RepeatDeadzone || ExtendedGamepad.leftThumbstick.xAxis.value > RepeatDeadzone))
+					{
+						axisValue = ExtendedGamepad.leftThumbstick.xAxis.value;
+						bSendAnalogChangedMessage = true;
+					}
+			}
+            else if (UEAxis == FGamepadKeyNames::LeftAnalogY)
+			{
+				if ((ExtendedPreviousGamepad != nil && ExtendedGamepad.leftThumbstick.yAxis.value != ExtendedPreviousGamepad.leftThumbstick.yAxis.value) ||
+                    (ExtendedGamepad.leftThumbstick.yAxis.value < -RepeatDeadzone || ExtendedGamepad.leftThumbstick.yAxis.value > RepeatDeadzone))
+				{
+					axisValue = ExtendedGamepad.leftThumbstick.yAxis.value;
+					bSendAnalogChangedMessage = true;
+				}
+			}
+            else if (UEAxis == FGamepadKeyNames::RightAnalogX)
+			{
+				if ((ExtendedPreviousGamepad != nil && ExtendedGamepad.rightThumbstick.xAxis.value != ExtendedPreviousGamepad.rightThumbstick.xAxis.value) ||
+                	(ExtendedGamepad.rightThumbstick.xAxis.value < -RepeatDeadzone || ExtendedGamepad.rightThumbstick.xAxis.value > RepeatDeadzone))
+				{
+					axisValue = ExtendedGamepad.rightThumbstick.xAxis.value;
+					bSendAnalogChangedMessage = true;
+				}
+			}
+            else if (UEAxis == FGamepadKeyNames::RightAnalogY)
+			{
+				if ((ExtendedPreviousGamepad != nil && ExtendedGamepad.rightThumbstick.yAxis.value != ExtendedPreviousGamepad.rightThumbstick.yAxis.value) ||
+                    (ExtendedGamepad.rightThumbstick.yAxis.value < -RepeatDeadzone || ExtendedGamepad.rightThumbstick.yAxis.value > RepeatDeadzone))
+				{
+					axisValue = ExtendedGamepad.rightThumbstick.yAxis.value;
+					bSendAnalogChangedMessage = true;
+				}
+			}
+            else if (UEAxis == FGamepadKeyNames::LeftTriggerAnalog)
+			{
+				if ((ExtendedPreviousGamepad != nil && ExtendedGamepad.leftTrigger.value != ExtendedPreviousGamepad.leftTrigger.value) ||
+                	(ExtendedGamepad.leftTrigger.value < -RepeatDeadzone || ExtendedGamepad.leftTrigger.value > RepeatDeadzone))
+				{
+					axisValue = ExtendedGamepad.leftTrigger.value;
+					bSendAnalogChangedMessage = true;
+				}
+			}
+            else if (UEAxis == FGamepadKeyNames::RightTriggerAnalog)
+			{
+				if ((ExtendedPreviousGamepad != nil && ExtendedGamepad.rightTrigger.value != ExtendedPreviousGamepad.rightTrigger.value) ||
+                	(ExtendedGamepad.rightTrigger.value < -RepeatDeadzone || ExtendedGamepad.rightTrigger.value > RepeatDeadzone))
+				{
+					axisValue = ExtendedGamepad.rightTrigger.value;
+					bSendAnalogChangedMessage = true;
+				}
+			}
             break;
             
         case ControllerType::SiriRemote:
             
-            if (UEAxis == FGamepadKeyNames::LeftAnalogX){if ((ExtendedPreviousGamepad != nil && MicroGamepad.dpad.xAxis.value != ExtendedPreviousGamepad.dpad.xAxis.value) ||
-                                                             (MicroGamepad.dpad.xAxis.value < -RepeatDeadzone || MicroGamepad.dpad.xAxis.value > RepeatDeadzone)){axisValue = MicroGamepad.dpad.xAxis.value;}}
-            else if (UEAxis == FGamepadKeyNames::LeftAnalogY){if ((MicroPreviousGamepad != nil && MicroGamepad.dpad.yAxis.value != MicroPreviousGamepad.dpad.yAxis.value) ||
-                                                                  (MicroGamepad.dpad.yAxis.value < -RepeatDeadzone || MicroGamepad.dpad.yAxis.value > RepeatDeadzone)){axisValue = MicroGamepad.dpad.yAxis.value;}}
+            if (UEAxis == FGamepadKeyNames::LeftAnalogX)
+			{
+				if ((ExtendedPreviousGamepad != nil && MicroGamepad.dpad.xAxis.value != ExtendedPreviousGamepad.dpad.xAxis.value) ||
+                (MicroGamepad.dpad.xAxis.value < -RepeatDeadzone || MicroGamepad.dpad.xAxis.value > RepeatDeadzone))
+				{
+					axisValue = MicroGamepad.dpad.xAxis.value;
+					bSendAnalogChangedMessage = true;
+				}
+			}
+            else if (UEAxis == FGamepadKeyNames::LeftAnalogY)
+			{
+				if ((MicroPreviousGamepad != nil && MicroGamepad.dpad.yAxis.value != MicroPreviousGamepad.dpad.yAxis.value) ||
+                (MicroGamepad.dpad.yAxis.value < -RepeatDeadzone || MicroGamepad.dpad.yAxis.value > RepeatDeadzone))
+				{
+					axisValue = MicroGamepad.dpad.yAxis.value;
+					bSendAnalogChangedMessage = true;
+				}
+			}
             break;
     }
 #if APPLE_CONTROLLER_DEBUG
     NSLog(@"Axis %s is %f", TCHAR_TO_ANSI(*UEAxis.ToString()), axisValue);
 #endif
-    MessageHandler->OnControllerAnalog(UEAxis, UserId, DeviceId, axisValue);
+	if (bSendAnalogChangedMessage)
+	{
+		MessageHandler->OnControllerAnalog(UEAxis, UserId, DeviceId, axisValue);
+	}
 }
