Mac - Fix Cursor offset from visual because of outdated ScreensArray when processing OnWindowDidMove resulting from a DPI change
Fix MainThreadChecker issue in FMacApplication::OnWindowDidMove
Move bDisplayReconfiguring from FCocoaWindow to FMacApplication since it's not dependant of FCocoaWindow

#tests local testing on mac & preflight

diff --git a/Engine/Source/Runtime/ApplicationCore/Private/Mac/CocoaWindow.cpp b/Engine/Source/Runtime/ApplicationCore/Private/Mac/CocoaWindow.cpp
--- a/Engine/Source/Runtime/ApplicationCore/Private/Mac/CocoaWindow.cpp
+++ b/Engine/Source/Runtime/ApplicationCore/Private/Mac/CocoaWindow.cpp
@@ -25,7 +25,6 @@ NSString* NSPerformDragOperation = @"NSPerformDragOperation";
 {
 	WindowMode = EWindowMode::Windowed;
 	bAcceptsInput = false;
-	bDisplayReconfiguring = false;
 	bRenderInitialized = false;
 	bIsBeingOrderedFront = false;
 	bIsBeingResized = false;
@@ -83,11 +82,6 @@ NSString* NSPerformDragOperation = @"NSPerformDragOperation";
 	return WindowMode;
 }
 
-- (void)setDisplayReconfiguring:(bool)bIsDisplayReconfiguring
-{
-	bDisplayReconfiguring = bIsDisplayReconfiguring;
-}
-
 - (void)orderFrontAndMakeMain:(bool)bMain andKey:(bool)bKey
 {
 	SCOPED_AUTORELEASE_POOL;
@@ -362,7 +356,7 @@ NSString* NSPerformDragOperation = @"NSPerformDragOperation";
 	CachedScreen = [self screen];
 	// The windowdidChangeScreen notification only comes after you finish dragging.
 	// It does however, work fine for handling display arrangement changes that cause a window to go offscreen.
-	if (bDisplayReconfiguring)
+	if (MacApplication && MacApplication->IsDisplayReconfiguring())
 	{
 		SCOPED_AUTORELEASE_POOL;
 		NSScreen* Screen = [self screen];
diff --git a/Engine/Source/Runtime/ApplicationCore/Private/Mac/MacApplication.cpp b/Engine/Source/Runtime/ApplicationCore/Private/Mac/MacApplication.cpp
--- a/Engine/Source/Runtime/ApplicationCore/Private/Mac/MacApplication.cpp
+++ b/Engine/Source/Runtime/ApplicationCore/Private/Mac/MacApplication.cpp
@@ -129,6 +129,7 @@ FMacApplication::FMacApplication()
 ,	AccessibilityCacheTimer(Nil)
 ,	AccessibilityAnnouncementDelayTimer(Nil)
 #endif
+,	bDisplayReconfiguring(false)
 {
 	TextInputMethodSystem = MakeShareable(new FMacTextInputMethodSystem);
 	if (!TextInputMethodSystem->Initialize())
@@ -694,8 +695,16 @@ NSEvent* FMacApplication::HandleNSEvent(NSEvent* Event)
 void FMacApplication::OnDisplayReconfiguration(CGDirectDisplayID Display, CGDisplayChangeSummaryFlags Flags, void* UserInfo)
 {
 	FMacApplication* App = (FMacApplication*)UserInfo;
-	if (Flags & kCGDisplayDesktopShapeChangedFlag)
+
+	
+	if(Flags & kCGDisplayBeginConfigurationFlag)
 	{
+		App->bDisplayReconfiguring = true;
+	}
+	else if (Flags & kCGDisplayDesktopShapeChangedFlag)
+	{
+		App->bDisplayReconfiguring = false;
+
 		App->UpdateScreensArray();
 
 		// Slate needs to know when desktop size changes.
@@ -1206,13 +1215,24 @@ void FMacApplication::OnWindowDidMove(TSharedRef<FMacWindow> Window)
     }
     else
     {
-		// Make sure to take notch and padding into account
-        NSRect WindowFrame = [Window->GetWindowHandle() frame];
-        NSRect OpenGLFrame = [Window->GetWindowHandle() openGLFrame];
+		__block NSRect WindowFrame, OpenGLFrame;
+		MainThreadCall(^{
+			// Make sure to take notch and padding into account
+			WindowFrame = [Window->GetWindowHandle() frame];
+			OpenGLFrame = [Window->GetWindowHandle() openGLFrame];
+		});
 
         const double X = WindowFrame.origin.x;
         const double Y = WindowFrame.origin.y + OpenGLFrame.size.height;
 
+		
+		// Update ScreensArray when displays are reconfiguring to have up to date informations
+		// This is to prevent UE-176002 that happened because UpdateScreensArray is called once
+		// the display reconfiguring is done but OnWindowDidResize is call during the reconfiguring
+		if (bDisplayReconfiguring)
+		{
+			UpdateScreensArray();
+		}
         FVector2D SlatePosition = ConvertCocoaPositionToSlate(X, Y);
         
         MessageHandler->OnMovedWindow(Window, FMath::TruncToInt(SlatePosition.X), FMath::TruncToInt(SlatePosition.Y));
diff --git a/Engine/Source/Runtime/ApplicationCore/Private/Mac/MacWindow.cpp b/Engine/Source/Runtime/ApplicationCore/Private/Mac/MacWindow.cpp
--- a/Engine/Source/Runtime/ApplicationCore/Private/Mac/MacWindow.cpp
+++ b/Engine/Source/Runtime/ApplicationCore/Private/Mac/MacWindow.cpp
@@ -96,7 +96,6 @@ void FMacWindow::Initialize( FMacApplication* const Application, const TSharedRe
 			[WindowHandle setReleasedWhenClosed:NO];
 			[WindowHandle setWindowMode: EWindowMode::Windowed];
 			[WindowHandle setAcceptsInput: Definition->AcceptsInput];
-			[WindowHandle setDisplayReconfiguring: false];
 			[WindowHandle setAcceptsMouseMovedEvents: YES];
 			[WindowHandle setDelegate: WindowHandle];
 
@@ -590,12 +589,6 @@ void FMacWindow::OnDisplayReconfiguration(CGDirectDisplayID Display, CGDisplayCh
 		{
 			[WindowHandle setMovable: YES];
 			[WindowHandle setMovableByWindowBackground: NO];
-			
-			[WindowHandle setDisplayReconfiguring: true];
-		}
-		else if(Flags & kCGDisplayDesktopShapeChangedFlag)
-		{
-			[WindowHandle setDisplayReconfiguring: false];
 		}
 	}
 }
diff --git a/Engine/Source/Runtime/ApplicationCore/Public/Mac/CocoaWindow.h b/Engine/Source/Runtime/ApplicationCore/Public/Mac/CocoaWindow.h
--- a/Engine/Source/Runtime/ApplicationCore/Public/Mac/CocoaWindow.h
+++ b/Engine/Source/Runtime/ApplicationCore/Public/Mac/CocoaWindow.h
@@ -16,7 +16,6 @@
 @interface FCocoaWindow : NSWindow <NSWindowDelegate, NSDraggingDestination>
 {
 	EWindowMode::Type WindowMode;
-	bool bDisplayReconfiguring;
 	bool bRenderInitialized;
 	bool bIsBeingOrderedFront;
 	CGFloat Opacity;
@@ -58,9 +57,6 @@
 /**	@return The current mode for this Cocoa window. */
 - (EWindowMode::Type)windowMode;
 
-/** Mutator that specifies that the display arrangement is being reconfigured when bIsDisplayReconfiguring is true. */
-- (void)setDisplayReconfiguring:(bool)bIsDisplayReconfiguring;
-
 /** Order window to the front. */
 - (void)orderFrontAndMakeMain:(bool)bMain andKey:(bool)bKey;
 
diff --git a/Engine/Source/Runtime/ApplicationCore/Public/Mac/MacApplication.h b/Engine/Source/Runtime/ApplicationCore/Public/Mac/MacApplication.h
--- a/Engine/Source/Runtime/ApplicationCore/Public/Mac/MacApplication.h
+++ b/Engine/Source/Runtime/ApplicationCore/Public/Mac/MacApplication.h
@@ -237,6 +237,8 @@ public:
 	void OnWindowActivationChanged(const TSharedRef<FMacWindow>& Window, const EWindowActivation ActivationType);
 
 	static void OnDisplayReconfiguration(CGDirectDisplayID Display, CGDisplayChangeSummaryFlags Flags, void* UserInfo);
+	
+	inline bool IsDisplayReconfiguring() const { return bDisplayReconfiguring; }
 
 public:
     virtual IInputInterface* GetInputInterface() override { return this; }
@@ -414,6 +416,8 @@ private:
 	
 	NSTimer* AccessibilityAnnouncementDelayTimer;
 #endif
+	
+	bool bDisplayReconfiguring;
 
 	friend class FMacWindow;
 };
