Set WindowsApplication.UseWorkerThreadForRawInput=True and use TQueue for mouse inputs while removing GetRawInputBuffer

Previous implementation had infrequent crashes involving TLockFreePointerListFIFO and derefencing GetRawInputBuffer results. GetRawInputBuffer was only a minor infrequent optimization since we are processing inputs on a worker thread anyway.

diff --git a/Engine/Config/Windows/BaseWindowsEngine.ini b/Engine/Config/Windows/BaseWindowsEngine.ini
--- a/Engine/Config/Windows/BaseWindowsEngine.ini
+++ b/Engine/Config/Windows/BaseWindowsEngine.ini
@@ -70,4 +70,4 @@ MaxNumTlsSlots=2048
 MaxNumThreadsWithTlsSlots=512
 
 [ConsoleVariables]
-WindowsApplication.UseWorkerThreadForRawInput=False
+WindowsApplication.UseWorkerThreadForRawInput=True
diff --git a/Engine/Source/Runtime/ApplicationCore/Private/Windows/WindowsApplication.cpp b/Engine/Source/Runtime/ApplicationCore/Private/Windows/WindowsApplication.cpp
--- a/Engine/Source/Runtime/ApplicationCore/Private/Windows/WindowsApplication.cpp
+++ b/Engine/Source/Runtime/ApplicationCore/Private/Windows/WindowsApplication.cpp
@@ -855,51 +855,6 @@ namespace UE::WindowsApplication::RawInputThread
 	DECLARE_LOG_CATEGORY_EXTERN(LogWindowsRawInputThread, Display, VeryVerbose);
 	DEFINE_LOG_CATEGORY(LogWindowsRawInputThread);
 
-	// Lockless queue to push mouse inputs from the worker thread to the game thread
-	class FMouseQueue
-	{
-	public:
-		FMouseQueue() = default;
-
-		~FMouseQueue()
-		{
-			RAWMOUSE* Node = FreeList.Pop();
-			while (Node)
-			{
-				delete Node;
-				Node = FreeList.Pop();
-			}
-		}
-
-		bool Dequeue(RAWMOUSE& OutItem)
-		{
-			RAWMOUSE* const Node = QueuedList.Pop();
-			const bool Success = Node != nullptr;
-			if (Success)
-			{
-				OutItem = *Node;
-				FreeList.Push(Node);
-			}
-
-			return Success;
-		}
-
-		void Enqueue(const RAWMOUSE& Mouse)
-		{
-			RAWMOUSE* Node = FreeList.Pop();
-			Node = Node ? Node : new RAWMOUSE;
-			*Node = Mouse;
-			QueuedList.Push(Node);
-		}
-
-	private:
-		TLockFreePointerListUnordered<RAWMOUSE, PLATFORM_CACHE_LINE_SIZE> FreeList;
-		TLockFreePointerListFIFO<RAWMOUSE, PLATFORM_CACHE_LINE_SIZE> QueuedList;
-
-		FMouseQueue(const FMouseQueue&) = delete;
-		FMouseQueue& operator=(const FMouseQueue&) = delete;
-	};
-
 	constexpr uint32 WM_UE_RAWINPUT_QUIT = WM_APP;
 	constexpr uint32 WM_UE_RAWINPUT_REGISTER = WM_APP + 1;
 	constexpr uint32 WM_UE_RAWINPUT_UNREGISTER = WM_APP + 2;
@@ -917,12 +872,10 @@ namespace UE::WindowsApplication::RawInputThread
 
 		HWND WindowHandle = NULL;
 
-		FMouseQueue MouseInputs;
+		TQueue<RAWMOUSE> MouseInputs;
 		std::atomic<uint32> InputCount = 0;
 
 		bool bRegistered = false;
-		TArray<SIZE_T> DataArray; // GetRawInputBuffer wants DWORD/QWORD alignment for 32/64-bit architecture
-		TArray<SIZE_T> BufferArray; // GetRawInputBuffer wants DWORD/QWORD alignment for 32/64-bit architecture
 
 		FWindowsRawInputRunnable() = default;
 
@@ -1091,55 +1044,17 @@ namespace UE::WindowsApplication::RawInputThread
 
 	void FWindowsRawInputRunnable::HandleInputMessage(HRAWINPUT InputHandle)
 	{
-		// Lazy-allocate because HRAWINPUT might be necessary to get the size
-		if (DataArray.IsEmpty())
-		{
-			UINT Size = 0;
-			UINT Result = ::GetRawInputData(InputHandle, RID_INPUT, NULL, &Size, sizeof(RAWINPUTHEADER));
-			if (Result != 0)
-			{
-				UE_LOG(LogWindowsDesktop, Error, TEXT("Could not get the required size for raw input data. Return value: %u"), Result);
-				return;
-			}
+		uint32 Size = 0;
+		::GetRawInputData(InputHandle, RID_INPUT, NULL, &Size, sizeof(RAWINPUTHEADER));
 
-			DataArray.AddDefaulted((Size + sizeof(SIZE_T) - 1) / sizeof(SIZE_T));
-			BufferArray.AddDefaulted(((RawInputBlockCount * Size) + sizeof(SIZE_T) - 1) / sizeof(SIZE_T));
-		}
+		TUniquePtr<uint8[]> RawData = MakeUnique<uint8[]>(Size);
 
-		// Relying only on GetRawInputBuffer seems to lead to missing inputs
-		PRAWINPUT DataRaw = (PRAWINPUT)DataArray.GetData();
-		UINT DataSize = DataArray.Num() * sizeof(SIZE_T);
-		const UINT DataResult = GetRawInputData(InputHandle, RID_INPUT, DataRaw, &DataSize, sizeof(RAWINPUTHEADER));
-		if ((DataResult != UINT_MAX) && (DataResult > 0))
+		if (::GetRawInputData(InputHandle, RID_INPUT, RawData.Get(), &Size, sizeof(RAWINPUTHEADER)) == Size)
 		{
-			MouseInputs.Enqueue(DataRaw->data.mouse);
+			const RAWINPUT* const Raw = (const RAWINPUT* const)RawData.Get();
+			MouseInputs.Enqueue(Raw->data.mouse);
 			++InputCount;
 		}
-
-		PRAWINPUT BufferRaw = (PRAWINPUT)BufferArray.GetData();
-		UINT BufferSize = BufferArray.Num() * sizeof(SIZE_T);
-		UINT BufferResult = GetRawInputBuffer(BufferRaw, &BufferSize, sizeof(RAWINPUTHEADER));
-		while ((BufferResult != UINT_MAX) && (BufferResult > 0))
-		{
-			for (UINT Index = 0; Index < BufferResult; ++Index)
-			{
-				if (BufferRaw->header.dwType == RIM_TYPEMOUSE)
-				{
-					MouseInputs.Enqueue(BufferRaw->data.mouse);
-					++InputCount;
-				}
-
-				typedef ULONGLONG QWORD;
-				BufferRaw = NEXTRAWINPUTBLOCK(BufferRaw);
-			}
-
-			if (BufferResult < RawInputBlockCount)
-			{
-				break;
-			}
-
-			BufferResult = GetRawInputBuffer((PRAWINPUT)BufferRaw, &BufferSize, sizeof(RAWINPUTHEADER));
-		}
 	}
 
 	void FWindowsRawInputRunnable::Initialize(HINSTANCE HInstance)
