[ASIS] In case of ASIS usage support start Engine w/o availabe window.

#rnx

diff --git a/Engine/Source/Runtime/ApplicationCore/Private/Android/AndroidWindow.cpp b/Engine/Source/Runtime/ApplicationCore/Private/Android/AndroidWindow.cpp
--- a/Engine/Source/Runtime/ApplicationCore/Private/Android/AndroidWindow.cpp
+++ b/Engine/Source/Runtime/ApplicationCore/Private/Android/AndroidWindow.cpp
@@ -678,8 +678,9 @@ void FAndroidWindow::CalculateSurfaceSize(int32_t& SurfaceWidth, int32_t& Surfac
 	}
 	else
 	{
+#if !USE_ANDROID_STANDALONE
 		FAndroidWindow::WaitForWindowDimensions();
-
+#endif
 		SurfaceWidth = (GSurfaceViewWidth > 0) ? GSurfaceViewWidth : CachedNativeWindowWidth;
 		SurfaceHeight = (GSurfaceViewHeight > 0) ? GSurfaceViewHeight : CachedNativeWindowHeight;
 	}
