A temporary fix to the black screen when playing startup movies.

diff --git a/Engine/Plugins/Runtime/AndroidMoviePlayer/Source/AndroidMoviePlayer/Private/AndroidMovieStreamer.cpp b/Engine/Plugins/Runtime/AndroidMoviePlayer/Source/AndroidMoviePlayer/Private/AndroidMovieStreamer.cpp
--- a/Engine/Plugins/Runtime/AndroidMoviePlayer/Source/AndroidMoviePlayer/Private/AndroidMovieStreamer.cpp
+++ b/Engine/Plugins/Runtime/AndroidMoviePlayer/Source/AndroidMoviePlayer/Private/AndroidMovieStreamer.cpp
@@ -29,6 +29,9 @@ FAndroidMediaPlayerStreamer::FAndroidMediaPlayerStreamer()
 	, CurrentPosition(-1)
 {
 	JavaMediaPlayer = MakeShareable(new FJavaAndroidMediaPlayer(false, FAndroidMisc::ShouldUseVulkan(), false));
+	// temp hack for restoring the feature of loading movie
+	extern bool ForceAndroidBitmapRendererLegacy;
+	ForceAndroidBitmapRendererLegacy = true;
 	MovieViewport = MakeShareable(new FMovieViewport());
 }
 
diff --git a/Engine/Source/Runtime/ApplicationCore/Private/Android/AndroidJavaMediaPlayer.cpp b/Engine/Source/Runtime/ApplicationCore/Private/Android/AndroidJavaMediaPlayer.cpp
--- a/Engine/Source/Runtime/ApplicationCore/Private/Android/AndroidJavaMediaPlayer.cpp
+++ b/Engine/Source/Runtime/ApplicationCore/Private/Android/AndroidJavaMediaPlayer.cpp
@@ -382,18 +382,40 @@ void FJavaAndroidMediaPlayer::SetAudioVolume(float Volume)
 	CallMethod<void>(SetAudioVolumeMethod, Volume);
 }
 
+bool ForceAndroidBitmapRendererLegacy = false;
+
 /** This function has to be called before calling any GetVideoLastXXXXXXX() functions */
 void FJavaAndroidMediaPlayer::SelectBitmapRenderer()
 {
 	if (BRSelection == BitmapRendererSelection::Unknown)
 	{
-		if (UsingVulkanRenderer && AndroidEnableMediaPlayerBRIR && FAndroidMisc::IsExternalMemoryAndroidHardwareBufferExtensionLoaded())
+		if (ForceAndroidBitmapRendererLegacy)
 		{
-			BRSelection = BitmapRendererSelection::ImageReader;
+			BRSelection = BitmapRendererSelection::Legacy;
+			ForceAndroidBitmapRendererLegacy = false;
 		}
 		else
+		if (UsingVulkanRenderer)
 		{
-			BRSelection = BitmapRendererSelection::Legacy;
+			if (AndroidEnableMediaPlayerBRIR && FAndroidMisc::IsExternalMemoryAndroidHardwareBufferExtensionLoaded())
+			{
+				BRSelection = BitmapRendererSelection::ImageReader;
+			}
+			else
+			{
+				BRSelection = BitmapRendererSelection::Legacy;
+			}
+		}
+		else
+		{
+			if (AndroidEnableMediaPlayerBRIR)
+			{
+				BRSelection = BitmapRendererSelection::ImageReader;
+			}
+			else
+			{
+				BRSelection = BitmapRendererSelection::Legacy;
+			}
 		}
 
 		JNIEnv* JEnv = FAndroidApplication::GetJavaEnv();
