Add error handling for GetRawInputData size call

Also replace heap allocation with stack allocation that will guarantee alignment of 16.

diff --git a/Engine/Source/Runtime/ApplicationCore/Private/Windows/WindowsApplication.cpp b/Engine/Source/Runtime/ApplicationCore/Private/Windows/WindowsApplication.cpp
--- a/Engine/Source/Runtime/ApplicationCore/Private/Windows/WindowsApplication.cpp
+++ b/Engine/Source/Runtime/ApplicationCore/Private/Windows/WindowsApplication.cpp
@@ -1045,15 +1045,16 @@ namespace UE::WindowsApplication::RawInputThread
 	void FWindowsRawInputRunnable::HandleInputMessage(HRAWINPUT InputHandle)
 	{
 		uint32 Size = 0;
-		::GetRawInputData(InputHandle, RID_INPUT, NULL, &Size, sizeof(RAWINPUTHEADER));
-
-		TUniquePtr<uint8[]> RawData = MakeUnique<uint8[]>(Size);
-
-		if (::GetRawInputData(InputHandle, RID_INPUT, RawData.Get(), &Size, sizeof(RAWINPUTHEADER)) == Size)
+		if (::GetRawInputData(InputHandle, RID_INPUT, NULL, &Size, sizeof(RAWINPUTHEADER)) == 0)
 		{
-			const RAWINPUT* const Raw = (const RAWINPUT* const)RawData.Get();
-			MouseInputs.Enqueue(Raw->data.mouse);
-			++InputCount;
+			void* const RawData = FMemory_Alloca(Size);
+
+			if (::GetRawInputData(InputHandle, RID_INPUT, RawData, &Size, sizeof(RAWINPUTHEADER)) == Size)
+			{
+				const RAWINPUT* const Raw = static_cast<RAWINPUT*>(RawData);
+				MouseInputs.Enqueue(Raw->data.mouse);
+				++InputCount;
+			}
 		}
 	}
 
@@ -2282,43 +2283,44 @@ int32 FWindowsApplication::ProcessMessage( HWND hwnd, uint32 msg, WPARAM wParam,
 		case WM_INPUT:
 			{
 				uint32 Size = 0;
-				::GetRawInputData((HRAWINPUT)lParam, RID_INPUT, NULL, &Size, sizeof(RAWINPUTHEADER));
-
-				TUniquePtr<uint8[]> RawData = MakeUnique<uint8[]>(Size);
-
-				if (::GetRawInputData((HRAWINPUT)lParam, RID_INPUT, RawData.Get(), &Size, sizeof(RAWINPUTHEADER)) == Size )
+				if (::GetRawInputData((HRAWINPUT)lParam, RID_INPUT, NULL, &Size, sizeof(RAWINPUTHEADER)) == 0)
 				{
-					const RAWINPUT* const Raw = (const RAWINPUT* const)RawData.Get();
+					void* const RawData = FMemory_Alloca(Size);
 
-					if (Raw->header.dwType == RIM_TYPEMOUSE) 
+					if (::GetRawInputData((HRAWINPUT)lParam, RID_INPUT, RawData, &Size, sizeof(RAWINPUTHEADER)) == Size)
 					{
-						RawMouseInputResult Result;
-						ProcessRawMouseInput(Raw->data.mouse, Result);
-						if (Result.bIsAbsoluteInput)
+						const RAWINPUT* const Raw = static_cast<RAWINPUT*>(RawData);
+
+						if (Raw->header.dwType == RIM_TYPEMOUSE)
 						{
-							if (bSimulatingHighPrecisionMouseInputForRDP)
+							RawMouseInputResult Result;
+							ProcessRawMouseInput(Raw->data.mouse, Result);
+							if (Result.bIsAbsoluteInput)
 							{
-								if (Result.bSimulateRelativeMovement)
+								if (bSimulatingHighPrecisionMouseInputForRDP)
 								{
-									DeferMessage(CurrentNativeEventWindowPtr, hwnd, msg, wParam, lParam, Result.DeltaX, Result.DeltaY, MOUSE_MOVE_RELATIVE);
+									if (Result.bSimulateRelativeMovement)
+									{
+										DeferMessage(CurrentNativeEventWindowPtr, hwnd, msg, wParam, lParam, Result.DeltaX, Result.DeltaY, MOUSE_MOVE_RELATIVE);
+									}
 								}
+								else
+								{
+									// Since the raw input is coming in as absolute it is likely the user is using a tablet
+									// or perhaps is interacting through a virtual desktop
+									DeferMessage(CurrentNativeEventWindowPtr, hwnd, msg, wParam, lParam, 0, 0, MOUSE_MOVE_ABSOLUTE);
+								}
+								return 1;
 							}
-							else
-							{
-								// Since the raw input is coming in as absolute it is likely the user is using a tablet
-								// or perhaps is interacting through a virtual desktop
-								DeferMessage(CurrentNativeEventWindowPtr, hwnd, msg, wParam, lParam, 0, 0, MOUSE_MOVE_ABSOLUTE);
-							}
-							return 1;
-						}
 
-						// Since raw input is coming in as relative it is likely a traditional mouse device
-						const int xPosRelative = Result.DeltaX;
-						const int yPosRelative = Result.DeltaY;
+							// Since raw input is coming in as relative it is likely a traditional mouse device
+							const int xPosRelative = Result.DeltaX;
+							const int yPosRelative = Result.DeltaY;
 
-						DeferMessage( CurrentNativeEventWindowPtr, hwnd, msg, wParam, lParam, xPosRelative, yPosRelative, MOUSE_MOVE_RELATIVE );
-						return 1;
-					} 
+							DeferMessage(CurrentNativeEventWindowPtr, hwnd, msg, wParam, lParam, xPosRelative, yPosRelative, MOUSE_MOVE_RELATIVE);
+							return 1;
+						}
+					}
 				}
 			}
 			break;
