iOS notification deep link support

diff --git a/Engine/Source/Runtime/ApplicationCore/Private/IOS/IOSAppDelegate.cpp b/Engine/Source/Runtime/ApplicationCore/Private/IOS/IOSAppDelegate.cpp
--- a/Engine/Source/Runtime/ApplicationCore/Private/IOS/IOSAppDelegate.cpp
+++ b/Engine/Source/Runtime/ApplicationCore/Private/IOS/IOSAppDelegate.cpp
@@ -90,6 +90,14 @@ FAutoConsoleVariableRef CVarThreadBlockTime(
 		TEXT("How long to block main IOS thread to make sure gamethread gets time.\n"),
 		ECVF_Default);
 
+FString GNotificationActivationProtocolKey = "activationProtocol";
+FAutoConsoleVariableRef CVarNotificationActivationProtocolKey(
+		TEXT("ios.notification.activationProtocolKey"),
+		GNotificationActivationProtocolKey,
+		TEXT("Notification data payload key that's associated with the activation protocol"),
+		ECVF_Default
+);
+
 
 static void SignalHandler(int32 Signal, struct __siginfo* Info, void* Context)
 {
@@ -958,6 +966,38 @@ static IOSAppDelegate* CachedDelegate = nil;
 	return self.PendingActivationProtocol;
 }
 
+-(void)SetPendingActivationProtocol:(NSString *)NewProtocol
+{
+	static FTSTicker::FDelegateHandle ProtocolActivationTickerHandle;
+	self.PendingActivationProtocol = NewProtocol;
+
+	if ([self.PendingActivationProtocol length] != 0)
+	{
+		if (!ProtocolActivationTickerHandle.IsValid())
+		{
+			ProtocolActivationTickerHandle = FTSTicker::GetCoreTicker().AddTicker(FTickerDelegate::CreateLambda(
+				[self](float)
+				{
+					if (FCoreDelegates::OnActivatedByProtocol.IsBound())
+					{
+						FCoreDelegates::OnActivatedByProtocol.Broadcast(FString([[IOSAppDelegate GetDelegate] PendingActivationProtocol]), PLATFORMUSERID_NONE);
+
+						self.PendingActivationProtocol = @"";
+						ProtocolActivationTickerHandle.Reset();
+						return false; // remove from ticker
+					}
+
+					return true; // try again next tick
+				}
+			));
+		}
+	}
+	else if (ProtocolActivationTickerHandle.IsValid())
+	{
+		FTSTicker::GetCoreTicker().RemoveTicker(ProtocolActivationTickerHandle);
+		ProtocolActivationTickerHandle.Reset();
+	}
+}
 
 - (NSProcessInfoThermalState)GetThermalState
 {
@@ -1167,39 +1207,9 @@ static FAutoConsoleVariableRef CVarGEnableThermalsReport(
 #if WITH_IOS_UNIVERSAL_LINKS
 -(BOOL)application:(UIApplication*)application continueUserActivity:(NSUserActivity*)userActivity restorationHandler:(void(^)(NSArray<id<UIUserActivityRestoring>>* _Nullable)) restorationHandler 
 {
-	static FTSTicker::FDelegateHandle ProtocolActivationTickerHandle;
-	
 	if ([userActivity.activityType isEqualToString:NSUserActivityTypeBrowsingWeb])
 	{
-		self.PendingActivationProtocol = userActivity.webpageURL.absoluteString;
-
-		if ([self.PendingActivationProtocol length] != 0)
-		{
-			if (!ProtocolActivationTickerHandle.IsValid())
-			{
-				ProtocolActivationTickerHandle = FTSTicker::GetCoreTicker().AddTicker(FTickerDelegate::CreateLambda(
-					[self](float)
-					{
-						if (FCoreDelegates::OnActivatedByProtocol.IsBound())
-						{
-							FCoreDelegates::OnActivatedByProtocol.Broadcast(FString([[IOSAppDelegate GetDelegate] PendingActivationProtocol]), PLATFORMUSERID_NONE);
-
-							self.PendingActivationProtocol = @"";
-							ProtocolActivationTickerHandle.Reset();
-							return false; // remove from ticker
-						}
-
-						return true; // try again next tick
-					}
-				));
-			}
-		}
-		else if (ProtocolActivationTickerHandle.IsValid())
-		{
-			FTSTicker::GetCoreTicker().RemoveTicker(ProtocolActivationTickerHandle);
-			ProtocolActivationTickerHandle.Reset();
-		}
-		
+		[self SetPendingActivationProtocol:userActivity.webpageURL.absoluteString];
 		return YES;
 	}
 	return NO;
@@ -1930,6 +1940,14 @@ didReceiveNotificationResponse:(UNNotificationResponse *)response
 			FAppEntry::gLaunchLocalNotificationActivationEvent = FString(activationEvent);
 			FAppEntry::gLaunchLocalNotificationFireDate = FMath::TruncToInt([response.notification.date timeIntervalSince1970]);
 		}
+		
+		// Try to find the activation protocol
+		id ActivationProtocolId = userInfo[(NSString*)FPlatformString::TCHARToCFString(*GNotificationActivationProtocolKey)];
+		if ([ActivationProtocolId isKindOfClass:[NSString class]])
+		{
+			NSString* ActivationProtocol = ActivationProtocolId;
+			[self SetPendingActivationProtocol:ActivationProtocol];
+		}
 	}
 	
 	HandleReceivedNotification(response.notification);
diff --git a/Engine/Source/Runtime/ApplicationCore/Public/IOS/IOSAppDelegate.h b/Engine/Source/Runtime/ApplicationCore/Public/IOS/IOSAppDelegate.h
--- a/Engine/Source/Runtime/ApplicationCore/Public/IOS/IOSAppDelegate.h
+++ b/Engine/Source/Runtime/ApplicationCore/Public/IOS/IOSAppDelegate.h
@@ -242,6 +242,7 @@ APPLICATIONCORE_API
 -(float)GetBackgroundingMainThreadBlockTime;
 -(void)OverrideBackgroundingMainThreadBlockTime:(float)BlockTime;
 -(NSString*)GetPendingActivationProtocol;
+-(void)SetPendingActivationProtocol:(NSString *)NewProtocol;
 
 #if !PLATFORM_TVOS && !PLATFORM_VISIONOS && !defined(SWIFT_IMPORT)
   +(EDeviceScreenOrientation) ConvertFromUIInterfaceOrientation:(UIInterfaceOrientation)Orientation;
