Fixed black screen when playing start movies by forcing MovieStreamer to use BitmapRendererLegacy.
Removed the previous ugly hotfix.

diff --git a/Engine/Build/Android/Java/src/com/epicgames/unreal/MediaPlayer14.java b/Engine/Build/Android/Java/src/com/epicgames/unreal/MediaPlayer14.java
--- a/Engine/Build/Android/Java/src/com/epicgames/unreal/MediaPlayer14.java
+++ b/Engine/Build/Android/Java/src/com/epicgames/unreal/MediaPlayer14.java
@@ -853,7 +853,7 @@ public class MediaPlayer14
 			BitmapRendererImageReader.FrameData data = mBitmapRenderer.updateFrameData(updateParam);
 
 			WaitOnBitmapRender = false;
-			return data.binaryBuffer;
+			return data != null ? data.binaryBuffer : null;
 		}
 		else
 		{
diff --git a/Engine/Plugins/Runtime/AndroidMoviePlayer/Source/AndroidMoviePlayer/Private/AndroidMovieStreamer.cpp b/Engine/Plugins/Runtime/AndroidMoviePlayer/Source/AndroidMoviePlayer/Private/AndroidMovieStreamer.cpp
--- a/Engine/Plugins/Runtime/AndroidMoviePlayer/Source/AndroidMoviePlayer/Private/AndroidMovieStreamer.cpp
+++ b/Engine/Plugins/Runtime/AndroidMoviePlayer/Source/AndroidMoviePlayer/Private/AndroidMovieStreamer.cpp
@@ -28,10 +28,7 @@ FAndroidMediaPlayerStreamer::FAndroidMediaPlayerStreamer()
 	: JavaMediaPlayer(nullptr)
 	, CurrentPosition(-1)
 {
-	JavaMediaPlayer = MakeShareable(new FJavaAndroidMediaPlayer(false, FAndroidMisc::ShouldUseVulkan(), false));
-	// temp hack for restoring the feature of loading movie
-	extern bool ForceAndroidBitmapRendererLegacy;
-	ForceAndroidBitmapRendererLegacy = true;
+	JavaMediaPlayer = MakeShareable(new FJavaAndroidMediaPlayer(false, FAndroidMisc::ShouldUseVulkan(), false, true));
 	MovieViewport = MakeShareable(new FMovieViewport());
 }
 
diff --git a/Engine/Source/Runtime/ApplicationCore/Private/Android/AndroidJavaMediaPlayer.cpp b/Engine/Source/Runtime/ApplicationCore/Private/Android/AndroidJavaMediaPlayer.cpp
--- a/Engine/Source/Runtime/ApplicationCore/Private/Android/AndroidJavaMediaPlayer.cpp
+++ b/Engine/Source/Runtime/ApplicationCore/Private/Android/AndroidJavaMediaPlayer.cpp
@@ -126,7 +126,7 @@ static FAutoConsoleVariableRef CVarAndroidEnableMediaPlayerBRIR(
 	TEXT("Determines how the Android Media Player handles video frame data. True for using ImageReader/HardwareBuffer. False for copying data through the CPU memory. (Default: true)"),
 	ECVF_Default);
 
-FJavaAndroidMediaPlayer::FJavaAndroidMediaPlayer(bool swizzlePixels, bool vulkanRenderer, bool needTrackInfo)
+FJavaAndroidMediaPlayer::FJavaAndroidMediaPlayer(bool swizzlePixels, bool vulkanRenderer, bool needTrackInfo, bool forceBitmapRenderLegacy)
 	: FJavaClassObject(GetClassName(), "(ZZZ)V", swizzlePixels, vulkanRenderer, needTrackInfo)
 	, GetDurationMethod(GetClassMethod("getDuration", "()I"))
 	, ResetMethod(GetClassMethod("reset", "()V"))
@@ -177,7 +177,7 @@ FJavaAndroidMediaPlayer::FJavaAndroidMediaPlayer(bool swizzlePixels, bool vulkan
 		SelectTrackMethod = GetClassMethod("selectTrack", "(I)V");
 	}
 
-	BRSelection = BitmapRendererSelection::Unknown;
+	BRSelection = forceBitmapRenderLegacy ? BitmapRendererSelection::Legacy : BitmapRendererSelection::Unknown;
 	UsingVulkanRenderer = vulkanRenderer && FAndroidMisc::ShouldUseVulkan();
 
 	JNIEnv* JEnv = FAndroidApplication::GetJavaEnv();
@@ -382,19 +382,11 @@ void FJavaAndroidMediaPlayer::SetAudioVolume(float Volume)
 	CallMethod<void>(SetAudioVolumeMethod, Volume);
 }
 
-bool ForceAndroidBitmapRendererLegacy = false;
-
 /** This function has to be called before calling any GetVideoLastXXXXXXX() functions */
 void FJavaAndroidMediaPlayer::SelectBitmapRenderer()
 {
 	if (BRSelection == BitmapRendererSelection::Unknown)
 	{
-		if (ForceAndroidBitmapRendererLegacy)
-		{
-			BRSelection = BitmapRendererSelection::Legacy;
-			ForceAndroidBitmapRendererLegacy = false;
-		}
-		else
 		if (UsingVulkanRenderer)
 		{
 			if (AndroidEnableMediaPlayerBRIR && FAndroidMisc::IsExternalMemoryAndroidHardwareBufferExtensionLoaded())
diff --git a/Engine/Source/Runtime/ApplicationCore/Public/Android/AndroidJavaMediaPlayer.h b/Engine/Source/Runtime/ApplicationCore/Public/Android/AndroidJavaMediaPlayer.h
--- a/Engine/Source/Runtime/ApplicationCore/Public/Android/AndroidJavaMediaPlayer.h
+++ b/Engine/Source/Runtime/ApplicationCore/Public/Android/AndroidJavaMediaPlayer.h
@@ -67,7 +67,7 @@ public:
 	};
 
 public:
-	FJavaAndroidMediaPlayer(bool swizzlePixels, bool vulkanRenderer, bool needTrackInfo);
+	FJavaAndroidMediaPlayer(bool swizzlePixels, bool vulkanRenderer, bool needTrackInfo, bool forceBitmapRenderLegacy = false);
 	virtual ~FJavaAndroidMediaPlayer();
 	int32 GetDuration();
 	void Reset();
