Fixing MetalRHI command buffer errors occuring due to initializing the RHI when the app is backgrounded.

This is not a 100%, but it reduces the chances hit the issue.
Original implementation by [REDACTED_PERSON].

#rnx

diff --git a/Engine/Source/Runtime/ApplicationCore/Private/IOS/IOSAppDelegate.cpp b/Engine/Source/Runtime/ApplicationCore/Private/IOS/IOSAppDelegate.cpp
--- a/Engine/Source/Runtime/ApplicationCore/Private/IOS/IOSAppDelegate.cpp
+++ b/Engine/Source/Runtime/ApplicationCore/Private/IOS/IOSAppDelegate.cpp
@@ -332,6 +332,11 @@ static IOSAppDelegate* CachedDelegate = nil;
 
     FTaskTagScope Scope(ETaskTag::EGameThread);
 
+	while(self.bIsSuspended || [UIApplication sharedApplication].applicationState != UIApplicationStateActive)
+	{
+		usleep(100);
+	}
+	
 	FAppEntry::Init();
 
 	// check for update on app store if cvar is enabled
