Add new function to grab all connected devices paired to user

diff --git a/Engine/Source/Runtime/ApplicationCore/Private/GenericPlatform/GenericPlatformInputDeviceMapper.cpp b/Engine/Source/Runtime/ApplicationCore/Private/GenericPlatform/GenericPlatformInputDeviceMapper.cpp
--- a/Engine/Source/Runtime/ApplicationCore/Private/GenericPlatform/GenericPlatformInputDeviceMapper.cpp
+++ b/Engine/Source/Runtime/ApplicationCore/Private/GenericPlatform/GenericPlatformInputDeviceMapper.cpp
@@ -151,6 +151,21 @@ int32 IPlatformInputDeviceMapper::GetAllConnectedInputDevices(TArray<FInputDevic
 	return OutInputDevices.Num();
 }
 
+int32 IPlatformInputDeviceMapper::GetAllConnectedInputDevicesForUser(const FPlatformUserId UserId, TArray<FInputDeviceId>& OutInputDevices) const
+{
+	OutInputDevices.Reset();
+
+	for (TPair<FInputDeviceId, FPlatformInputDeviceState> MappedDevice : MappedInputDevices)
+	{
+		if (MappedDevice.Value.OwningPlatformUser == UserId && MappedDevice.Value.ConnectionState == EInputDeviceConnectionState::Connected)
+		{
+			OutInputDevices.AddUnique(MappedDevice.Key);
+		}
+	}
+
+	return OutInputDevices.Num();
+}
+
 int32 IPlatformInputDeviceMapper::GetAllActiveUsers(TArray<FPlatformUserId>& OutUsers) const
 {
 	OutUsers.Reset();
diff --git a/Engine/Source/Runtime/ApplicationCore/Public/GenericPlatform/GenericPlatformInputDeviceMapper.h b/Engine/Source/Runtime/ApplicationCore/Public/GenericPlatform/GenericPlatformInputDeviceMapper.h
--- a/Engine/Source/Runtime/ApplicationCore/Public/GenericPlatform/GenericPlatformInputDeviceMapper.h
+++ b/Engine/Source/Runtime/ApplicationCore/Public/GenericPlatform/GenericPlatformInputDeviceMapper.h
@@ -61,6 +61,15 @@ public:
 	 */
 	APPLICATIONCORE_API virtual int32 GetAllConnectedInputDevices(TArray<FInputDeviceId>& OutInputDevices) const;
 
+	/**
+	 * Gather all currently connected input devices for specific user
+	 *
+	 * @param UserId				The Platform User to gather the input devices of.
+	 * @param OutInputDevices		Array of input device ID's that will be populated with the mapped devices.
+	 * @return						The number of mapped devices, INDEX_NONE if the user was not found.
+	 */
+	APPLICATIONCORE_API virtual int32 GetAllConnectedInputDevicesForUser(const FPlatformUserId UserId, TArray<FInputDeviceId>& OutInputDevices) const;
+
 	/**
 	 * Get all currently active platform ids, anyone who has a mapped input device
 	 *
