Generalized IOSAsyncTask to AppleAsyncTask in preparation for using WebKit in the macOS WebBrowser engine plugin with deprecation warnings

diff --git a/Engine/Plugins/Media/AvfMedia/Source/AvfMedia/Private/AvfMediaPrivate.h b/Engine/Plugins/Media/AvfMedia/Source/AvfMedia/Private/AvfMediaPrivate.h
--- a/Engine/Plugins/Media/AvfMedia/Source/AvfMedia/Private/AvfMediaPrivate.h
+++ b/Engine/Plugins/Media/AvfMedia/Source/AvfMedia/Private/AvfMediaPrivate.h
@@ -17,7 +17,7 @@ DECLARE_LOG_CATEGORY_EXTERN(LogAvfMedia, Log, All);
 	#define AVF_GAME_THREAD_BLOCK bool(void)
 	#define AVF_GAME_THREAD_RETURN return true
 	#define AVF_GAME_THREAD_CALL(Block) \
-			[FIOSAsyncTask CreateTaskWithBlock: \
+			[FAppleAsyncTask CreateTaskWithBlock: \
 				Block	\
 			]
 #endif
diff --git a/Engine/Plugins/Online/IOS/OnlineSubsystemIOS/Source/Private/OnlineAchievementsInterfaceIOS.cpp b/Engine/Plugins/Online/IOS/OnlineSubsystemIOS/Source/Private/OnlineAchievementsInterfaceIOS.cpp
--- a/Engine/Plugins/Online/IOS/OnlineSubsystemIOS/Source/Private/OnlineAchievementsInterfaceIOS.cpp
+++ b/Engine/Plugins/Online/IOS/OnlineSubsystemIOS/Source/Private/OnlineAchievementsInterfaceIOS.cpp
@@ -16,7 +16,7 @@ FOnlineAchievementsIOS::FOnlineAchievementsIOS(class FOnlineSubsystemIOS* InSubs
 
 void FOnlineAchievementsIOS::QueryAchievements(const FUniqueNetId& PlayerId, const FOnQueryAchievementsCompleteDelegate& Delegate)
 {
-	// Make a copy of the delegate so it won't be a reference inside the FIOSAsyncTask
+	// Make a copy of the delegate so it won't be a reference inside the FAppleAsyncTask
 	FOnQueryAchievementsCompleteDelegate CopiedDelegate = Delegate;
 
 	dispatch_async(dispatch_get_main_queue(), ^
@@ -29,7 +29,7 @@ void FOnlineAchievementsIOS::QueryAchievements(const FUniqueNetId& PlayerId, con
 			if (bSuccess)
 			{
 				// Write the achievements back to our Achievements array
-				[FIOSAsyncTask CreateTaskWithBlock : ^ bool(void)
+				[FAppleAsyncTask CreateTaskWithBlock : ^ bool(void)
 				{
 					UE_LOG_ONLINE_ACHIEVEMENTS(Display, TEXT("Loaded %d achievements"), [loadedAchievements count]);
 					Achievements.Empty();
@@ -54,7 +54,7 @@ void FOnlineAchievementsIOS::QueryAchievements(const FUniqueNetId& PlayerId, con
 				UE_LOG_ONLINE_ACHIEVEMENTS(Warning, TEXT("Failed to load achievements with error [%d]"), [error code]);
 
 				// Report back to the game thread this failed.
-				[FIOSAsyncTask CreateTaskWithBlock : ^ bool(void)
+				[FAppleAsyncTask CreateTaskWithBlock : ^ bool(void)
 				{
 					CopiedDelegate.ExecuteIfBound(PlayerId, false);
 					return true;
@@ -79,7 +79,7 @@ void FOnlineAchievementsIOS::QueryAchievementDescriptions(const FUniqueNetId& Pl
 			if (bSuccess)
 			{				
 				// Report back to the game thread whether this succeeded.
-				[FIOSAsyncTask CreateTaskWithBlock : ^ bool(void)
+				[FAppleAsyncTask CreateTaskWithBlock : ^ bool(void)
 				{
 					AchievementDescriptions.Empty();
 
@@ -113,7 +113,7 @@ void FOnlineAchievementsIOS::QueryAchievementDescriptions(const FUniqueNetId& Pl
 				UE_LOG_ONLINE_ACHIEVEMENTS(Warning, TEXT("Failed to load achievement descriptions with error [%d]"), [error code]);
 
 				// Report the failure back to the game thread
-				[FIOSAsyncTask CreateTaskWithBlock : ^ bool(void)
+				[FAppleAsyncTask CreateTaskWithBlock : ^ bool(void)
 				{
 					CopiedDelegate.ExecuteIfBound(PlayerId, false);
 					return true;
@@ -126,7 +126,7 @@ void FOnlineAchievementsIOS::QueryAchievementDescriptions(const FUniqueNetId& Pl
 
 void FOnlineAchievementsIOS::WriteAchievements(const FUniqueNetId& PlayerId, FOnlineAchievementsWriteRef& InWriteObject, const FOnAchievementsWrittenDelegate& Delegate)
 {
-	// Make a copy of the delegate so it won't be a reference inside the FIOSAsyncTask
+	// Make a copy of the delegate so it won't be a reference inside the FAppleAsyncTask
 	FOnAchievementsWrittenDelegate CopiedDelegate = Delegate;
 
 	// Hold a reference to the write object for the completion handler to write to.
@@ -195,7 +195,7 @@ void FOnlineAchievementsIOS::WriteAchievements(const FUniqueNetId& PlayerId, FOn
 						// Report any completed achievements to the game thread
 						if ([achievement isCompleted])
 						{
-							[FIOSAsyncTask CreateTaskWithBlock : ^ bool(void)
+							[FAppleAsyncTask CreateTaskWithBlock : ^ bool(void)
 							{
 								TriggerOnAchievementUnlockedDelegates(PlayerId, AchievementId);
 								return true;
@@ -215,7 +215,7 @@ void FOnlineAchievementsIOS::WriteAchievements(const FUniqueNetId& PlayerId, FOn
 				WriteObject->WriteState = (Error == nil) ? EOnlineAsyncTaskState::Done : EOnlineAsyncTaskState::Failed;
 
 				// Report back to the game thread whether this succeeded.
-				[FIOSAsyncTask CreateTaskWithBlock : ^ bool(void)
+				[FAppleAsyncTask CreateTaskWithBlock : ^ bool(void)
 				{
 					CopiedDelegate.ExecuteIfBound(PlayerId, Error == nil);
 					return true;
@@ -283,7 +283,7 @@ bool FOnlineAchievementsIOS::ResetAchievements( const FUniqueNetId& PlayerId )
 			if (bSuccess)
 			{
 				// Wipe out the achievement descriptions back on the game thread.
-				[FIOSAsyncTask CreateTaskWithBlock : ^ bool(void)
+				[FAppleAsyncTask CreateTaskWithBlock : ^ bool(void)
 				{
 					Achievements.Empty();
 					AchievementDescriptions.Empty();
diff --git a/Engine/Plugins/Online/IOS/OnlineSubsystemIOS/Source/Private/OnlineIdentityInterfaceIOS.cpp b/Engine/Plugins/Online/IOS/OnlineSubsystemIOS/Source/Private/OnlineIdentityInterfaceIOS.cpp
--- a/Engine/Plugins/Online/IOS/OnlineSubsystemIOS/Source/Private/OnlineIdentityInterfaceIOS.cpp
+++ b/Engine/Plugins/Online/IOS/OnlineSubsystemIOS/Source/Private/OnlineIdentityInterfaceIOS.cpp
@@ -67,7 +67,7 @@ bool FOnlineIdentityIOS::Login(int32 LocalUserNum, const FOnlineAccountCredentia
 					// Game Center notified a change of state. Identify what happened and handle the change from game thread
 					FGamCenterEvent Event = GetCurrentGameCenterEvent(Error);
 					
-					[FIOSAsyncTask CreateTaskWithBlock : ^ bool(void)
+					[FAppleAsyncTask CreateTaskWithBlock : ^ bool(void)
 					{
 						HandleGamCenterEvent(Event, bLoginInProgress);
 						bLoginInProgress = false;
@@ -89,7 +89,7 @@ bool FOnlineIdentityIOS::Login(int32 LocalUserNum, const FOnlineAccountCredentia
 		{
 			// If a login request is made after authenticateHandler is set just notify the current state
 			// Game Center state updates will always be done through authenticateHandler
-			[FIOSAsyncTask CreateTaskWithBlock : ^ bool(void)
+			[FAppleAsyncTask CreateTaskWithBlock : ^ bool(void)
 			{
 				if (UniqueNetId.IsValid())
 				{
diff --git a/Engine/Plugins/Online/IOS/OnlineSubsystemIOS/Source/Private/OnlineLeaderboardsInterfaceIOS.cpp b/Engine/Plugins/Online/IOS/OnlineSubsystemIOS/Source/Private/OnlineLeaderboardsInterfaceIOS.cpp
--- a/Engine/Plugins/Online/IOS/OnlineSubsystemIOS/Source/Private/OnlineLeaderboardsInterfaceIOS.cpp
+++ b/Engine/Plugins/Online/IOS/OnlineSubsystemIOS/Source/Private/OnlineLeaderboardsInterfaceIOS.cpp
@@ -121,7 +121,7 @@ bool FOnlineLeaderboardsIOS::ReadLeaderboardCompletionDelegate(NSArray* players,
             }
                         
             // Report back to the game thread whether this succeeded.
-            [FIOSAsyncTask CreateTaskWithBlock : ^ bool(void)
+            [FAppleAsyncTask CreateTaskWithBlock : ^ bool(void)
              {
                 ReadObject->ReadState = bWasSuccessful ? EOnlineAsyncTaskState::Done : EOnlineAsyncTaskState::Failed;
                 TriggerOnLeaderboardReadCompleteDelegates(bWasSuccessful);
@@ -333,7 +333,7 @@ bool FOnlineLeaderboardsIOS::FlushLeaderboards(const FName& SessionName)
 							}
 							
 							// Report back to the game thread whether this succeeded.
-							[FIOSAsyncTask CreateTaskWithBlock : ^ bool(void)
+							[FAppleAsyncTask CreateTaskWithBlock : ^ bool(void)
 							 {
 								TriggerOnLeaderboardFlushCompleteDelegates(SessionName, bSucceeded);
 								return true;
diff --git a/Engine/Plugins/Online/IOS/OnlineSubsystemIOS/Source/Private/OnlinePurchaseIOS.cpp b/Engine/Plugins/Online/IOS/OnlineSubsystemIOS/Source/Private/OnlinePurchaseIOS.cpp
--- a/Engine/Plugins/Online/IOS/OnlineSubsystemIOS/Source/Private/OnlinePurchaseIOS.cpp
+++ b/Engine/Plugins/Online/IOS/OnlineSubsystemIOS/Source/Private/OnlinePurchaseIOS.cpp
@@ -187,7 +187,7 @@ struct FOnlinePurchaseInProgressTransactionIOS
 
 	TSharedPtr<TArray<FOnlinePurchaseTransactionIOS>> IosTransactions = [self TranslateTransactions: Transactions];
 
-	[FIOSAsyncTask CreateTaskWithBlock : ^ bool(void)
+	[FAppleAsyncTask CreateTaskWithBlock : ^ bool(void)
 	{
 		if(_PurchaseReceiver)
 		{
@@ -203,7 +203,7 @@ struct FOnlinePurchaseInProgressTransactionIOS
 	
 	TSharedPtr<TArray<FOnlinePurchaseTransactionIOS>> IosTransactions = [self TranslateTransactions: Transactions];
 	
-	[FIOSAsyncTask CreateTaskWithBlock : ^ bool(void)
+	[FAppleAsyncTask CreateTaskWithBlock : ^ bool(void)
 	{
 		if(_PurchaseReceiver)
 		{
@@ -219,7 +219,7 @@ struct FOnlinePurchaseInProgressTransactionIOS
     UE_LOG_ONLINE_PURCHASE(Log, TEXT("FStoreKitPurchaseProxy::restoreCompletedTransactionsFailedWithError %d"), (int32)CompletionState);
 
     // Notifications from SKPaymentQueue are received on main thread. Use a task to notify from game thread
-    [FIOSAsyncTask CreateTaskWithBlock : ^ bool(void)
+    [FAppleAsyncTask CreateTaskWithBlock : ^ bool(void)
     {
         if(_PurchaseReceiver)
         {
@@ -242,7 +242,7 @@ struct FOnlinePurchaseInProgressTransactionIOS
 	TSharedPtr<TArray<FOnlinePurchaseTransactionIOS>> IosTransactions = [self TranslateTransactions: KnownTransactions];
 
 	// Notifications from SKPaymentQueue are received on main thread. Use a task to notify from game thread
-	[FIOSAsyncTask CreateTaskWithBlock : ^ bool(void)
+	[FAppleAsyncTask CreateTaskWithBlock : ^ bool(void)
 	{
 		if(_PurchaseReceiver)
 		{
diff --git a/Engine/Plugins/Online/IOS/OnlineSubsystemIOS/Source/Private/OnlineStoreIOS.cpp b/Engine/Plugins/Online/IOS/OnlineSubsystemIOS/Source/Private/OnlineStoreIOS.cpp
--- a/Engine/Plugins/Online/IOS/OnlineSubsystemIOS/Source/Private/OnlineStoreIOS.cpp
+++ b/Engine/Plugins/Online/IOS/OnlineSubsystemIOS/Source/Private/OnlineStoreIOS.cpp
@@ -48,7 +48,7 @@
 
     // Response for SKRequest is received in a working thread. Use a task to notify from game thread
     FSKProductsRequestHelper* Helper = (FSKProductsRequestHelper*)Request;
-    [FIOSAsyncTask CreateTaskWithBlock : ^ bool(void)
+    [FAppleAsyncTask CreateTaskWithBlock : ^ bool(void)
     {
         if(_StoreReceiver)
         {
@@ -64,7 +64,7 @@
 
     // Response for SKRequest is received in a working thread. Use a task to notify from game thread
     FSKProductsRequestHelper* Helper = (FSKProductsRequestHelper*)Request;
-    [FIOSAsyncTask CreateTaskWithBlock : ^ bool(void)
+    [FAppleAsyncTask CreateTaskWithBlock : ^ bool(void)
     {
         if(_StoreReceiver)
         {
diff --git a/Engine/Plugins/Online/IOS/OnlineSubsystemIOS/Source/Private/OnlineTurnBasedInterfaceIOS.cpp b/Engine/Plugins/Online/IOS/OnlineSubsystemIOS/Source/Private/OnlineTurnBasedInterfaceIOS.cpp
--- a/Engine/Plugins/Online/IOS/OnlineSubsystemIOS/Source/Private/OnlineTurnBasedInterfaceIOS.cpp
+++ b/Engine/Plugins/Online/IOS/OnlineSubsystemIOS/Source/Private/OnlineTurnBasedInterfaceIOS.cpp
@@ -513,7 +513,7 @@ void FOnlineTurnBasedIOS::OnMatchFound(FTurnBasedMatchRef Match)
 	{
 		dispatch_async(dispatch_get_main_queue(), ^
 		{
-			[FIOSAsyncTask CreateTaskWithBlock : ^ bool(void)
+			[FAppleAsyncTask CreateTaskWithBlock : ^ bool(void)
 			{
 				MatchmakerDelegate.Pin()->OnMatchFound(Match);
 				return true;
@@ -527,7 +527,7 @@ void FOnlineTurnBasedIOS::OnMatchEnded(FString MatchID)
 {
 	dispatch_async(dispatch_get_main_queue(), ^
 	{
-        [FIOSAsyncTask CreateTaskWithBlock : ^ bool(void)
+        [FAppleAsyncTask CreateTaskWithBlock : ^ bool(void)
         {
             if (TurnBasedMatchInterfaceObject)
             {
@@ -562,7 +562,7 @@ void FOnlineTurnBasedIOS::OnMatchReceivedTurnEvent(FString MatchID, bool BecameA
 			MatchArray.Add(NewMatch);
 			dispatch_async(dispatch_get_main_queue(), ^
 						   {
-				[FIOSAsyncTask CreateTaskWithBlock : ^ bool(void)
+				[FAppleAsyncTask CreateTaskWithBlock : ^ bool(void)
 				 {
 					if (TurnBasedMatchInterfaceObject)
 					{
diff --git a/Engine/Plugins/Online/OnlineSubsystemFacebook/Source/Private/IOS/FacebookHelper.cpp b/Engine/Plugins/Online/OnlineSubsystemFacebook/Source/Private/IOS/FacebookHelper.cpp
--- a/Engine/Plugins/Online/OnlineSubsystemFacebook/Source/Private/IOS/FacebookHelper.cpp
+++ b/Engine/Plugins/Online/OnlineSubsystemFacebook/Source/Private/IOS/FacebookHelper.cpp
@@ -58,7 +58,7 @@ THIRD_PARTY_INCLUDES_END
     FBSDKAccessToken* OldToken = [[note userInfo] objectForKey:@"FBSDKAccessTokenOld"];
 
     TWeakPtr<FIOSFacebookNotificationDelegate> CapturedOwner = Owner;
-    [FIOSAsyncTask CreateTaskWithBlock : ^ bool(void)
+    [FAppleAsyncTask CreateTaskWithBlock : ^ bool(void)
     {
         // Notify on the game thread
         if(auto PinnedOwner = CapturedOwner.Pin())
@@ -84,7 +84,7 @@ THIRD_PARTY_INCLUDES_END
 #endif
 
     TWeakPtr<FIOSFacebookNotificationDelegate> CapturedOwner = Owner;
-    [FIOSAsyncTask CreateTaskWithBlock : ^ bool(void)
+    [FAppleAsyncTask CreateTaskWithBlock : ^ bool(void)
     {
         // Notify on the game thread
         if(auto PinnedOwner = CapturedOwner.Pin())
@@ -115,7 +115,7 @@ THIRD_PARTY_INCLUDES_END
     FBSDKProfile* OldProfile = [[note userInfo] objectForKey:@"FBSDKProfileOld"];
 
     TWeakPtr<FIOSFacebookNotificationDelegate> CapturedOwner = Owner;
-    [FIOSAsyncTask CreateTaskWithBlock : ^ bool(void)
+    [FAppleAsyncTask CreateTaskWithBlock : ^ bool(void)
     {
         // Notify on the game thread
         if(auto PinnedOwner = CapturedOwner.Pin())
diff --git a/Engine/Plugins/Online/OnlineSubsystemFacebook/Source/Private/IOS/OnlineIdentityFacebook.cpp b/Engine/Plugins/Online/OnlineSubsystemFacebook/Source/Private/IOS/OnlineIdentityFacebook.cpp
--- a/Engine/Plugins/Online/OnlineSubsystemFacebook/Source/Private/IOS/OnlineIdentityFacebook.cpp
+++ b/Engine/Plugins/Online/OnlineSubsystemFacebook/Source/Private/IOS/OnlineIdentityFacebook.cpp
@@ -185,7 +185,7 @@ bool FOnlineIdentityFacebook::Login(int32 LocalUserNum, const FOnlineAccountCred
                             DeclinedPermissions.Add(permission);
                         }
 
-						[FIOSAsyncTask CreateTaskWithBlock : ^ bool(void)
+						[FAppleAsyncTask CreateTaskWithBlock : ^ bool(void)
 						{
 							// Trigger this on the game thread
 							if (bSuccessfulLogin)
@@ -228,7 +228,7 @@ bool FOnlineIdentityFacebook::Login(int32 LocalUserNum, const FOnlineAccountCred
 				if (bIsUsingClassicLogin)
 				{
 					const FString AccessToken([FBSDKAccessToken.currentAccessToken tokenString]);
-					[FIOSAsyncTask CreateTaskWithBlock : ^ bool(void)
+					[FAppleAsyncTask CreateTaskWithBlock : ^ bool(void)
 					 {
 						Login(LocalUserNum, FString(AccessToken));
 						return true;
@@ -236,7 +236,7 @@ bool FOnlineIdentityFacebook::Login(int32 LocalUserNum, const FOnlineAccountCred
 				}
 				else
 				{
-					[FIOSAsyncTask CreateTaskWithBlock : ^ bool(void)
+					[FAppleAsyncTask CreateTaskWithBlock : ^ bool(void)
 					 {
 						LoginLimited(LocalUserNum);
 						return true;
@@ -312,7 +312,7 @@ void FOnlineIdentityFacebook::OnLoginAttemptComplete(int32 LocalUserNum, bool bS
 			FBSDKLoginManager* loginManager = [[FBSDKLoginManager alloc] init];
 			[loginManager logOut];
 
-			[FIOSAsyncTask CreateTaskWithBlock : ^ bool(void)
+			[FAppleAsyncTask CreateTaskWithBlock : ^ bool(void)
 			 {
 				// Trigger this on the game thread
 				UE_LOG_ONLINE_IDENTITY(Display, TEXT("Facebook login failed: %s"), *NewErrorStr);
@@ -347,7 +347,7 @@ bool FOnlineIdentityFacebook::Logout(int32 LocalUserNum)
 			FBSDKLoginManager* loginManager = [[FBSDKLoginManager alloc] init];
 			[loginManager logOut];
 
-			[FIOSAsyncTask CreateTaskWithBlock : ^ bool(void)
+			[FAppleAsyncTask CreateTaskWithBlock : ^ bool(void)
 			{
 				// Trigger this on the game thread
 				FUniqueNetIdPtr UserId = GetUniquePlayerId(LocalUserNum);
diff --git a/Engine/Plugins/Online/OnlineSubsystemFacebook/Source/Private/IOS/OnlineSharingFacebook.cpp b/Engine/Plugins/Online/OnlineSubsystemFacebook/Source/Private/IOS/OnlineSharingFacebook.cpp
--- a/Engine/Plugins/Online/OnlineSubsystemFacebook/Source/Private/IOS/OnlineSharingFacebook.cpp
+++ b/Engine/Plugins/Online/OnlineSubsystemFacebook/Source/Private/IOS/OnlineSharingFacebook.cpp
@@ -87,7 +87,7 @@ bool FOnlineSharingFacebook::RequestNewReadPermissions(int32 LocalUserNum, EOnli
 											   completion: ^(FBSDKLoginManagerLoginResult* Result, NSError* Error)
                         {
                             UE_LOG_ONLINE_SHARING(Display, TEXT("logInFromViewController : Success - %d"), Error == nil);
-                            [FIOSAsyncTask CreateTaskWithBlock : ^ bool(void)
+                            [FAppleAsyncTask CreateTaskWithBlock : ^ bool(void)
                             {
                                 if (Error == nil)
                                 {
@@ -161,7 +161,7 @@ bool FOnlineSharingFacebook::RequestNewPublishPermissions(int32 LocalUserNum, EO
 											   completion: ^(FBSDKLoginManagerLoginResult* Result, NSError* Error)
 						{
                             UE_LOG_ONLINE_SHARING(Display, TEXT("logInWithPublishPermissions : Success - %d"), Error == nil);
-                            [FIOSAsyncTask CreateTaskWithBlock : ^ bool(void)
+                            [FAppleAsyncTask CreateTaskWithBlock : ^ bool(void)
                             {
                                 if (Error == nil)
                                 {
diff --git a/Engine/Plugins/Online/OnlineSubsystemGoogle/Source/Private/IOS/GoogleHelper.cpp b/Engine/Plugins/Online/OnlineSubsystemGoogle/Source/Private/IOS/GoogleHelper.cpp
--- a/Engine/Plugins/Online/OnlineSubsystemGoogle/Source/Private/IOS/GoogleHelper.cpp
+++ b/Engine/Plugins/Online/OnlineSubsystemGoogle/Source/Private/IOS/GoogleHelper.cpp
@@ -179,7 +179,7 @@ bool GetAuthTokenFromSignInResult(GIDGoogleUser* User, NSString* ServerAuthCode,
 
 	UE_LOG_ONLINE_IDENTITY(Display, TEXT("SignIn: %s"), *SignInData.ToDebugString());
 
-	[FIOSAsyncTask CreateTaskWithBlock : ^ bool(void)
+	[FAppleAsyncTask CreateTaskWithBlock : ^ bool(void)
 	{
 		// Notify on the game thread
 		_OnSignInComplete.Broadcast(SignInData);
@@ -199,7 +199,7 @@ bool GetAuthTokenFromSignInResult(GIDGoogleUser* User, NSString* ServerAuthCode,
 	{
 		SignOutData.Response = EGoogleLoginResponse::RESPONSE_OK;
 	}
-	[FIOSAsyncTask CreateTaskWithBlock : ^ bool(void)
+	[FAppleAsyncTask CreateTaskWithBlock : ^ bool(void)
 	{
 		UE_LOG_ONLINE_IDENTITY(Display, TEXT("didDisconnectWithUser Complete: %s"), *SignOutData.ToDebugString());
 
diff --git a/Engine/Plugins/Online/WebAuth/Source/Private/IOS/IOSPlatformWebAuth.cpp b/Engine/Plugins/Online/WebAuth/Source/Private/IOS/IOSPlatformWebAuth.cpp
--- a/Engine/Plugins/Online/WebAuth/Source/Private/IOS/IOSPlatformWebAuth.cpp
+++ b/Engine/Plugins/Online/WebAuth/Source/Private/IOS/IOSPlatformWebAuth.cpp
@@ -49,7 +49,7 @@ bool FIOSWebAuth::AuthSessionWithURL(const FString &UrlStr, const FString &Schem
 	dispatch_async(dispatch_get_main_queue(), ^ {
 		ASWebAuthenticationSession* SavedAuthSession = [[ASWebAuthenticationSession alloc] initWithURL:Url callbackURLScheme: Scheme completionHandler:^(NSURL * _Nullable callbackURL, NSError * _Nullable error)
 		{
-			[FIOSAsyncTask CreateTaskWithBlock : ^ bool(void)
+			[FAppleAsyncTask CreateTaskWithBlock : ^ bool(void)
 			 {
 				FIOSCoreDelegates::OnOpenURL.Remove(OpenUrlHandle);
 				// Response received
@@ -71,7 +71,7 @@ bool FIOSWebAuth::AuthSessionWithURL(const FString &UrlStr, const FString &Schem
 		check(PresentationContextProvider);
 		SavedAuthSession.presentationContextProvider = PresentationContextProvider;
 		
-		[FIOSAsyncTask CreateTaskWithBlock : ^ bool(void)
+		[FAppleAsyncTask CreateTaskWithBlock : ^ bool(void)
 		{
 			OpenUrlHandle = FIOSCoreDelegates::OnOpenURL.AddLambda([this, AppScheme = FString(Scheme), SavedAuthSession](UIApplication* application, NSURL* url, NSString* sourceApplication, id annotation)
 			{
@@ -85,7 +85,7 @@ bool FIOSWebAuth::AuthSessionWithURL(const FString &UrlStr, const FString &Schem
 					[SavedAuthSession cancel];
 					[SavedAuthSession release];
 					
-					[FIOSAsyncTask CreateTaskWithBlock : ^ bool(void)
+					[FAppleAsyncTask CreateTaskWithBlock : ^ bool(void)
 					 {
 						FIOSCoreDelegates::OnOpenURL.Remove(OpenUrlHandle);
 						const char* StrCallbackURL = [url.absoluteString UTF8String];
diff --git a/Engine/Source/Runtime/ApplicationCore/Private/IOS/IOSAppDelegate.cpp b/Engine/Source/Runtime/ApplicationCore/Private/IOS/IOSAppDelegate.cpp
--- a/Engine/Source/Runtime/ApplicationCore/Private/IOS/IOSAppDelegate.cpp
+++ b/Engine/Source/Runtime/ApplicationCore/Private/IOS/IOSAppDelegate.cpp
@@ -467,7 +467,7 @@ static IOSAppDelegate* CachedDelegate = nil;
         }
 
         // drain the async task queue from the game thread
-        [FIOSAsyncTask ProcessAsyncTasks];
+        [FAppleAsyncTask ProcessAsyncTasks];
 	}
 
     dispatch_sync(dispatch_get_main_queue(),^
@@ -577,7 +577,7 @@ static IOSAppDelegate* CachedDelegate = nil;
 		UE_CLOG(InterruptionType == AVAudioSessionInterruptionTypeBegan, LogIOSAudioSession, Log, TEXT("AudioSession interruption notification: Began"));
 		UE_CLOG(InterruptionType == AVAudioSessionInterruptionTypeEnded, LogIOSAudioSession, Log, TEXT("AudioSession interruption notification: Ended"));
 
-		[FIOSAsyncTask CreateTaskWithBlock : ^ bool(void)
+		[FAppleAsyncTask CreateTaskWithBlock : ^ bool(void)
 		{
 			// the audio context should resume immediately after interrupt, if suspended
 			FAppEntry::ResetAudioContextResumeTime();
@@ -606,7 +606,7 @@ static IOSAppDelegate* CachedDelegate = nil;
 	[[NSNotificationCenter defaultCenter] addObserverForName:AVAudioSessionRouteChangeNotification object : nil queue : nil usingBlock : ^ (NSNotification *notification)
 	{
 		NSUInteger RouteChangeReason = [[[notification userInfo] objectForKey:AVAudioSessionRouteChangeReasonKey] unsignedIntegerValue];
-		[FIOSAsyncTask CreateTaskWithBlock : ^ bool(void)
+		[FAppleAsyncTask CreateTaskWithBlock : ^ bool(void)
 		{
 			switch (RouteChangeReason)
 			{
@@ -1461,7 +1461,7 @@ static FAutoConsoleVariableRef CVarGEnableThermalsReport(
 	}
 	else
 	{
-		[FIOSAsyncTask CreateTaskWithBlock : ^ bool(void)
+		[FAppleAsyncTask CreateTaskWithBlock : ^ bool(void)
 		{
 			FIOSApplication* Application = [IOSAppDelegate GetDelegate].IOSApplication;
 			Application->OrientationChanged(Orientation);
@@ -1647,7 +1647,7 @@ FCriticalSection RenderSuspend;
         [[UIApplication sharedApplication] endBackgroundTask: BackgroundTaskId];
     }];
 	
-	[FIOSAsyncTask CreateTaskWithBlock : ^ bool(void)
+	[FAppleAsyncTask CreateTaskWithBlock : ^ bool(void)
 	{
 		// the audio context should resume immediately after interrupt, if suspended
 		FAppEntry::ResetAudioContextResumeTime();
@@ -1668,7 +1668,7 @@ FCriticalSection RenderSuspend;
 	/*
 	 Called as part of the transition from the background to the inactive state; here you can undo many of the changes made on entering the background.
 	 */
-	[FIOSAsyncTask CreateTaskWithBlock : ^ bool(void)
+	[FAppleAsyncTask CreateTaskWithBlock : ^ bool(void)
 	{
 		// the audio context should resume immediately after interrupt, if suspended
 		FAppEntry::ResetAudioContextResumeTime();
@@ -2090,7 +2090,7 @@ CORE_API bool IOSShowDashboardUI()
     
     // Battery level is from 0.0 to 1.0, get it in terms of 0-100
     self.BatteryLevel = ((int)([Device batteryLevel] * 100));
-	[FIOSAsyncTask CreateTaskWithBlock : ^ bool(void)
+	[FAppleAsyncTask CreateTaskWithBlock : ^ bool(void)
 	{
 		UE_LOG(LogIOS, Display, TEXT("Battery Level Changed: %d"), self.BatteryLevel);
 		return true;
@@ -2104,7 +2104,7 @@ CORE_API bool IOSShowDashboardUI()
     UIDevice* Device = [UIDevice currentDevice];
     UIDeviceBatteryState State = Device.batteryState;
     self.bBatteryState = State == UIDeviceBatteryStateUnplugged || State == UIDeviceBatteryStateUnknown;
-	[FIOSAsyncTask CreateTaskWithBlock : ^ bool(void)
+	[FAppleAsyncTask CreateTaskWithBlock : ^ bool(void)
 	{
 		UE_LOG(LogIOS, Display, TEXT("Battery State Changed: %d"), self.bBatteryState);
 		return true;
@@ -2127,7 +2127,7 @@ CORE_API bool IOSShowDashboardUI()
         case NSProcessInfoThermalStateCritical:	Severity = FCoreDelegates::ETemperatureSeverity::Critical; Level = TEXT("Critical"); break;
     }
     
-    [FIOSAsyncTask CreateTaskWithBlock : ^ bool(void)
+    [FAppleAsyncTask CreateTaskWithBlock : ^ bool(void)
      {
         UE_LOG(LogIOS, Display, TEXT("Temperature Changed: %s"), *Level);
         FCoreDelegates::OnTemperatureChange.Broadcast(Severity);
@@ -2139,7 +2139,7 @@ CORE_API bool IOSShowDashboardUI()
 -(void)lowPowerModeChanged:(NSNotification *)notification
 {
 #if !PLATFORM_TVOS
-    [FIOSAsyncTask CreateTaskWithBlock : ^ bool(void)
+    [FAppleAsyncTask CreateTaskWithBlock : ^ bool(void)
      {
         bool bInLowPowerMode = [[NSProcessInfo processInfo] isLowPowerModeEnabled];
         UE_LOG(LogIOS, Display, TEXT("Low Power Mode Changed: %d"), bInLowPowerMode);
diff --git a/Engine/Source/Runtime/ApplicationCore/Private/IOS/IOSApplication.cpp b/Engine/Source/Runtime/ApplicationCore/Private/IOS/IOSApplication.cpp
--- a/Engine/Source/Runtime/ApplicationCore/Private/IOS/IOSApplication.cpp
+++ b/Engine/Source/Runtime/ApplicationCore/Private/IOS/IOSApplication.cpp
@@ -227,7 +227,7 @@ TSharedRef< FGenericWindow > FIOSApplication::MakeWindow()
 void FIOSApplication::OrientationChanged(UIInterfaceOrientation orientation)
 {
 	// queue up the size as we see it now all the way to the RHI 
-	[FIOSAsyncTask CreateTaskWithBlock : ^ bool(void)
+	[FAppleAsyncTask CreateTaskWithBlock : ^ bool(void)
 	{
 		__block FPlatformRect WindowRect;
 		dispatch_sync(dispatch_get_main_queue(), ^{
diff --git a/Engine/Source/Runtime/ApplicationCore/Private/IOS/IOSView.cpp b/Engine/Source/Runtime/ApplicationCore/Private/IOS/IOSView.cpp
--- a/Engine/Source/Runtime/ApplicationCore/Private/IOS/IOSView.cpp
+++ b/Engine/Source/Runtime/ApplicationCore/Private/IOS/IOSView.cpp
@@ -1021,7 +1021,7 @@ self.accessibilityElements = @[Window.accessibilityContainer];
 	ScreenRect.Left = FMath::TruncToInt(Frame.origin.x);
 	ScreenRect.Right = FMath::TruncToInt(Frame.origin.x + Frame.size.width);
 	
-	[FIOSAsyncTask CreateTaskWithBlock:^bool(void){
+	[FAppleAsyncTask CreateTaskWithBlock:^bool(void){
 		[IOSAppDelegate GetDelegate].IOSApplication->OnVirtualKeyboardShown().Broadcast(ScreenRect);
 		return true;
 	 }];
@@ -1030,7 +1030,7 @@ self.accessibilityElements = @[Window.accessibilityContainer];
 // Called when the UIKeyboardWillHideNotification is sent
 - (void)keyboardWillBeHidden:(NSNotification*)aNotification
 {
-	[FIOSAsyncTask CreateTaskWithBlock:^bool(void){
+	[FAppleAsyncTask CreateTaskWithBlock:^bool(void){
 		[IOSAppDelegate GetDelegate].IOSApplication->OnVirtualKeyboardHidden().Broadcast();
 		return true;
 	 }];
diff --git a/Engine/Source/Runtime/Core/Private/Apple/AppleAsyncTask.cpp b/Engine/Source/Runtime/Core/Private/Apple/AppleAsyncTask.cpp
new file mode 100644
--- /dev/null
+++ b/Engine/Source/Runtime/Core/Private/Apple/AppleAsyncTask.cpp
@@ -0,0 +1,136 @@
+// Copyright Epic Games, Inc. All Rights Reserved.
+
+#include "Apple/AppleAsyncTask.h"
+#include "HAL/PlatformAtomics.h"
+
+@implementation FAppleAsyncTask
+
+@synthesize UserData;
+@synthesize GameThreadCallback;
+
+/** All currently running tasks (which can be created on iOS thread or main thread) */
+NSMutableArray* RunningTasks;
+
+/**
+ * Static class constructor, called before any instances are created
+ */
++ (void)initialize
+{
+	// create the RunningTasks object one time
+	RunningTasks = [[NSMutableArray arrayWithCapacity:4] retain];
+}
+
+/**
+ * Initialize the async task
+ */
+- (id)init
+{
+	self = [super init];
+
+	// add ourself to the list of tasks
+	@synchronized(RunningTasks)
+	{
+		[RunningTasks addObject:self];
+	}
+
+	// return ourself, the constructed object
+	return self;
+}
+
+
++ (void)CreateTaskWithBlock:(bool (^)(void))Block
+{
+	// create a task, and add it to the array
+	FAppleAsyncTask* Task = [[FAppleAsyncTask alloc] init];
+	// set the callback
+	Task.GameThreadCallback = Block;
+	// safely tell the game thread we are ready to go
+	[Task FinishedTask];
+}
+
+- (void)FinishedTask
+{
+	FPlatformAtomics::InterlockedIncrement(&bIsReadyForGameThread);
+}
+
+/**
+ * Check for completion
+ *
+ * @return TRUE if we succeeded (the completion block will have been called)
+ */
+- (bool)CheckForCompletion
+{
+	// handle completion
+	if (bIsReadyForGameThread)
+	{
+		// call the game thread block
+		if (GameThreadCallback)
+		{
+			if (GameThreadCallback())
+			{
+				// only return true if the callback says it's complete
+				return true;
+			}
+		}
+		else
+		{
+			// if there isn't a callback, then just return TRUE to remove the 
+			// task from the queue
+			return true;
+		}
+		
+	}
+
+	// all other cases, we are not complete
+	return false;
+}
+
+/**
+ * Tick all currently running tasks
+ */
++ (void)ProcessAsyncTasks
+{
+	FAppleAsyncTask* CurrentTask = nil;
+
+	// grab one out of the queue to process outside of the lock
+	@synchronized(RunningTasks)
+	{
+		if ([RunningTasks count] > 0)
+		{
+			CurrentTask = [RunningTasks objectAtIndex:0];
+			[RunningTasks removeObjectAtIndex:0];
+		}
+	}
+
+	if (CurrentTask != nil)
+	{
+		if ([CurrentTask CheckForCompletion])
+		{
+			// release the object count
+			[CurrentTask release];
+		}
+		else
+		{
+			// if it's not done, but it back in the list
+			@synchronized(RunningTasks)
+			{
+				[RunningTasks addObject:CurrentTask];
+			}
+		}
+	}
+}
+
+
+/** 
+ * Application destructor
+ */
+- (void)dealloc 
+{
+	[GameThreadCallback release];
+	self.UserData = nil;
+
+	[super dealloc];
+}
+
+
+@end
diff --git a/Engine/Source/Runtime/Core/Private/IOS/IOSAsyncTask.cpp b/Engine/Source/Runtime/Core/Private/IOS/IOSAsyncTask.cpp
--- a/Engine/Source/Runtime/Core/Private/IOS/IOSAsyncTask.cpp
+++ b/Engine/Source/Runtime/Core/Private/IOS/IOSAsyncTask.cpp
@@ -1,136 +1,6 @@
 // Copyright Epic Games, Inc. All Rights Reserved.
 
 #include "IOS/IOSAsyncTask.h"
-#include "HAL/PlatformAtomics.h"
 
 @implementation FIOSAsyncTask
-
-@synthesize UserData;
-@synthesize GameThreadCallback;
-
-/** All currently running tasks (which can be created on iOS thread or main thread) */
-NSMutableArray* RunningTasks;
-
-/**
- * Static class constructor, called before any instances are created
- */
-+ (void)initialize
-{
-	// create the RunningTasks object one time
-	RunningTasks = [[NSMutableArray arrayWithCapacity:4] retain];
-}
-
-/**
- * Initialize the async task
- */
-- (id)init
-{
-	self = [super init];
-
-	// add ourself to the list of tasks
-	@synchronized(RunningTasks)
-	{
-		[RunningTasks addObject:self];
-	}
-
-	// return ourself, the constructed object
-	return self;
-}
-
-
-+ (void)CreateTaskWithBlock:(bool (^)(void))Block
-{
-	// create a task, and add it to the array
-	FIOSAsyncTask* Task = [[FIOSAsyncTask alloc] init];
-	// set the callback
-	Task.GameThreadCallback = Block;
-	// safely tell the game thread we are ready to go
-	[Task FinishedTask];
-}
-
-- (void)FinishedTask
-{
-	FPlatformAtomics::InterlockedIncrement(&bIsReadyForGameThread);
-}
-
-/**
- * Check for completion
- *
- * @return TRUE if we succeeded (the completion block will have been called)
- */
-- (bool)CheckForCompletion
-{
-	// handle completion
-	if (bIsReadyForGameThread)
-	{
-		// call the game thread block
-		if (GameThreadCallback)
-		{
-			if (GameThreadCallback())
-			{
-				// only return true if the callback says it's complete
-				return true;
-			}
-		}
-		else
-		{
-			// if there isn't a callback, then just return TRUE to remove the 
-			// task from the queue
-			return true;
-		}
-		
-	}
-
-	// all other cases, we are not complete
-	return false;
-}
-
-/**
- * Tick all currently running tasks
- */
-+ (void)ProcessAsyncTasks
-{
-	FIOSAsyncTask* CurrentTask = nil;
-
-	// grab one out of the queue to process outside of the lock
-	@synchronized(RunningTasks)
-	{
-		if ([RunningTasks count] > 0)
-		{
-			CurrentTask = [RunningTasks objectAtIndex:0];
-			[RunningTasks removeObjectAtIndex:0];
-		}
-	}
-
-	if (CurrentTask != nil)
-	{
-		if ([CurrentTask CheckForCompletion])
-		{
-			// release the object count
-			[CurrentTask release];
-		}
-		else
-		{
-			// if it's not done, but it back in the list
-			@synchronized(RunningTasks)
-			{
-				[RunningTasks addObject:CurrentTask];
-			}
-		}
-	}
-}
-
-
-/** 
- * Application destructor
- */
-- (void)dealloc 
-{
-	[GameThreadCallback release];
-	self.UserData = nil;
-
-	[super dealloc];
-}
-
-
 @end
diff --git a/Engine/Source/Runtime/Core/Public/Apple/AppleAsyncTask.h b/Engine/Source/Runtime/Core/Public/Apple/AppleAsyncTask.h
new file mode 100644
--- /dev/null
+++ b/Engine/Source/Runtime/Core/Public/Apple/AppleAsyncTask.h
@@ -0,0 +1,48 @@
+// Copyright Epic Games, Inc. All Rights Reserved.
+
+#pragma once
+
+#include "CoreTypes.h"
+
+#include "CoreTypes.h"
+#import <Foundation/Foundation.h>
+
+
+@interface FAppleAsyncTask : NSObject
+{
+@private
+	/** Whether or not the task is ready to have GameThread callback called (set on iOS thread) */
+	int32 bIsReadyForGameThread;
+}
+
+/** Extra data for this async task */
+@property (retain) id UserData;
+
+/**
+ * Code to run on the game thread when the async task completes 
+ * @return true when the task is complete, and it will be pulled from the list of tasks and destroyed
+ */
+@property (copy) bool (^GameThreadCallback)(void);
+
+
+/**
+ * Create an async task object, set the block, and mark it as ready to be processed on main thread
+ * For advanced uses (ie, setting UserData or delaying when it's ready for game thread, use the usual [[alloc] init] pattern
+ *
+ * Note, this doesn't return an AsyncObject so no one is tempted to use it - once it's ready for game thread, it's dangerous to touch it
+ * This is callable on any thread, not just main thread
+ */
++ (void)CreateTaskWithBlock:(bool (^)(void))Block;
+
+/**
+ * Mark that the task is complete on the iOS thread, and now the
+ * GameThread can be fired (the Task is unsafe to use after this call)
+ */
+- (void)FinishedTask;
+
+/**
+ * Tick all currently running tasks
+ */
++ (void)ProcessAsyncTasks;
+
+@end
diff --git a/Engine/Source/Runtime/Core/Public/IOS/IOSAsyncTask.h b/Engine/Source/Runtime/Core/Public/IOS/IOSAsyncTask.h
--- a/Engine/Source/Runtime/Core/Public/IOS/IOSAsyncTask.h
+++ b/Engine/Source/Runtime/Core/Public/IOS/IOSAsyncTask.h
@@ -2,47 +2,10 @@
 
 #pragma once
 
-#include "CoreTypes.h"
+#include "Apple/AppleAsyncTask.h"
 
-#include "CoreTypes.h"
-#import <Foundation/Foundation.h>
-
-
-@interface FIOSAsyncTask : NSObject
+UE_DEPRECATED(5.7, "Use the more general FAppleAsyncTask instead")
+@interface FIOSAsyncTask : FAppleAsyncTask
 {
-@private
-	/** Whether or not the task is ready to have GameThread callback called (set on iOS thread) */
-	int32 bIsReadyForGameThread;
 }
-
-/** Extra data for this async task */
-@property (retain) id UserData;
-
-/**
- * Code to run on the game thread when the async task completes 
- * @return true when the task is complete, and it will be pulled from the list of tasks and destroyed
- */
-@property (copy) bool (^GameThreadCallback)(void);
-
-
-/**
- * Create an async task object, set the block, and mark it as ready to be processed on main thread
- * For advanced uses (ie, setting UserData or delaying when it's ready for game thread, use the usual [[alloc] init] pattern
- *
- * Note, this doesn't return an AsyncObject so no one is tempted to use it - once it's ready for game thread, it's dangerous to touch it
- * This is callable on any thread, not just main thread
- */
-+ (void)CreateTaskWithBlock:(bool (^)(void))Block;
-
-/**
- * Mark that the task is complete on the iOS thread, and now the
- * GameThread can be fired (the Task is unsafe to use after this call)
- */
-- (void)FinishedTask;
-
-/**
- * Tick all currently running tasks
- */
-+ (void)ProcessAsyncTasks;
-
 @end
diff --git a/Engine/Source/Runtime/PreLoadScreen/Private/PreLoadScreenManager.cpp b/Engine/Source/Runtime/PreLoadScreen/Private/PreLoadScreenManager.cpp
--- a/Engine/Source/Runtime/PreLoadScreen/Private/PreLoadScreenManager.cpp
+++ b/Engine/Source/Runtime/PreLoadScreen/Private/PreLoadScreenManager.cpp
@@ -765,6 +765,6 @@ void FPreLoadScreenManager::Android_PlatformSpecificGameLogicFrameTick()
 void FPreLoadScreenManager::IOS_PlatformSpecificGameLogicFrameTick()
 {
 	// drain the async task queue from the game thread
-	[FIOSAsyncTask ProcessAsyncTasks];
+	[FAppleAsyncTask ProcessAsyncTasks];
 }
 #endif //PLATFORM_IOS
diff --git a/Engine/Source/Runtime/Slate/Private/Framework/Text/IOS/IOSPlatformTextField.cpp b/Engine/Source/Runtime/Slate/Private/Framework/Text/IOS/IOSPlatformTextField.cpp
--- a/Engine/Source/Runtime/Slate/Private/Framework/Text/IOS/IOSPlatformTextField.cpp
+++ b/Engine/Source/Runtime/Slate/Private/Framework/Text/IOS/IOSPlatformTextField.cpp
@@ -239,7 +239,7 @@ void FIOSPlatformTextField::ShowVirtualKeyboard(bool bShow, int32 UserIndex, TSh
 												TextEntry = FText::FromString(AlertTextField.text);
 												AlertController = nil;
 											
-												FIOSAsyncTask* AsyncTask = [[FIOSAsyncTask alloc] init];
+												FAppleAsyncTask* AsyncTask = [[FAppleAsyncTask alloc] init];
 												AsyncTask.GameThreadCallback = ^ bool(void)
 												{
 													if(TextWidget.IsValid())
@@ -275,7 +275,7 @@ void FIOSPlatformTextField::ShowVirtualKeyboard(bool bShow, int32 UserIndex, TSh
 												}];
 												AlertController = nil;
 											
-												FIOSAsyncTask* AsyncTask = [[FIOSAsyncTask alloc] init];
+												FAppleAsyncTask* AsyncTask = [[FAppleAsyncTask alloc] init];
 												AsyncTask.GameThreadCallback = ^ bool(void)
 												{
 													// clear the TextWidget
diff --git a/Engine/Source/Runtime/WebBrowser/Private/IOS/IOSCookieManager.cpp b/Engine/Source/Runtime/WebBrowser/Private/IOS/IOSCookieManager.cpp
--- a/Engine/Source/Runtime/WebBrowser/Private/IOS/IOSCookieManager.cpp
+++ b/Engine/Source/Runtime/WebBrowser/Private/IOS/IOSCookieManager.cpp
@@ -43,7 +43,7 @@ void FIOSCookieManager::DeleteCookies(const FString& URL, const FString& CookieN
 		}
 		
 		// Notify on the game thread
-		[FIOSAsyncTask CreateTaskWithBlock:^bool(void)
+		[FAppleAsyncTask CreateTaskWithBlock:^bool(void)
 		 {
 			 if (Completed)
 			 {
diff --git a/Engine/Source/Runtime/WebBrowser/Private/IOS/IOSPlatformWebBrowser.cpp b/Engine/Source/Runtime/WebBrowser/Private/IOS/IOSPlatformWebBrowser.cpp
--- a/Engine/Source/Runtime/WebBrowser/Private/IOS/IOSPlatformWebBrowser.cpp
+++ b/Engine/Source/Runtime/WebBrowser/Private/IOS/IOSPlatformWebBrowser.cpp
@@ -367,7 +367,7 @@ class SIOSWebBrowserWidget : public SLeafWidget
 			TWeakPtr<FWebBrowserWindow> AsyncWebBrowserWindowPtr = WebBrowserWindowPtr;
 
 			// Notify on the game thread
-			[FIOSAsyncTask CreateTaskWithBlock : ^ bool(void)
+			[FAppleAsyncTask CreateTaskWithBlock : ^ bool(void)
 			{
 				TSharedPtr<FWebBrowserWindow> BrowserWindow = AsyncWebBrowserWindowPtr.Pin();
 				if (BrowserWindow.IsValid())
@@ -407,7 +407,7 @@ class SIOSWebBrowserWidget : public SLeafWidget
 			FString Message = InMessage;
 			TWeakPtr<FWebBrowserWindow> AsyncWebBrowserWindowPtr = WebBrowserWindowPtr;
 
-			[FIOSAsyncTask CreateTaskWithBlock : ^ bool(void)
+			[FAppleAsyncTask CreateTaskWithBlock : ^ bool(void)
 			{
 				TSharedPtr<FWebBrowserWindow> BrowserWindow = AsyncWebBrowserWindowPtr.Pin();
 				if (BrowserWindow.IsValid())
@@ -473,7 +473,7 @@ class SIOSWebBrowserWidget : public SLeafWidget
 		{
 			TWeakPtr<FWebBrowserWindow> AsyncWebBrowserWindowPtr = WebBrowserWindowPtr;
 			
-			[FIOSAsyncTask CreateTaskWithBlock : ^ bool(void)
+			[FAppleAsyncTask CreateTaskWithBlock : ^ bool(void)
 			{
 				if (TSharedPtr<FWebBrowserWindow> BrowserWindow = AsyncWebBrowserWindowPtr.Pin())
 				{
