[nDisplay] Added a setting to enable high dpi support for an application running unattended.

diff --git a/Engine/Plugins/Editor/DisplayClusterLaunch/Config/BaseDisplayClusterLaunch.ini b/Engine/Plugins/Editor/DisplayClusterLaunch/Config/BaseDisplayClusterLaunch.ini
--- a/Engine/Plugins/Editor/DisplayClusterLaunch/Config/BaseDisplayClusterLaunch.ini
+++ b/Engine/Plugins/Editor/DisplayClusterLaunch/Config/BaseDisplayClusterLaunch.ini
@@ -1,5 +1,5 @@
 [/Script/DisplayClusterLaunchEditor.DisplayClusterLaunchEditorProjectSettings]
 AdditionalConsoleVariables=("p.Chaos.Solver.Deterministic=1","LevelInstance.ForceEditorWorldMode=1")
 AdditionalConsoleCommands=("DisableAllScreenMessages")
-CommandLineArguments=("messaging","dc_cluster","nosplash","fixedseed","NoVerifyGC","noxrstereo","xrtrackingonly","RemoteControlIsHeadless","dx12","dc_dev_mono","unattended","handleensurepercent=0","ini:Engine:[/Script/Engine.Engine]:GameEngine=/Script/DisplayCluster.DisplayClusterGameEngine,[/Script/Engine.Engine]:GameViewportClientClassName=/Script/DisplayCluster.DisplayClusterViewportClient,[/Script/Engine.UserInterfaceSettings]:bAllowHighDPIInGameMode=True","ini:Game:[/Script/EngineSettings.GeneralProjectSettings]:bUseBorderlessWindow=True")
+CommandLineArguments=("messaging","dc_cluster","nosplash","fixedseed","NoVerifyGC","noxrstereo","xrtrackingonly","RemoteControlIsHeadless","dx12","dc_dev_mono","unattended","handleensurepercent=0","ini:Engine:[/Script/Engine.Engine]:GameEngine=/Script/DisplayCluster.DisplayClusterGameEngine,[/Script/Engine.Engine]:GameViewportClientClassName=/Script/DisplayCluster.DisplayClusterViewportClient,[/Script/Engine.UserInterfaceSettings]:bAllowHighDPIInGameMode=True,[/Script/Engine.UserInterfaceSettings]:bAllowHighDpiWhenUnattended=True","ini:Game:[/Script/EngineSettings.GeneralProjectSettings]:bUseBorderlessWindow=True")
 Logging=(Category="LogDisplayClusterRenderSync",VerbosityLevel=Log)
diff --git a/Engine/Plugins/VirtualProduction/Switchboard/Source/Switchboard/switchboard/devices/ndisplay/plugin_ndisplay.py b/Engine/Plugins/VirtualProduction/Switchboard/Source/Switchboard/switchboard/devices/ndisplay/plugin_ndisplay.py
--- a/Engine/Plugins/VirtualProduction/Switchboard/Source/Switchboard/switchboard/devices/ndisplay/plugin_ndisplay.py
+++ b/Engine/Plugins/VirtualProduction/Switchboard/Source/Switchboard/switchboard/devices/ndisplay/plugin_ndisplay.py
@@ -1071,7 +1071,9 @@ class DevicenDisplay(DeviceUnreal):
             "GameViewportClientClassName="
             "/Script/DisplayCluster.DisplayClusterViewportClient"
             ",[/Script/Engine.UserInterfaceSettings]:"
-            "bAllowHighDPIInGameMode=True")
+            "bAllowHighDPIInGameMode=True"
+            ",[/Script/Engine.UserInterfaceSettings]:"
+            "bAllowHighDpiWhenUnattended=True")
 
         ini_game = (
             "-ini:Game:"
diff --git a/Engine/Source/Runtime/ApplicationCore/Private/Windows/WindowsPlatformApplicationMisc.cpp b/Engine/Source/Runtime/ApplicationCore/Private/Windows/WindowsPlatformApplicationMisc.cpp
--- a/Engine/Source/Runtime/ApplicationCore/Private/Windows/WindowsPlatformApplicationMisc.cpp
+++ b/Engine/Source/Runtime/ApplicationCore/Private/Windows/WindowsPlatformApplicationMisc.cpp
@@ -11,6 +11,7 @@
 #include "Math/Color.h"
 #include "Misc/ScopeExit.h"
 #include "Modules/ModuleManager.h"
+#include "Misc/ConfigCacheIni.h"
 #include "Misc/CoreDelegates.h"
 #include "Windows/WindowsPlatformOutputDevices.h"
 #include "GenericPlatform/GenericPlatformCrashContext.h"
@@ -197,6 +198,10 @@ void FWindowsPlatformApplicationMisc::SetHighDPIMode()
 {
 	if (IsHighDPIAwarenessEnabled())
 	{
+		// Enable HighDPIMode for an application running unattended.
+		bool bAllowHighDpiWhenUnattended = false;
+		GConfig->GetBool(TEXT("/Script/Engine.UserInterfaceSettings"), TEXT("bAllowHighDpiWhenUnattended"), bAllowHighDpiWhenUnattended, GEngineIni);
+
 		if (void* ShCoreDll = FPlatformProcess::GetDllHandle(TEXT("shcore.dll")))
 		{
 			typedef enum _PROCESS_DPI_AWARENESS {
@@ -212,7 +217,7 @@ void FWindowsPlatformApplicationMisc::SetHighDPIMode()
 			typedef HRESULT(STDAPICALLTYPE *GetProcessDpiAwarenessProc)(HANDLE hProcess, PROCESS_DPI_AWARENESS* Value);
 			GetProcessDpiAwarenessProc GetProcessDpiAwareness = (GetProcessDpiAwarenessProc)FPlatformProcess::GetDllExport(ShCoreDll, TEXT("GetProcessDpiAwareness"));
 
-			if (SetProcessDpiAwareness && GetProcessDpiAwareness && !IsRunningCommandlet() && !FApp::IsUnattended())
+			if (SetProcessDpiAwareness && GetProcessDpiAwareness && !IsRunningCommandlet() && (!FApp::IsUnattended() || bAllowHighDpiWhenUnattended))
 			{
 				PROCESS_DPI_AWARENESS CurrentAwareness = PROCESS_DPI_UNAWARE;
 
@@ -237,7 +242,7 @@ void FWindowsPlatformApplicationMisc::SetHighDPIMode()
 			typedef BOOL(WINAPI *SetProcessDpiAwareProc)(void);
 			SetProcessDpiAwareProc SetProcessDpiAware = (SetProcessDpiAwareProc)FPlatformProcess::GetDllExport(User32Dll, TEXT("SetProcessDPIAware"));
 
-			if (SetProcessDpiAware && !IsRunningCommandlet() && !FApp::IsUnattended())
+			if (SetProcessDpiAware && !IsRunningCommandlet() && (!FApp::IsUnattended() || bAllowHighDpiWhenUnattended))
 			{
 				UE_LOG(LogInit, Log, TEXT("Setting process to DPI aware"));
 
