Add logs on audio route changed notification

diff --git a/Engine/Source/Runtime/ApplicationCore/Private/IOS/IOSAppDelegate.cpp b/Engine/Source/Runtime/ApplicationCore/Private/IOS/IOSAppDelegate.cpp
--- a/Engine/Source/Runtime/ApplicationCore/Private/IOS/IOSAppDelegate.cpp
+++ b/Engine/Source/Runtime/ApplicationCore/Private/IOS/IOSAppDelegate.cpp
@@ -605,6 +605,8 @@ static IOSAppDelegate* CachedDelegate = nil;
 
 	[[NSNotificationCenter defaultCenter] addObserverForName:AVAudioSessionRouteChangeNotification object : nil queue : nil usingBlock : ^ (NSNotification *notification)
 	{
+		UE_LOG(LogIOSAudioSession, Log, TEXT("AVAudioSessionRouteChange: UserInfo: [%s] CurrentRoute: [%s]"), *FString([[notification userInfo]description]), *FString([[[AVAudioSession sharedInstance] currentRoute] description]));
+
 		NSUInteger RouteChangeReason = [[[notification userInfo] objectForKey:AVAudioSessionRouteChangeReasonKey] unsignedIntegerValue];
 		[FAppleAsyncTask CreateTaskWithBlock : ^ bool(void)
 		{
