Defer IOS view rotation until we get a new frame in the new orientation

diff --git a/Engine/Source/Runtime/Apple/MetalRHI/Private/MetalViewport.cpp b/Engine/Source/Runtime/Apple/MetalRHI/Private/MetalViewport.cpp
--- a/Engine/Source/Runtime/Apple/MetalRHI/Private/MetalViewport.cpp
+++ b/Engine/Source/Runtime/Apple/MetalRHI/Private/MetalViewport.cpp
@@ -18,6 +18,7 @@
 #include "Mac/CocoaWindow.h"
 #include "Mac/CocoaThread.h"
 #else
+#include "IOS/IOSApplication.h"
 #include "IOS/IOSAppDelegate.h"
 #endif
 #include "RenderCommandFence.h"
@@ -82,7 +83,7 @@ static FAutoConsoleVariableRef CVarMetalNonBlockingPresent(
 static FCriticalSection ViewportsMutex;
 static TSet<FMetalViewport*> Viewports;
 
-FMetalViewport::FMetalViewport(FMetalDevice& InDevice, void* WindowHandle, uint32 InSizeX, uint32 InSizeY, bool bInIsFullscreen, EPixelFormat Format)
+FMetalViewport::FMetalViewport(FMetalDevice& InDevice, void* WindowHandle, uint32 InSizeX, uint32 InSizeY, bool bInIsFullscreen, EPixelFormat InFormat)
 	: Device(InDevice)
 	, Drawable{nullptr}
 	, BackBuffer{nullptr, nullptr}
@@ -93,6 +94,10 @@ FMetalViewport::FMetalViewport(FMetalDevice& InDevice, void* WindowHandle, uint3
 	, FrameAvailable{0}
 	, LastCompleteFrame{nullptr}
 	, bIsFullScreen{bInIsFullscreen}
+#if PLATFORM_IOS && !PLATFORM_TVOS && !PLATFORM_VISIONOS
+	, SizeX(InSizeX)
+	, SizeY(InSizeY)
+#endif
 #if PLATFORM_MAC
 	, View{nullptr}
 #endif
@@ -136,7 +141,19 @@ FMetalViewport::FMetalViewport(FMetalDevice& InDevice, void* WindowHandle, uint3
 		[[Window standardWindowButton:NSWindowCloseButton] setAction:@selector(performClose:)];
 	});
 #endif
-	Resize(InSizeX, InSizeY, bInIsFullscreen, Format);
+
+	Resize(InSizeX, InSizeY, bInIsFullscreen, InFormat);
+
+#if PLATFORM_IOS && !PLATFORM_TVOS && !PLATFORM_VISIONOS
+	// An orientation change from LandscapeLeft to LandscapeRight won't trigger a SetRes, so we need to react to an orientation change here
+	OrientationChangedHandle = FCoreDelegates::ApplicationReceivedScreenOrientationChangedNotificationDelegate.AddLambda([this](int32 ScreenOrientation)
+	{
+		if ([IOSAppDelegate GetMaskFromEDeviceScreenOrientation:((EDeviceScreenOrientation)ScreenOrientation)] != OrientationMask)
+		{
+			Resize(SizeX, SizeY, bIsFullScreen, Format);
+		}
+	});
+#endif
 	
 	{
 		FScopeLock Lock(&ViewportsMutex);
@@ -161,6 +178,10 @@ FMetalViewport::~FMetalViewport()
 		Viewports.Remove(this);
 	}
 	
+#if PLATFORM_IOS && !PLATFORM_TVOS && !PLATFORM_VISIONOS
+	FCoreDelegates::ApplicationReceivedScreenOrientationChangedNotificationDelegate.Remove(OrientationChangedHandle);
+#endif
+
 	BackBuffer[0].SafeRelease();	// when the rest of the engine releases it, its framebuffers will be released too (those the engine knows about)
 	BackBuffer[1].SafeRelease();
 	check(!IsValidRef(BackBuffer[0]));
@@ -188,14 +209,14 @@ uint32 FMetalViewport::GetViewportIndex(EMetalViewportAccessFlag Accessor) const
 	}
 }
 
-void FMetalViewport::Resize(uint32 InSizeX, uint32 InSizeY, bool bInIsFullscreen, EPixelFormat Format)
+void FMetalViewport::Resize(uint32 InSizeX, uint32 InSizeY, bool bInIsFullscreen, EPixelFormat InFormat)
 {
 	bIsFullScreen = bInIsFullscreen;
 	uint32 Index = GetViewportIndex(EMetalViewportAccessGame);
 	
-	bool bUseHDR = GRHISupportsHDROutput && Format == GRHIHDRDisplayOutputFormat;
+	bool bUseHDR = GRHISupportsHDROutput && InFormat == GRHIHDRDisplayOutputFormat;
 	
-    MTL::PixelFormat MetalFormat = (MTL::PixelFormat)GPixelFormats[Format].PlatformFormat;
+    MTL::PixelFormat MetalFormat = (MTL::PixelFormat)GPixelFormats[InFormat].PlatformFormat;
 	
     ENQUEUE_RENDER_COMMAND(FlushPendingRHICommands)(
         [Viewport = this](FRHICommandListImmediate& RHICmdList)
@@ -203,7 +224,7 @@ void FMetalViewport::Resize(uint32 InSizeX, uint32 InSizeY, bool bInIsFullscreen
             GRHICommandList.GetImmediateCommandList().BlockUntilGPUIdle();
         });
     
-	if (IsValidRef(BackBuffer[Index]) && Format != BackBuffer[Index]->GetFormat())
+	if (IsValidRef(BackBuffer[Index]) && InFormat != BackBuffer[Index]->GetFormat())
 	{
 		// Really need to flush the RHI thread & GPU here...
 		AddRef();
@@ -245,6 +266,10 @@ void FMetalViewport::Resize(uint32 InSizeX, uint32 InSizeY, bool bInIsFullscreen
 	// (even though Apple's HDR displays are P3)
 	// and its compositor will do the conversion.
 	{
+#if PLATFORM_IOS && !PLATFORM_TVOS && !PLATFORM_VISIONOS
+		__block UIInterfaceOrientationMask CachedOrientationMask;
+#endif
+
 		dispatch_sync(dispatch_get_main_queue(), ^{
 			IOSAppDelegate* AppDelegate = [IOSAppDelegate GetDelegate];
 			FIOSView* IOSView = AppDelegate.IOSView;
@@ -257,7 +282,19 @@ void FMetalViewport::Resize(uint32 InSizeX, uint32 InSizeY, bool bInIsFullscreen
 			}
 			
 			[IOSView UpdateRenderWidth:InSizeX andHeight:InSizeY];
-		});
+			
+#if PLATFORM_IOS && !PLATFORM_TVOS && !PLATFORM_VISIONOS
+			CachedOrientationMask = [IOSAppDelegate GetMaskFromUIInterfaceOrientation:(FIOSApplication::CachedOrientation)];
+#endif
+	});
+
+#if PLATFORM_IOS && !PLATFORM_TVOS && !PLATFORM_VISIONOS
+		// Cache new size and orientation
+		SizeX = InSizeX;
+		SizeY = InSizeY;
+		Format = InFormat;
+		OrientationMask = CachedOrientationMask;
+#endif
 	}
 #endif
 
@@ -268,7 +305,7 @@ void FMetalViewport::Resize(uint32 InSizeX, uint32 InSizeY, bool bInIsFullscreen
 		TRefCountPtr<FMetalSurface> DoubleBuffer;
 
 		FRHITextureCreateDesc CreateDesc =
-			FRHITextureCreateDesc::Create2D(TEXT("BackBuffer"), InSizeX, InSizeY, Format)
+			FRHITextureCreateDesc::Create2D(TEXT("BackBuffer"), InSizeX, InSizeY, InFormat)
 			.SetClearValue(FClearValueBinding::None)
 			.SetFlags(ETextureCreateFlags::RenderTargetable);
 		
@@ -459,7 +496,29 @@ void FMetalViewport::Present(FMetalCommandQueue& CommandQueue, bool bLockToVsync
 		CurrentLayer.displaySyncEnabled = bLockToVsync || (!(IsRunningGame() && bIsFullScreen));
 	}
 #endif
-	
+
+#if PLATFORM_IOS && !PLATFORM_TVOS && !PLATFORM_VISIONOS
+	extern bool GIOSDelayRotationUntilPresent;
+	if (GIOSDelayRotationUntilPresent)
+	{
+		// If this frame is at a new orientation, notify to allow rotation to the new orientation
+		if (OrientationMask != OldOrientationMask)
+		{
+			OldOrientationMask = OrientationMask;
+
+			IOSAppDelegate* AppDelegate = [IOSAppDelegate GetDelegate];
+			uint32 ViewportSizeX = SizeX;
+			uint32 ViewportSizeY = SizeY;
+			UIInterfaceOrientationMask ViewportOrientationMask = OrientationMask;
+
+			// dispatched sync so the main queue can take a snapshot of the view in the old orientation before we complete presenting, to crossfade
+			dispatch_sync(dispatch_get_main_queue(), ^{
+				[AppDelegate.IOSController notifyPresentAfterRotateOrientationMask:ViewportOrientationMask withSizeX:ViewportSizeX withSizeY:ViewportSizeY];
+			});
+		}
+	}
+#endif
+
 	LastCompleteFrame = GetBackBuffer(EMetalViewportAccessRHI);
 	FPlatformAtomics::InterlockedExchange(&FrameAvailable, 1);
 	
diff --git a/Engine/Source/Runtime/Apple/MetalRHI/Public/MetalViewport.h b/Engine/Source/Runtime/Apple/MetalRHI/Public/MetalViewport.h
--- a/Engine/Source/Runtime/Apple/MetalRHI/Public/MetalViewport.h
+++ b/Engine/Source/Runtime/Apple/MetalRHI/Public/MetalViewport.h
@@ -48,7 +48,7 @@ typedef void (^FMetalViewportPresentHandler)(uint32 CGDirectDisplayID, double Ou
 class FMetalViewport : public FRHIViewport
 {
 public:
-	FMetalViewport(FMetalDevice& InDevice, void* WindowHandle, uint32 InSizeX, uint32 InSizeY, bool bInIsFullscreen,EPixelFormat Format);
+	FMetalViewport(FMetalDevice& InDevice, void* WindowHandle, uint32 InSizeX, uint32 InSizeY, bool bInIsFullscreen,EPixelFormat InFormat);
 	~FMetalViewport();
 
 	void Resize(uint32 InSizeX, uint32 InSizeY, bool bInIsFullscreen,EPixelFormat Format);
@@ -105,6 +105,15 @@ private:
 	TRefCountPtr<FMetalSurface> LastCompleteFrame;
 	bool bIsFullScreen;
 
+#if PLATFORM_IOS && !PLATFORM_TVOS && !PLATFORM_VISIONOS	
+	uint32 SizeX;
+	uint32 SizeY;
+	EPixelFormat Format;
+	UIInterfaceOrientationMask OrientationMask;
+	UIInterfaceOrientationMask OldOrientationMask;
+    FDelegateHandle OrientationChangedHandle;
+#endif
+
 #if PLATFORM_MAC
 	FMetalView* View;
 #endif
diff --git a/Engine/Source/Runtime/ApplicationCore/Private/IOS/IOSAppDelegate.cpp b/Engine/Source/Runtime/ApplicationCore/Private/IOS/IOSAppDelegate.cpp
--- a/Engine/Source/Runtime/ApplicationCore/Private/IOS/IOSAppDelegate.cpp
+++ b/Engine/Source/Runtime/ApplicationCore/Private/IOS/IOSAppDelegate.cpp
@@ -98,6 +98,15 @@ FAutoConsoleVariableRef CVarNotificationActivationProtocolKey(
 		ECVF_Default
 );
 
+#if !PLATFORM_TVOS
+bool GIOSDelayRotationUntilPresent = 1;
+FAutoConsoleVariableRef CVarDelayRotationUntilPresent(
+		TEXT("ios.DelayRotationUntilPresent"),
+		GIOSDelayRotationUntilPresent,
+		TEXT("Whether to delay device rotation handling until we have a frame available at the new orientation (default=1).\n"),
+		ECVF_Default);
+#endif
+
 
 static void SignalHandler(int32 Signal, struct __siginfo* Info, void* Context)
 {
@@ -227,11 +236,11 @@ bool FIOSCoreDelegates::PassesPushNotificationFilters(NSDictionary* Payload)
 @synthesize AlertResponse;
 @synthesize bDeviceInPortraitMode;
 @synthesize bEngineInit;
+@synthesize bPlatformInit;
 @synthesize OSVersion;
 
 @synthesize Window;
 @synthesize IOSView;
-@synthesize SlateController;
 @synthesize timer;
 @synthesize IdleTimerEnableTimer;
 @synthesize IdleTimerEnablePeriod;
@@ -299,7 +308,6 @@ static IOSAppDelegate* CachedDelegate = nil;
 #endif
 	[Window release];
 	[IOSView release];
-	[SlateController release];
 	[timer release];
 #if WITH_ACCESSIBILITY
 	if (AccessibilityCacheTimer != nil)
@@ -364,6 +372,14 @@ static IOSAppDelegate* CachedDelegate = nil;
 
 	self.bEngineInit = true;
 
+#if !PLATFORM_TVOS && !PLATFORM_VISIONOS
+	// Ask iOS to re-evaluate allowed orientations if we change the ios.DelayRotationUntilPresent
+	CVarDelayRotationUntilPresent->OnChangedDelegate().AddLambda([](IConsoleVariable* Variable)
+	{
+		[IOSAppDelegate UpdateSupportedInterfaceOrientations];
+	});
+#endif
+
 	// put a render thread job to turn off the splash screen after the first render flip
 	if (GShowSplashScreen)
 	{
@@ -1011,18 +1027,6 @@ static IOSAppDelegate* CachedDelegate = nil;
     return self.ThermalState;
 }
 
-- (UIViewController*) IOSController
-{
-	// walk the responder chain until we get to a VC
-	UIResponder *Responder = IOSView;
-	while (Responder != nil && ![Responder isKindOfClass:[UIViewController class]])
-	{
-		Responder = [Responder nextResponder];
-	}
-	return (UIViewController*)Responder;
-}
-
-
 bool GIsSuspended = 0;
 - (void)ToggleSuspend:(bool)bSuspend
 {
@@ -1038,6 +1042,9 @@ bool GIsSuspended = 0;
 	{
         FIOSPlatformRHIFramePacer::Resume();
 		FAppEntry::Resume();
+
+		// Check rotation state after returning from background
+		[self didRotate:nil];
 	}
     
 	if (IOSView && IOSView->bIsInitialized)
@@ -1066,6 +1073,7 @@ bool GIsSuspended = 0;
 {
 	self.bDeviceInPortraitMode = false;
 	self.bEngineInit = false;
+	self.bPlatformInit = false;
 
 	return YES;
 }
@@ -1089,6 +1097,12 @@ static FAutoConsoleVariableRef CVarGEnableThermalsReport(
 	self.InterfaceOrientation = [self.Window.windowScene interfaceOrientation];
 	self.bDeviceInPortraitMode = UIInterfaceOrientationIsPortrait(self.InterfaceOrientation);
 	printf("========= This app is in %s mode\n", self.bDeviceInPortraitMode ? "PORTRAIT" : "LANDSCAPE");
+
+#if !PLATFORM_VISIONOS
+	// Register for device orientation changes
+	[[UIDevice currentDevice] beginGeneratingDeviceOrientationNotifications];
+    [[NSNotificationCenter defaultCenter] addObserver:self selector:@selector(didRotate:) name:UIDeviceOrientationDidChangeNotification object:nil];
+#endif
 #endif
 
 	// check OS version to make sure we have the API
@@ -1138,9 +1152,6 @@ static FAutoConsoleVariableRef CVarGEnableThermalsReport(
 #if !PLATFORM_TVOS && !PLATFORM_VISIONOS
 	UNUserNotificationCenter *Center = [UNUserNotificationCenter currentNotificationCenter];
 	Center.delegate = self;
-	// Register for device orientation changes
-	[[UIDevice currentDevice] beginGeneratingDeviceOrientationNotifications];
-    [[NSNotificationCenter defaultCenter] addObserver:self selector:@selector(didRotate:) name:UIDeviceOrientationDidChangeNotification object:nil];
 
 #if !UE_BUILD_SHIPPING
 	// make a history buffer
@@ -1436,40 +1447,118 @@ static FAutoConsoleVariableRef CVarGEnableThermalsReport(
 
 - (void) didRotate:(NSNotification *)notification
 {   
-#if !PLATFORM_TVOS &&!PLATFORM_VISIONOS
+#if !PLATFORM_TVOS && !PLATFORM_VISIONOS
 	// get the interface orientation
 	
 	if(GIsSuspended)
 	{
 		return;
 	}
+		
+	FPlatformMisc::LowLevelOutputDebugStringf(TEXT("[rotation] didRotate windowScene=%d device=%d bPlatformInit=%d\n"), (int)[self.Window.windowScene interfaceOrientation], (int)[UIDevice currentDevice].orientation, bPlatformInit);
 	
-	NSLog(@"didRotate orientation = %d", (int)[self.Window.windowScene interfaceOrientation]);
-	
-    UIInterfaceOrientation Orientation = [self.Window.windowScene interfaceOrientation];
-	self.InterfaceOrientation = Orientation;
-	
-    if (self.bEngineInit)
+    if (self.bPlatformInit)
     {
-		FFunctionGraphTask::CreateAndDispatchWhenReady([Orientation]()
+		UIInterfaceOrientationMask PotentialMask;
+		UIInterfaceOrientation NewOrientation = self.InterfaceOrientation;
+		if (GIOSDelayRotationUntilPresent)
 		{
-			FIOSApplication* Application = [IOSAppDelegate GetDelegate].IOSApplication;
-			Application->OrientationChanged(Orientation);
-			FCoreDelegates::ApplicationReceivedScreenOrientationChangedNotificationDelegate.Broadcast((int32)[IOSAppDelegate ConvertFromUIInterfaceOrientation:Orientation]);
+			UIInterfaceOrientationMask AllowedMask = [self.IOSController supportedInterfaceOrientations_Internal];
 
-			//we also want to fire off the safe frame event
-			FCoreDelegates::OnSafeFrameChangedEvent.Broadcast();
-		}, TStatId(), NULL, ENamedThreads::GameThread);
+			// See if we support the current device orientation
+			UIInterfaceOrientation PotentialOrientation = [IOSAppDelegate GetUIInterfaceOrientationFromUIDeviceOrientation: [UIDevice currentDevice].orientation];	
+			PotentialMask = [IOSAppDelegate GetMaskFromUIInterfaceOrientation: PotentialOrientation];
+			if (PotentialOrientation != UIInterfaceOrientationUnknown && (PotentialMask & AllowedMask) != 0)
+			{
+				FPlatformMisc::LowLevelOutputDebugStringf(TEXT("[rotation] didRotate: using device orientation\n"));
+				NewOrientation = PotentialOrientation;
+			}
+			else
+			{
+				// If not, check if we support the current windowScene interfaceOrientation. It could have been programmatically set.
+				PotentialOrientation = [self.Window.windowScene interfaceOrientation];
+				PotentialMask = [IOSAppDelegate GetMaskFromUIInterfaceOrientation: PotentialOrientation];
+				if ((PotentialMask & AllowedMask) != 0)
+				{
+					NewOrientation = PotentialOrientation;
+					FPlatformMisc::LowLevelOutputDebugStringf(TEXT("[rotation] didRotate: using windowScene interfaceOrientation\n"));
+				}
+				else
+				{			
+					// Force rotation to the first allowed orientation
+					if ((AllowedMask & UIInterfaceOrientationMaskLandscape) == UIInterfaceOrientationMaskLandscape)
+					{
+						FString PreferredLandscapeOrientation;
+						GConfig->GetString(TEXT("/Script/IOSRuntimeSettings.IOSRuntimeSettings"), TEXT("PreferredLandscapeOrientation"), PreferredLandscapeOrientation, GEngineIni);
+						
+						// Preferred default is LandscapeRight as per 
+						if (PreferredLandscapeOrientation.Equals("LandscapeLeft"))
+						{
+							NewOrientation =  UIInterfaceOrientationLandscapeLeft;
+							PotentialMask = UIInterfaceOrientationMaskLandscapeLeft;
+						}
+						else
+						{
+							NewOrientation =  UIInterfaceOrientationLandscapeRight;
+							PotentialMask = UIInterfaceOrientationMaskLandscapeRight;
+						}
+					}
+					else if (AllowedMask & UIInterfaceOrientationMaskLandscapeRight)
+					{
+						NewOrientation =  UIInterfaceOrientationLandscapeRight;
+						PotentialMask = UIInterfaceOrientationMaskLandscapeRight;
+					}
+					else if (AllowedMask & UIInterfaceOrientationMaskLandscapeLeft)
+					{
+						NewOrientation =  UIInterfaceOrientationLandscapeLeft;
+						PotentialMask = UIInterfaceOrientationMaskLandscapeLeft;
+					}
+					else if (AllowedMask & UIInterfaceOrientationMaskPortrait)
+					{
+						NewOrientation = UIInterfaceOrientationPortrait;
+						PotentialMask = UIInterfaceOrientationMaskPortrait;
+					}
+					else if (AllowedMask & UIInterfaceOrientationMaskPortraitUpsideDown)
+					{
+						NewOrientation = UIInterfaceOrientationPortraitUpsideDown;
+						PotentialMask = UIInterfaceOrientationMaskPortraitUpsideDown;
+					}
+					FPlatformMisc::LowLevelOutputDebugStringf(TEXT("[rotation] didRotate: using first allowed orientation\n"));
+				}
+			}
+		}
+		else
+		{
+			NewOrientation = [self.Window.windowScene interfaceOrientation];
+		}
+
+		if (self.InterfaceOrientation != NewOrientation)
+		{
+			FPlatformMisc::LowLevelOutputDebugStringf(TEXT("[rotation] didRotate: rotating to new orientation %d\n"), (int)NewOrientation);
+
+			self.InterfaceOrientation = NewOrientation;
+			FIOSApplication* Application = self.IOSApplication;
+			FIOSView* View = self.IOSView;
+			IOSViewController* Controller = self.IOSController;
+			
+			FFunctionGraphTask::CreateAndDispatchWhenReady([NewOrientation, Application, View, Controller, PotentialMask]()
+			{
+				Application->OrientationChanged(NewOrientation);
+			}, TStatId(), NULL, ENamedThreads::GameThread);
+		}
 	}
 	else
 	{
+		UIInterfaceOrientation NewOrientation = [self.Window.windowScene interfaceOrientation];
+		self.InterfaceOrientation = NewOrientation;
+
 		[FAppleAsyncTask CreateTaskWithBlock : ^ bool(void)
 		{
 			FIOSApplication* Application = [IOSAppDelegate GetDelegate].IOSApplication;
-			Application->OrientationChanged(Orientation);
-			FCoreDelegates::ApplicationReceivedScreenOrientationChangedNotificationDelegate.Broadcast((int32)[IOSAppDelegate ConvertFromUIInterfaceOrientation:Orientation]);
+			Application->OrientationChanged(NewOrientation);
+			FCoreDelegates::ApplicationReceivedScreenOrientationChangedNotificationDelegate.Broadcast((int32)[IOSAppDelegate ConvertFromUIInterfaceOrientation:NewOrientation]);
 
-			//we also want to fire off the safe frame event
+			// we also want to fire off the safe frame event
 			FCoreDelegates::OnSafeFrameChangedEvent.Broadcast();
 			return true;
 		}];
@@ -1848,6 +1937,66 @@ extern double GCStartTime;
 	}
 }
 
++(UIInterfaceOrientation) GetUIInterfaceOrientationFromUIDeviceOrientation:(UIDeviceOrientation)Orientation
+{
+	switch(Orientation)
+	{
+		default:
+		case UIDeviceOrientationUnknown : return UIInterfaceOrientationUnknown; break;
+		case UIDeviceOrientationPortrait : return UIInterfaceOrientationPortrait; break;
+		case UIDeviceOrientationPortraitUpsideDown : return UIInterfaceOrientationPortraitUpsideDown; break;
+		// Note UIDeviceOrientation and UIInterfaceOrientation have have lansdcape left and right switched
+		// https://developer.apple.com/documentation/uikit/uiinterfaceorientation
+		case UIDeviceOrientationLandscapeLeft : return UIInterfaceOrientationLandscapeRight; break;
+		case UIDeviceOrientationLandscapeRight : return UIInterfaceOrientationLandscapeLeft; break;
+	}
+}
+
++(UIInterfaceOrientationMask) GetMaskFromUIInterfaceOrientation:(UIInterfaceOrientation)Orientation
+{
+	switch(Orientation)
+	{
+		default:
+		case UIInterfaceOrientationUnknown : return UIInterfaceOrientationMaskAll; break;
+		case UIInterfaceOrientationPortrait : return UIInterfaceOrientationMaskPortrait; break;
+		case UIInterfaceOrientationPortraitUpsideDown : return UIInterfaceOrientationMaskPortraitUpsideDown; break;
+		case UIInterfaceOrientationLandscapeLeft : return UIInterfaceOrientationMaskLandscapeLeft; break;
+		case UIInterfaceOrientationLandscapeRight : return UIInterfaceOrientationMaskLandscapeRight; break;
+	}
+}
+
++(UIInterfaceOrientationMask) GetMaskFromEDeviceScreenOrientation:(EDeviceScreenOrientation)Orientation
+{
+	switch(Orientation)
+	{
+		default:
+		case EDeviceScreenOrientation::Unknown : return UIInterfaceOrientationMaskAll; break;
+		case EDeviceScreenOrientation::Portrait : return UIInterfaceOrientationMaskPortrait; break;
+		case EDeviceScreenOrientation::PortraitUpsideDown : return UIInterfaceOrientationMaskPortraitUpsideDown; break;
+		case EDeviceScreenOrientation::LandscapeLeft : return UIInterfaceOrientationMaskLandscapeLeft; break;
+		case EDeviceScreenOrientation::LandscapeRight : return UIInterfaceOrientationMaskLandscapeRight; break;
+	}
+}
+
++(void)UpdateSupportedInterfaceOrientations;
+{
+	dispatch_async(dispatch_get_main_queue(),
+	^{
+		if (@available(iOS 16.0, visionOS 1.0, *))
+		{
+			[[IOSAppDelegate GetDelegate].IOSController setNeedsUpdateOfSupportedInterfaceOrientations];
+		}
+		else
+		{
+// Deprecated iOS16 - suppress warning if deployment target	>= iOS16
+#pragma clang diagnostic push
+#pragma clang diagnostic ignored "-Wdeprecated-declarations"
+			[UIViewController attemptRotationToDeviceOrientation];
+#pragma clang diagnostic pop
+		}
+	});
+}
+
 void HandleReceivedNotification(UNNotification* notification)
 {
 	bool IsLocal = false;
diff --git a/Engine/Source/Runtime/ApplicationCore/Private/IOS/IOSApplication.cpp b/Engine/Source/Runtime/ApplicationCore/Private/IOS/IOSApplication.cpp
--- a/Engine/Source/Runtime/ApplicationCore/Private/IOS/IOSApplication.cpp
+++ b/Engine/Source/Runtime/ApplicationCore/Private/IOS/IOSApplication.cpp
@@ -210,46 +210,52 @@ void FDisplayMetrics::RebuildDisplayMetrics(FDisplayMetrics& OutDisplayMetrics)
     
 }
 
-void FIOSApplication::CacheDisplayMetrics()
-{
 #if !PLATFORM_TVOS
+void FIOSApplication::CacheDisplayMetrics(UIInterfaceOrientation Orientation)
+{
+	CachedOrientation = Orientation;
+	//todo: rotate insets to match given orientation
 	CachedInsets = [[[[UIApplication sharedApplication] delegate] window] safeAreaInsets];
-	CachedOrientation = [[[[[UIApplication sharedApplication] delegate] window] windowScene] interfaceOrientation];
-#endif
 }
+#endif
 
 TSharedRef< FGenericWindow > FIOSApplication::MakeWindow()
 {
 	return FIOSWindow::Make();
 }
 
-#if !PLATFORM_TVOS
-void FIOSApplication::OrientationChanged(UIInterfaceOrientation orientation)
+#if !PLATFORM_TVOS && !PLATFORM_VISIONOS
+
+void FIOSApplication::OrientationChanged(UIInterfaceOrientation NewOrientation)
 {
+	FPlatformMisc::LowLevelOutputDebugStringf(TEXT("[rotation] FIOSApplication::OrientationChanged: %d\n"), NewOrientation);
+
 	// queue up the size as we see it now all the way to the RHI 
 	[FAppleAsyncTask CreateTaskWithBlock : ^ bool(void)
 	{
+	 	IOSAppDelegate* AppDelegate = [IOSAppDelegate GetDelegate];
+		FIOSApplication* App = AppDelegate.IOSApplication;
+		
 		__block FPlatformRect WindowRect;
 		dispatch_sync(dispatch_get_main_queue(), ^{
+			App->CacheDisplayMetrics(NewOrientation);	// Requires UI Thread
 			WindowRect = FIOSWindow::GetScreenRect();	// Requires UI Thread
 		});
 
 		int32 WindowWidth = WindowRect.Right - WindowRect.Left;
 		int32 WindowHeight = WindowRect.Bottom - WindowRect.Top;
-		
-	 	FIOSApplication* App = [IOSAppDelegate GetDelegate].IOSApplication;
+
+		// reize viewport
 		App->GetMessageHandler()->OnSizeChanged(App->Windows[0], WindowWidth, WindowHeight, false);
 		App->GetMessageHandler()->OnResizingWindow(App->Windows[0]);
-		
-		dispatch_sync(dispatch_get_main_queue(), ^{
-			App->CacheDisplayMetrics();	// Requires UI Thread
-		});
-	 	
+		 	
 		FDisplayMetrics DisplayMetrics;
 		FDisplayMetrics::RebuildDisplayMetrics(DisplayMetrics);
 		App->BroadcastDisplayMetricsChanged(DisplayMetrics);
 		FCoreDelegates::OnSafeFrameChangedEvent.Broadcast();
-	 	return true;
+
+		FCoreDelegates::ApplicationReceivedScreenOrientationChangedNotificationDelegate.Broadcast((int32)[IOSAppDelegate ConvertFromUIInterfaceOrientation:NewOrientation]);
+		return true;
 	 }];
 }
 #endif
diff --git a/Engine/Source/Runtime/ApplicationCore/Private/IOS/IOSView.cpp b/Engine/Source/Runtime/ApplicationCore/Private/IOS/IOSView.cpp
--- a/Engine/Source/Runtime/ApplicationCore/Private/IOS/IOSView.cpp
+++ b/Engine/Source/Runtime/ApplicationCore/Private/IOS/IOSView.cpp
@@ -345,12 +345,12 @@ MTL::Device* GMetalDevice = nullptr;
  */
 - (void)layoutSubviews
 {
-#if !PLATFORM_TVOS
+#if !PLATFORM_TVOS && !PLATFORM_VISIONOS
 	// IOSApplication::OrientationChange calls into Slate which issues rendering commands
 	if(![IOSAppDelegate GetDelegate].bIsSuspended)
 	{
-		auto orientation = [self.window.windowScene interfaceOrientation];
-		FIOSApplication::OrientationChanged(orientation);
+		UIInterfaceOrientation CurrentOrientation = [self.window.windowScene interfaceOrientation];
+		FIOSApplication::OrientationChanged(CurrentOrientation);
 	}
 #endif
 }
@@ -1083,6 +1083,8 @@ self.accessibilityElements = @[Window.accessibilityContainer];
 		Swap(Frame.size.width, Frame.size.height);
 	}
 #endif
+ 
+	[IOSAppDelegate GetDelegate].IOSController = self;
 
 	self.view = [[UIView alloc] initWithFrame:Frame];
 
@@ -1105,10 +1107,11 @@ self.accessibilityElements = @[Window.accessibilityContainer];
 }
 
 #if !PLATFORM_TVOS
+
 /**
  * Tell the OS about the default supported orientations
  */
-- (NSUInteger)supportedInterfaceOrientations
+ - (UIInterfaceOrientationMask)supportedInterfaceOrientations_Internal
 {
 	const IOSAppDelegate *AppDelegate = [IOSAppDelegate GetDelegate];
 	const FIOSView *View = [AppDelegate IOSView];
@@ -1125,6 +1128,43 @@ self.accessibilityElements = @[Window.accessibilityContainer];
 	UIApplication *app = [UIApplication sharedApplication];
 	return [AppDelegate application:app supportedInterfaceOrientationsForWindow:[AppDelegate window]];
 }
+
+ - (UIInterfaceOrientationMask)supportedInterfaceOrientations
+ {
+	IOSAppDelegate* AppDelegate = [IOSAppDelegate GetDelegate];
+	const FIOSView* View = [AppDelegate IOSView];
+	extern bool GIOSDelayRotationUntilPresent;
+	if (GIOSDelayRotationUntilPresent && AppDelegate.bPlatformInit && View != nil && (int)View->LastFrameInterfaceOrientationMask != 0)
+	{
+		FPlatformMisc::LowLevelOutputDebugStringf(TEXT("[rotation] supportedInterfaceOrientations: allowing %d\n"), (int)View->LastFrameInterfaceOrientationMask);
+
+		// Only allow rotation to the orientation of the last frame we have rendered, until a new one is available in the new orientation.
+		return View->LastFrameInterfaceOrientationMask;
+	}
+	else
+	{
+		if (GIOSDelayRotationUntilPresent)
+		{
+			if (!AppDelegate.bPlatformInit)
+			{
+				FPlatformMisc::LowLevelOutputDebugStringf(TEXT("[rotation] supportedInterfaceOrientations: bPlatformInit=false\n"));
+			}
+			else
+			if (View == nil)
+			{
+				FPlatformMisc::LowLevelOutputDebugStringf(TEXT("[rotation] supportedInterfaceOrientations: view is nil\n"));
+			}
+			else
+			{
+				FPlatformMisc::LowLevelOutputDebugStringf(TEXT("[rotation] supportedInterfaceOrientations: last %d\n"), (int)View->LastFrameInterfaceOrientationMask);
+			}
+		}
+
+		return [self supportedInterfaceOrientations_Internal];
+	}
+ }
+
+
 #endif
 
 /**
@@ -1187,8 +1227,29 @@ self.accessibilityElements = @[Window.accessibilityContainer];
 	return UIRectEdgeBottom;
 }
 
+- (void)viewWillTransitionToSize : (CGSize)size withTransitionCoordinator : (id<UIViewControllerTransitionCoordinator>)coordinator
+{
+	const IOSAppDelegate *AppDelegate = [IOSAppDelegate GetDelegate];
+	const FIOSView *View = [AppDelegate IOSView];
+	
+	FPlatformMisc::LowLevelOutputDebugStringf(TEXT("[rotation] viewWillTransitionToSize\n"));
+
+	[super viewWillTransitionToSize : size withTransitionCoordinator : coordinator];	
+}
+
+#if !PLATFORM_TVOS && !PLATFORM_VISIONOS
 
+- (void)notifyPresentAfterRotateOrientationMask : (UIInterfaceOrientationMask)NewOrientationMask withSizeX : (unsigned int)SizeX withSizeY : (unsigned int)SizeY
+{
+	const FIOSView* View = [[IOSAppDelegate GetDelegate] IOSView];
 
+	FPlatformMisc::LowLevelOutputDebugStringf(TEXT("[rotation] notifyPresentAfterRotate OrientationMask:%d SizeX:%d SizeY:%d\n"), NewOrientationMask, SizeX, SizeY);
+	
+	View->LastFrameInterfaceOrientationMask = NewOrientationMask;
 
+	// Attempt to rotate to desired orientation
+	[IOSAppDelegate UpdateSupportedInterfaceOrientations];
+ } 
+ #endif
 
 @end
diff --git a/Engine/Source/Runtime/ApplicationCore/Private/IOS/IOSWindow.cpp b/Engine/Source/Runtime/ApplicationCore/Private/IOS/IOSWindow.cpp
--- a/Engine/Source/Runtime/ApplicationCore/Private/IOS/IOSWindow.cpp
+++ b/Engine/Source/Runtime/ApplicationCore/Private/IOS/IOSWindow.cpp
@@ -3,6 +3,7 @@
 #include "IOSWindow.h"
 #include "IOS/IOSAppDelegate.h"
 #include "IOS/IOSView.h"
+#include "IOS/IOSApplication.h"
 
 FIOSWindow::~FIOSWindow()
 {
@@ -133,7 +134,23 @@ FPlatformRect FIOSWindow::GetScreenRect()
 		ScreenRect.Left = FMath::TruncToInt(Frame.origin.x * Scale);
 		ScreenRect.Bottom = FMath::TruncToInt((Frame.origin.y + Frame.size.height) * Scale);
 		ScreenRect.Right = FMath::TruncToInt((Frame.origin.x + Frame.size.width) * Scale);
-		
+
+#if !PLATFORM_TVOS && !PLATFORM_VISIONOS
+		extern bool GIOSDelayRotationUntilPresent;
+		if (GIOSDelayRotationUntilPresent)
+		{	
+			int32 WindowWidth = ScreenRect.Right - ScreenRect.Left;
+			int32 WindowHeight = ScreenRect.Bottom - ScreenRect.Top;
+
+			if ((UIInterfaceOrientationIsPortrait(FIOSApplication::CachedOrientation) && WindowWidth > WindowHeight) ||
+				(UIInterfaceOrientationIsLandscape(FIOSApplication::CachedOrientation) && WindowHeight > WindowWidth))
+			{
+				Exchange(ScreenRect.Top, ScreenRect.Left);
+				Exchange(ScreenRect.Bottom, ScreenRect.Right);
+			}			
+		}
+#endif
+
 		// see QuantizeSceneBufferSize
 		const uint32 DividableBy = 4;
 		const uint32 Mask = ~(DividableBy - 1);
diff --git a/Engine/Source/Runtime/ApplicationCore/Public/IOS/IOSAppDelegate.h b/Engine/Source/Runtime/ApplicationCore/Public/IOS/IOSAppDelegate.h
--- a/Engine/Source/Runtime/ApplicationCore/Public/IOS/IOSAppDelegate.h
+++ b/Engine/Source/Runtime/ApplicationCore/Public/IOS/IOSAppDelegate.h
@@ -111,8 +111,6 @@ namespace FAppEntry
 
 @class FIOSView;
 @class IOSViewController;
-@class SlateOpenGLESViewController;
-
 
 #ifndef SWIFT_IMPORT
 APPLICATIONCORE_API
@@ -140,21 +138,19 @@ APPLICATIONCORE_API
 @property (retain) NSArray* SwiftLayerViewports;
 #endif
 
+/** The view containing the game contents */
 @property (retain) FIOSView* IOSView;
 
 #ifndef SWIFT_IMPORT
 @property class FIOSApplication* IOSApplication;
 #endif
 
-/** The controller to handle rotation of the view */
-@property (readonly) UIViewController* IOSController;
+/** The controller for the main view */
+@property (retain) IOSViewController* IOSController;
 
 /** The view controlled by the auto-rotating controller */
 @property (retain) UIView* RootView;
 
-/** The controller to handle rotation of the view */
-@property (retain) SlateOpenGLESViewController* SlateController;
-
 /** The value of the alert response (atomically set since main thread and game thread use it */
 @property (assign) int AlertResponse;
 
@@ -163,6 +159,7 @@ APPLICATIONCORE_API
 
 @property bool bDeviceInPortraitMode;
 
+/** splash screen view controller */
 @property (retain) UIViewController* viewController;
 
 @property (retain) NSTimer* timer;
@@ -194,6 +191,9 @@ APPLICATIONCORE_API
 
 /** True if the engine has been initialized */
 @property (atomic) bool bEngineInit;
+/** True after the end of FAppEntry::PlatformInit */
+@property (atomic) bool bPlatformInit;
+
 
 /** Delays game initialization slightly in case we have a URL launch to handle */
 @property (retain) NSTimer* CommandLineParseTimer;
@@ -246,10 +246,19 @@ APPLICATIONCORE_API
 
 #if !PLATFORM_TVOS && !PLATFORM_VISIONOS && !defined(SWIFT_IMPORT)
   +(EDeviceScreenOrientation) ConvertFromUIInterfaceOrientation:(UIInterfaceOrientation)Orientation;
+  +(UIInterfaceOrientation) GetUIInterfaceOrientationFromUIDeviceOrientation:(UIDeviceOrientation)Orientation;
+  +(UIInterfaceOrientationMask) GetMaskFromUIInterfaceOrientation:(UIInterfaceOrientation)Orientation;
+  +(UIInterfaceOrientationMask) GetMaskFromEDeviceScreenOrientation:(EDeviceScreenOrientation)Orientation;
+#endif
+
+#if !PLATFORM_TVOS && !PLATFORM_VISIONOS
+  +(void)UpdateSupportedInterfaceOrientations;
 #endif
 
 -(bool)IsUpdateAvailable;
 
+-(void) didRotate:(NSNotification*)notification;
+
 @property (assign) bool bAudioSessionInitialized;
 
 /** TRUE if the device is playing background music and we want to allow that */
diff --git a/Engine/Source/Runtime/ApplicationCore/Public/IOS/IOSApplication.h b/Engine/Source/Runtime/ApplicationCore/Public/IOS/IOSApplication.h
--- a/Engine/Source/Runtime/ApplicationCore/Public/IOS/IOSApplication.h
+++ b/Engine/Source/Runtime/ApplicationCore/Public/IOS/IOSApplication.h
@@ -32,8 +32,9 @@ public:
 	virtual void AddExternalInputDevice(TSharedPtr<class IInputDevice> InputDevice);
 
 #if !PLATFORM_TVOS
-	static void OrientationChanged(UIInterfaceOrientation orientation);
-
+#if !PLATFORM_VISIONOS
+	static void OrientationChanged(UIInterfaceOrientation NewOrientation);
+#endif
 	static UIInterfaceOrientation CachedOrientation;
 #endif
 
@@ -73,5 +74,7 @@ private:
 	static FCriticalSection CriticalSection;
 	static bool bOrientationChanged;
 
-	void CacheDisplayMetrics();
+#if !PLATFORM_TVOS
+	void CacheDisplayMetrics(UIInterfaceOrientation Orientation);
+#endif
 };
diff --git a/Engine/Source/Runtime/ApplicationCore/Public/IOS/IOSView.h b/Engine/Source/Runtime/ApplicationCore/Public/IOS/IOSView.h
--- a/Engine/Source/Runtime/ApplicationCore/Public/IOS/IOSView.h
+++ b/Engine/Source/Runtime/ApplicationCore/Public/IOS/IOSView.h
@@ -80,7 +80,12 @@ APPLICATIONCORE_API
 	BOOL bSecureTextEntry;
 	
 	volatile int KeyboardShowCount;
-	NSUInteger SupportedInterfaceOrientations;
+
+#if !PLATFORM_TVOS
+	UIInterfaceOrientationMask SupportedInterfaceOrientations;
+	// the orientation for the current frame we are displaying, or 0 if none is available
+	UIInterfaceOrientationMask LastFrameInterfaceOrientationMask;
+#endif
 }
 
 #if WITH_ACCESSIBILITY
@@ -134,6 +139,13 @@ APPLICATIONCORE_API
 @interface IOSViewController : UIViewController
 #endif
 {
-
 }
+
+#if !PLATFORM_TVOS && !PLATFORM_VISIONOS
+-(void)notifyPresentAfterRotateOrientationMask : (UIInterfaceOrientationMask)NewOrientationMask withSizeX : (unsigned int)SizeX withSizeY : (unsigned int)SizeY;
+
+- (UIInterfaceOrientationMask)supportedInterfaceOrientations_Internal;
+- (UIInterfaceOrientationMask)supportedInterfaceOrientations;
+#endif
+
 @end
diff --git a/Engine/Source/Runtime/Core/Private/IOS/IOSPlatformMisc.cpp b/Engine/Source/Runtime/Core/Private/IOS/IOSPlatformMisc.cpp
--- a/Engine/Source/Runtime/Core/Private/IOS/IOSPlatformMisc.cpp
+++ b/Engine/Source/Runtime/Core/Private/IOS/IOSPlatformMisc.cpp
@@ -499,30 +499,16 @@ void FIOSPlatformMisc::SetAllowedDeviceOrientation(EDeviceScreenOrientation NewA
 	AllowedDeviceOrientation = NewAllowedDeviceOrientation;
 	
 #if !PLATFORM_TVOS && !PLATFORM_VISIONOS
-	[IOSAppDelegate GetDelegate].IOSView->SupportedInterfaceOrientations = GetUIInterfaceOrientationMask(NewAllowedDeviceOrientation);
-#endif
-	
-	// Notify the view controller about a change in supported interface orientations.
-	
-#if !PLATFORM_TVOS
-	dispatch_async(dispatch_get_main_queue(),
-		^{
-			if (@available(iOS 16.0, visionOS 1.0, *))
-			{
-				[[IOSAppDelegate GetDelegate].IOSController setNeedsUpdateOfSupportedInterfaceOrientations];
-			}
-#if !PLATFORM_VISIONOS
-			else
-			{
-// Deprecated iOS16 - suppress warning if deployment target	>= iOS16
-#pragma clang diagnostic push
-#pragma clang diagnostic ignored "-Wdeprecated-declarations"
-				[UIViewController attemptRotationToDeviceOrientation];
-#pragma clang diagnostic pop
-			}
-#endif
-		}
-	);
+	IOSAppDelegate* AppDelegate = [IOSAppDelegate GetDelegate];
+	AppDelegate.IOSView->SupportedInterfaceOrientations = GetUIInterfaceOrientationMask(NewAllowedDeviceOrientation);
+
+	FPlatformMisc::LowLevelOutputDebugStringf(TEXT("[rotation] SetAllowedDeviceOrientation mask %d\n"), (int)GetUIInterfaceOrientationMask(NewAllowedDeviceOrientation));
+
+	// Check rotation state
+	[AppDelegate didRotate:nil];
+
+	// Attempt to rotate to desired orientation
+	[IOSAppDelegate UpdateSupportedInterfaceOrientations];
 #endif
 }
 
diff --git a/Engine/Source/Runtime/Launch/Private/IOS/IOSAppDelegateConsoleHandling.cpp b/Engine/Source/Runtime/Launch/Private/IOS/IOSAppDelegateConsoleHandling.cpp
--- a/Engine/Source/Runtime/Launch/Private/IOS/IOSAppDelegateConsoleHandling.cpp
+++ b/Engine/Source/Runtime/Launch/Private/IOS/IOSAppDelegateConsoleHandling.cpp
@@ -10,6 +10,8 @@ extern bool GShowSplashScreen;
 extern bool GDebugConsoleOpen;
 #endif
 
+@class IOSViewController;
+
 @implementation IOSAppDelegate (ConsoleHandling)
 
 #if !UE_BUILD_SHIPPING && !PLATFORM_TVOS
@@ -138,7 +140,7 @@ extern bool GDebugConsoleOpen;
  */
 - (void)ShowAlert:(NSMutableArray*)StringArray
 {
-	UIViewController *ViewController = [IOSAppDelegate GetDelegate].IOSController;
+	UIViewController *ViewController = (UIViewController*)[IOSAppDelegate GetDelegate].IOSController;
 	if (GShowSplashScreen)
 	{
 		if ([[IOSAppDelegate GetDelegate].Window viewWithTag : 200] != nil)
diff --git a/Engine/Source/Runtime/Launch/Private/IOS/LaunchIOS.cpp b/Engine/Source/Runtime/Launch/Private/IOS/LaunchIOS.cpp
--- a/Engine/Source/Runtime/Launch/Private/IOS/LaunchIOS.cpp
+++ b/Engine/Source/Runtime/Launch/Private/IOS/LaunchIOS.cpp
@@ -317,9 +317,6 @@ void FAppEntry::PreInit(IOSAppDelegate* AppDelegate, UIApplication* Application)
 
 	[AppDelegate.Window setRootViewController:IOSController];
 
-	// windows owns it now
-//	[IOSController release];
-
 #if !PLATFORM_TVOS
 	// reset badge count on launch
 	// Deprecated iOS17 - suppress warning if deployment target	> iOS17
@@ -428,6 +425,7 @@ void FAppEntry::PlatformInit()
 	FDisplayMetrics::RebuildDisplayMetrics(DisplayMetrics);
 	FSystemResolution::RequestResolutionChange(DisplayMetrics.PrimaryDisplayWidth, DisplayMetrics.PrimaryDisplayHeight, EWindowMode::Fullscreen);
 	IConsoleManager::Get().CallAllConsoleVariableSinks();
+	AppDelegate.bPlatformInit = true;
 }
 
 extern TcpConsoleListener *ConsoleListener;
