UE-300382 Update key modifiers during OLE drag and drop changes

#tests Tested drag and dropping from the file explorer on top of visible app, tested by dragging over taskbar to get the activate path
#rnx

diff --git a/Engine/Source/Runtime/ApplicationCore/Private/Windows/WindowsApplication.cpp b/Engine/Source/Runtime/ApplicationCore/Private/Windows/WindowsApplication.cpp
--- a/Engine/Source/Runtime/ApplicationCore/Private/Windows/WindowsApplication.cpp
+++ b/Engine/Source/Runtime/ApplicationCore/Private/Windows/WindowsApplication.cpp
@@ -4010,8 +4010,12 @@ HRESULT FWindowsApplication::OnOLEDragEnter( const HWND HWnd, const FDragDropOLE
 		return 0;
 	}
 
+	CachedDragAndDropKeyState = 0;
+
 	if ( Window->IsEnabled() )
 	{
+		UpdateAllModifierKeyStates();
+		CachedDragAndDropKeyState = KeyState;
 		if ((OLEData.Type & FDragDropOLEData::Text) && (OLEData.Type & FDragDropOLEData::Files))
 		{
 			*CursorEffect = MessageHandler->OnDragEnterExternal(Window.ToSharedRef(), OLEData.OperationText, OLEData.OperationFilenames);
@@ -4041,6 +4045,11 @@ HRESULT FWindowsApplication::OnOLEDragOver( const HWND HWnd, DWORD KeyState, POI
 	{
 		if ( Window->IsEnabled() )
 		{
+			if (KeyState != CachedDragAndDropKeyState)
+			{
+				UpdateAllModifierKeyStates();
+				CachedDragAndDropKeyState = KeyState;
+			}
 			*CursorEffect = MessageHandler->OnDragOver( Window.ToSharedRef() );
 		}
 		else
@@ -4074,6 +4083,12 @@ HRESULT FWindowsApplication::OnOLEDrop( const HWND HWnd, const FDragDropOLEData&
 	{
 		if ( Window->IsEnabled() )
 		{
+			if (KeyState != CachedDragAndDropKeyState)
+			{
+				UpdateAllModifierKeyStates();
+				CachedDragAndDropKeyState = KeyState;
+			}
+
 			*CursorEffect = MessageHandler->OnDragDrop(Window.ToSharedRef());
 		}
 		else
diff --git a/Engine/Source/Runtime/ApplicationCore/Public/Windows/WindowsApplication.h b/Engine/Source/Runtime/ApplicationCore/Public/Windows/WindowsApplication.h
--- a/Engine/Source/Runtime/ApplicationCore/Public/Windows/WindowsApplication.h
+++ b/Engine/Source/Runtime/ApplicationCore/Public/Windows/WindowsApplication.h
@@ -612,6 +612,7 @@ private:
 	FIntPoint LastCursorPointPreWrap;
 	int32 NumPreWrapMsgsToRespect;
 	RECT ClipCursorRect;
+	DWORD CachedDragAndDropKeyState;
 
 	double RepeatTime[EMouseButtons::Invalid];
 };
