Cache Win32 cursor info using task graph to reduce time spent on the game thread

diff --git a/Engine/Source/Runtime/ApplicationCore/Private/Windows/WindowsApplication.cpp b/Engine/Source/Runtime/ApplicationCore/Private/Windows/WindowsApplication.cpp
--- a/Engine/Source/Runtime/ApplicationCore/Private/Windows/WindowsApplication.cpp
+++ b/Engine/Source/Runtime/ApplicationCore/Private/Windows/WindowsApplication.cpp
@@ -1200,7 +1200,6 @@ FWindowsApplication::FWindowsApplication( const HINSTANCE HInstance, const HICON
 	, UIAManager(new FWindowsUIAManager(*this))
 #endif
 	, bSimulatingHighPrecisionMouseInputForRDP(false)
-	, CachedPreHighPrecisionMousePosForRDP(FIntPoint::ZeroValue)
 	, LastCursorPoint(FIntPoint::ZeroValue)
 	, LastCursorPointPreWrap(FIntPoint::ZeroValue)
 	, NumPreWrapMsgsToRespect(0)
@@ -1469,6 +1468,17 @@ FWindowsApplication::~FWindowsApplication()
 		TextInputMethodSystem->Terminate();
 	}
 
+	if (bAsyncCursorTaskActive)
+	{
+		const UE::FEventCountToken Token = AsyncCursorEventCount.PrepareWait();
+		if (bAsyncCursorTaskActive)
+		{
+			AsyncCursorEventCount.Wait(Token);
+		}
+	}
+
+	AsyncCursorTask.Wait();
+
 	::CoUninitialize();
 	OleUninitialize();
 }
@@ -1566,20 +1576,29 @@ static TSharedPtr< FWindowsWindow > FindWindowByHWND(const TArray< TSharedRef< F
 
 bool FWindowsApplication::IsCursorDirectlyOverSlateWindow() const
 {
-	QUICK_SCOPE_CYCLE_COUNTER(STAT_STAT_IsCursorDirectlyOverSlateWindow);
+	QUICK_SCOPE_CYCLE_COUNTER(STAT_IsCursorDirectlyOverSlateWindow);
 
-	POINT CursorPos;
-	BOOL bGotPoint = ::GetCursorPos(&CursorPos);
-	if (bGotPoint)
+	FRWScopeLock Lock(AsyncCursorRWLock, SLT_ReadOnly);
+	if (bAsyncCursorPosValid)
 	{
-		HWND hWnd = ::WindowFromPoint(CursorPos);
-
-		TSharedPtr< FWindowsWindow > SlatWindowUnderCursor = FindWindowByHWND(Windows, hWnd);
-		return SlatWindowUnderCursor.IsValid();
+		TSharedPtr< FWindowsWindow > SlateWindowUnderCursor = FindWindowByHWND(Windows, AsyncCursorWindowHandle);
+		return SlateWindowUnderCursor.IsValid();
 	}
+
 	return false;
 }
 
+bool FWindowsApplication::StaticGetCursorPosition(POINT& OutSursorPos)
+{
+	return WindowsApplication ? WindowsApplication->GetCursorPosition(OutSursorPos) : false;
+}
+
+bool FWindowsApplication::GetCursorPosition(POINT& OutSursorPos)
+{
+	FRWScopeLock Lock(AsyncCursorRWLock, SLT_ReadOnly);
+	OutSursorPos = AsyncCursorPos;
+	return bAsyncCursorPosValid;
+}
 
 void FWindowsApplication::SetCapture( const TSharedPtr< FGenericWindow >& InWindow )
 {
@@ -1613,17 +1632,14 @@ void FWindowsApplication::SetHighPrecisionMouseMode( const bool Enable, const TS
 			::GetClipCursor(&ClipCursorRect);
 
 			//UE_LOG(LogWindowsDesktop, Log, TEXT("Entering High Precision to Top: %d Bottom: %d Left: %d Right: %d ---- Starting x: %d y: %d"), ClipCursorRect.top, ClipCursorRect.bottom, ClipCursorRect.left, ClipCursorRect.right, CursorPos.x, CursorPos.y);
-		
-			CachedPreHighPrecisionMousePosForRDP = FIntPoint(CursorPos.x, CursorPos.y);
-			
+
 			LastCursorPoint = FIntPoint(CursorPos.x, CursorPos.y);
-			
+
 			LastCursorPointPreWrap = FIntPoint::ZeroValue;
 			NumPreWrapMsgsToRespect = 0;
 		}
 		else
 		{
-			CachedPreHighPrecisionMousePosForRDP = FIntPoint(INT_MAX, INT_MAX);
 			bSimulatingHighPrecisionMouseInputForRDP = false;
 		}
 	}
@@ -4271,6 +4287,7 @@ void FWindowsApplication::PostPumpMessages()
 {
 	WindowsApplication->UpdateHighPrecisionMouseMode();
 	WindowsApplication->ProcessWorkerInputs();
+	WindowsApplication->CacheCursorInfo();
 }
 
 void FWindowsApplication::ProcessRawMouseInput(const RAWMOUSE& Mouse, RawMouseInputResult& OutResult)
@@ -4459,26 +4476,23 @@ void FWindowsApplication::ProcessWorkerInputs()
 			TSharedPtr< FWindowsWindow > Window;
 
 			POINT ButtonCursorPoint;
-			BOOL bGotPoint;
-			{
-				bGotPoint = ::GetCursorPos(&ButtonCursorPoint);
-			}
 
-			if (bGotPoint)
 			{
-				HWND hWnd = ::WindowFromPoint(ButtonCursorPoint);
-
-				if (hWnd)
+				FRWScopeLock Lock(AsyncCursorRWLock, SLT_ReadOnly);
+				if (bAsyncCursorPosValid)
+				{
+					ButtonCursorPoint = AsyncCursorPos;
+					if (AsyncCursorWindowHandle)
+					{
+						Window = FindWindowByHWND(Windows, AsyncCursorWindowHandle);
+					}
+				}
+				else
 				{
-					Window = FindWindowByHWND(Windows, hWnd);
+					ButtonCursorPoint.x = 0;
+					ButtonCursorPoint.y = 0;
 				}
 			}
-			else
-			{
-				UE_LOG(LogWindowsDesktop, Error, TEXT("Could not GetCursorPos for raw input. Error code: %u"), GetLastError());
-				ButtonCursorPoint.x = 0;
-				ButtonCursorPoint.y = 0;
-			}
 
 			const FVector2D ButtonCursorPos = FVector2D(ButtonCursorPoint.x, ButtonCursorPoint.y);
 			const double CurrentTime = FPlatformTime::Seconds();
@@ -4525,6 +4539,43 @@ void FWindowsApplication::ProcessMouseButton(double CurrentTime, const TSharedPt
 	}
 }
 
+void FWindowsApplication::CacheCursorInfo()
+{
+	TRACE_CPUPROFILER_EVENT_SCOPE(CacheCursorInfo);
+
+	if (bAsyncCursorTaskActive)
+	{
+		const UE::FEventCountToken Token = AsyncCursorEventCount.PrepareWait();
+		if (bAsyncCursorTaskActive)
+		{
+			AsyncCursorEventCount.Wait(Token);
+		}
+	}
+
+	bAsyncCursorTaskActive = true;
+
+	AsyncCursorTask = UE::Tasks::Launch(UE_SOURCE_LOCATION,
+		[&CursorRWLock = AsyncCursorRWLock,
+		&bCursorPosValid = bAsyncCursorPosValid,
+		&CursorPos = AsyncCursorPos,
+		&CursorWindowHandle = AsyncCursorWindowHandle,
+		&bCursorTaskActive = bAsyncCursorTaskActive,
+		&CursorEventCount = AsyncCursorEventCount]
+		{
+			TRACE_CPUPROFILER_EVENT_SCOPE(AsyncCursorTask);
+
+			{
+				FRWScopeLock Lock(CursorRWLock, SLT_Write);
+				bCursorPosValid = ::GetCursorPos(&CursorPos);
+				CursorWindowHandle = bCursorPosValid ? ::WindowFromPoint(CursorPos) : NULL;
+			}
+
+			bCursorTaskActive = false;
+			CursorEventCount.Notify();
+		}
+	);
+}
+
 // Restore the windowsx.h macro for IsMaximized
 #pragma pop_macro("IsMaximized")
 
diff --git a/Engine/Source/Runtime/ApplicationCore/Private/Windows/WindowsCursor.cpp b/Engine/Source/Runtime/ApplicationCore/Private/Windows/WindowsCursor.cpp
--- a/Engine/Source/Runtime/ApplicationCore/Private/Windows/WindowsCursor.cpp
+++ b/Engine/Source/Runtime/ApplicationCore/Private/Windows/WindowsCursor.cpp
@@ -7,6 +7,7 @@
 #include "Misc/CoreMisc.h"
 #include "Math/Vector2D.h"
 #include "Math/Color.h"
+#include "Windows/WindowsApplication.h"
 #include "Windows/WindowsWindow.h"
 #include "HAL/IConsoleManager.h"
 #include "HAL/PlatformProcess.h"
@@ -229,10 +230,13 @@ void* FWindowsCursor::CreateCursorFromRGBABuffer(const FColor* Pixels, int32 Wid
 
 FVector2D FWindowsCursor::GetPosition() const
 {
-	POINT CursorPos;
-	::GetCursorPos(&CursorPos);
+	POINT OutCursorPos;
+	if (!FWindowsApplication::StaticGetCursorPosition(OutCursorPos))
+	{
+		::GetCursorPos(&OutCursorPos);
+	}
 
-	return FVector2D( CursorPos.x, CursorPos.y );
+	return FVector2D(OutCursorPos.x, OutCursorPos.y);
 }
 
 void FWindowsCursor::SetPosition( const int32 X, const int32 Y )
diff --git a/Engine/Source/Runtime/ApplicationCore/Public/Windows/WindowsApplication.h b/Engine/Source/Runtime/ApplicationCore/Public/Windows/WindowsApplication.h
--- a/Engine/Source/Runtime/ApplicationCore/Public/Windows/WindowsApplication.h
+++ b/Engine/Source/Runtime/ApplicationCore/Public/Windows/WindowsApplication.h
@@ -8,6 +8,7 @@
 #include "HAL/IConsoleManager.h"
 #include "GenericPlatform/GenericApplication.h"
 #include "GenericPlatform/IInputInterface.h"
+#include "Tasks/Task.h"
 #include "Windows/AllowWindowsPlatformTypes.h"
 	#include <Ole2.h>
 	#include <oleidl.h>
@@ -314,6 +315,9 @@ public:
 	/** Process raw mouse inputs gathered from a helper thread/window. */
 	void ProcessWorkerInputs();
 
+	/** Cache Win32 calls that can take a while. */
+	void CacheCursorInfo();
+
 	/** Virtual destructor. */
 	APPLICATIONCORE_API virtual ~FWindowsApplication();
 
@@ -375,6 +379,8 @@ public:
 	APPLICATIONCORE_API virtual bool IsGamepadAttached() const override;
 	APPLICATIONCORE_API virtual FModifierKeysState GetModifierKeys() const override;
 	APPLICATIONCORE_API virtual bool IsCursorDirectlyOverSlateWindow() const override;
+	static bool StaticGetCursorPosition(POINT& OutSursorPos);
+	bool GetCursorPosition(POINT& OutSursorPos);
 	APPLICATIONCORE_API virtual FPlatformRect GetWorkArea( const FPlatformRect& CurrentWindow ) const override;
 	APPLICATIONCORE_API virtual void GetInitialDisplayMetrics( FDisplayMetrics& OutDisplayMetrics ) const override;
 	APPLICATIONCORE_API virtual EWindowTitleAlignment::Type GetWindowTitleAlignment() const override;
@@ -530,6 +536,14 @@ private:
 	bool bCachedUseWorkerThreadForRawInput;
 	HWND HighPrecisionMouseInputWindowHandle;
 
+	UE::FEventCount AsyncCursorEventCount;
+	std::atomic<bool> bAsyncCursorTaskActive = false;
+	UE::Tasks::FTask AsyncCursorTask;
+	mutable FRWLock AsyncCursorRWLock;
+	BOOL bAsyncCursorPosValid{};
+	POINT AsyncCursorPos{};
+	HWND AsyncCursorWindowHandle = NULL;
+
 	bool bIsMouseAttached;
 
 	bool bForceActivateByMouse;
@@ -607,7 +621,6 @@ private:
 
 	bool bSimulatingHighPrecisionMouseInputForRDP;
 
-	FIntPoint CachedPreHighPrecisionMousePosForRDP;
 	FIntPoint LastCursorPoint;
 	FIntPoint LastCursorPointPreWrap;
 	int32 NumPreWrapMsgsToRespect;
