Linux windows need to be started as hidden to match Window's behavior.

#rnx

diff --git a/Engine/Source/Runtime/ApplicationCore/Private/Linux/LinuxWindow.cpp b/Engine/Source/Runtime/ApplicationCore/Private/Linux/LinuxWindow.cpp
--- a/Engine/Source/Runtime/ApplicationCore/Private/Linux/LinuxWindow.cpp
+++ b/Engine/Source/Runtime/ApplicationCore/Private/Linux/LinuxWindow.cpp
@@ -115,12 +115,15 @@ void FLinuxWindow::Initialize( FLinuxApplication* const Application, const TShar
         //      The SDL window doesn't need to be reshaped.
         //      the size of the window you input is the sizeof the client.
 	bool bIsResizable = false;
-        SDL_PropertiesID Props = SDL_CreateProperties();
-        SDL_SetStringProperty( Props, SDL_PROP_WINDOW_CREATE_TITLE_STRING, TCHAR_TO_ANSI( *Definition->Title ) );
-        SDL_SetNumberProperty( Props, SDL_PROP_WINDOW_CREATE_X_NUMBER, X );
-        SDL_SetNumberProperty( Props, SDL_PROP_WINDOW_CREATE_Y_NUMBER, Y );
-        SDL_SetNumberProperty( Props, SDL_PROP_WINDOW_CREATE_WIDTH_NUMBER, ClientWidth );
-        SDL_SetNumberProperty( Props, SDL_PROP_WINDOW_CREATE_HEIGHT_NUMBER, ClientHeight );
+
+	SDL_PropertiesID Props = SDL_CreateProperties();
+	SDL_SetStringProperty( Props, SDL_PROP_WINDOW_CREATE_TITLE_STRING, TCHAR_TO_ANSI( *Definition->Title ) );
+	SDL_SetNumberProperty( Props, SDL_PROP_WINDOW_CREATE_X_NUMBER, X );
+	SDL_SetNumberProperty( Props, SDL_PROP_WINDOW_CREATE_Y_NUMBER, Y );
+	SDL_SetNumberProperty( Props, SDL_PROP_WINDOW_CREATE_WIDTH_NUMBER, ClientWidth );
+	SDL_SetNumberProperty( Props, SDL_PROP_WINDOW_CREATE_HEIGHT_NUMBER, ClientHeight );
+
+	SDL_SetBooleanProperty(Props, SDL_PROP_WINDOW_CREATE_HIDDEN_BOOLEAN, true);
 
 	UE_LOG(LogLinuxWindowType, Verbose, TEXT("Linux Window: X=%d Y=%d, Width=%d, Height=%d"), X, Y, ClientWidth, ClientHeight);
 	switch(FLinuxPlatformApplicationMisc::WindowStyle())
@@ -149,7 +152,6 @@ void FLinuxWindow::Initialize( FLinuxApplication* const Application, const TShar
 		!Definition->HasSizingFrame && !Definition->IsModalWindow &&
 		!Definition->IsRegularWindow && Definition->SizeWillChangeOften)
 	{
-		SDL_SetBooleanProperty( Props, SDL_PROP_WINDOW_CREATE_HIDDEN_BOOLEAN, true );
 		// We don't set SDL_PROP_WINDOW_CREATE_TOOLTIP_BOOLEAN to true here because it requires a parent and Slate isn't set up
 		// to provide one for tooltips
 		bIsTooltipWindow = true;
@@ -186,7 +188,6 @@ void FLinuxWindow::Initialize( FLinuxApplication* const Application, const TShar
 		!Definition->IsModalWindow && !Definition->IsRegularWindow &&
 		bShouldActivate && !Definition->SizeWillChangeOften)
 	{
-		SDL_SetBooleanProperty( Props, SDL_PROP_WINDOW_CREATE_HIDDEN_BOOLEAN, true );
 		SDL_SetBooleanProperty( Props, SDL_PROP_WINDOW_CREATE_TRANSPARENT_BOOLEAN, true );
 		SDL_SetBooleanProperty( Props, SDL_PROP_WINDOW_CREATE_BORDERLESS_BOOLEAN, true );
 		bIsPopupWindow = true;
@@ -199,7 +200,6 @@ void FLinuxWindow::Initialize( FLinuxApplication* const Application, const TShar
 		!Definition->IsModalWindow && !Definition->IsRegularWindow &&
 		!bShouldActivate && !Definition->SizeWillChangeOften)
 	{
-		SDL_SetBooleanProperty( Props, SDL_PROP_WINDOW_CREATE_HIDDEN_BOOLEAN, true );
 		SDL_SetBooleanProperty( Props, SDL_PROP_WINDOW_CREATE_TRANSPARENT_BOOLEAN, true );
 		SDL_SetBooleanProperty( Props, SDL_PROP_WINDOW_CREATE_BORDERLESS_BOOLEAN, true );
 		bIsConsoleWindow = true;
