Linux SDL3 Notification Window handling changes

#rnx

diff --git a/Engine/Source/Runtime/ApplicationCore/Private/Linux/LinuxWindow.cpp b/Engine/Source/Runtime/ApplicationCore/Private/Linux/LinuxWindow.cpp
--- a/Engine/Source/Runtime/ApplicationCore/Private/Linux/LinuxWindow.cpp
+++ b/Engine/Source/Runtime/ApplicationCore/Private/Linux/LinuxWindow.cpp
@@ -164,7 +164,8 @@ void FLinuxWindow::Initialize( FLinuxApplication* const Application, const TShar
 		!Definition->IsModalWindow && !Definition->IsRegularWindow &&
 		!bShouldActivate && Definition->SizeWillChangeOften)
 	{
-		SDLWindowType = "_NET_WM_WINDOW_TYPE_NOTIFICATION";
+		SDL_SetBooleanProperty( Props, SDL_PROP_WINDOW_CREATE_UTILITY_BOOLEAN, true );
+		SDL_SetPointerProperty( Props, SDL_PROP_WINDOW_CREATE_PARENT_POINTER, InParent->GetHWnd() );		
 		bIsNotificationWindow = true;
 		UE_LOG(LogLinuxWindowType, Verbose, TEXT("*** New Window is a Notification Window ***"));
 	}
@@ -177,10 +178,21 @@ void FLinuxWindow::Initialize( FLinuxApplication* const Application, const TShar
 	{
 		SDL_SetBooleanProperty( Props, SDL_PROP_WINDOW_CREATE_MODAL_BOOLEAN, true );
         SDL_SetPointerProperty( Props, SDL_PROP_WINDOW_CREATE_PARENT_POINTER, InParent->GetHWnd() );
-		SDLWindowType = "_NET_WM_WINDOW_TYPE_NOTIFICATION";
 		bIsNotificationWindow = true;
 		UE_LOG(LogLinuxWindowType, Verbose, TEXT("*** New Window is another Notification Window ***"));
 	}
+	else if (!InParent.IsValid() && !Definition->HasOSWindowBorder &&
+		Definition->AcceptsInput && !Definition->IsTopmostWindow && 
+		!Definition->AppearsInTaskbar && !Definition->HasSizingFrame &&
+		!Definition->IsModalWindow && !Definition->IsRegularWindow &&
+		!bShouldActivate && Definition->SizeWillChangeOften && 
+		Definition->Type == EWindowType::Notification)
+	{
+		// Even parentless notification windows should be use the right window type
+		SDL_SetBooleanProperty( Props, SDL_PROP_WINDOW_CREATE_UTILITY_BOOLEAN, true );
+		bIsNotificationWindow = true;
+		UE_LOG(LogLinuxWindowType, Verbose, TEXT("*** New Window is a parentless Notification Window ***"));
+	}
 	// This is a popup menu window?
 	else if (InParent.IsValid() && !Definition->HasOSWindowBorder &&
 		Definition->AcceptsInput && !Definition->IsTopmostWindow && 
