Updated FSlateApplication::ProcessReply to skip enabling high precision mode if it's already enabled. This optimization already existed when *releasing* high precision mode
Updated SetCapture to cache the captured window handle, as recapturing the same window can be quite expensive
Added handling for WM_CAPTURECHANGED message to update cached handle in case we lose capture due to external factors
These changes save up to .5ms on redundant API calls

[RN] Updated mouse capture and high precision mode handling on Windows to save up to .5ms on redundant API calls

diff --git a/Engine/Source/Runtime/ApplicationCore/Private/Windows/WindowsApplication.cpp b/Engine/Source/Runtime/ApplicationCore/Private/Windows/WindowsApplication.cpp
--- a/Engine/Source/Runtime/ApplicationCore/Private/Windows/WindowsApplication.cpp
+++ b/Engine/Source/Runtime/ApplicationCore/Private/Windows/WindowsApplication.cpp
@@ -1600,16 +1600,22 @@ bool FWindowsApplication::GetCursorPosition(POINT& OutSursorPos)
 	return bAsyncCursorPosValid;
 }
 
-void FWindowsApplication::SetCapture( const TSharedPtr< FGenericWindow >& InWindow )
+void FWindowsApplication::SetCapture(const TSharedPtr< FGenericWindow >& InWindow)
 {
-	if ( InWindow.IsValid() )
-	{
-		::SetCapture( (HWND)InWindow->GetOSWindowHandle() );
-	}
-	else
+	HWND DesiredCaptureHandle = (InWindow.IsValid() ? (HWND)InWindow->GetOSWindowHandle() : NULL);
+
+	if (DesiredCaptureHandle != CaptureWindowHandle)
 	{
-		::ReleaseCapture();
+		if (DesiredCaptureHandle)
+		{
+			::SetCapture(DesiredCaptureHandle);
+		}
+		else
+		{
+			::ReleaseCapture();
+		}
 	}
+	CaptureWindowHandle = DesiredCaptureHandle;
 }
 
 void* FWindowsApplication::GetCapture( void ) const
@@ -2986,6 +2992,14 @@ int32 FWindowsApplication::ProcessMessage( HWND hwnd, uint32 msg, WPARAM wParam,
 			}
 			break;
 
+		case WM_CAPTURECHANGED:
+			{
+				// This will usually be null, but we'll store it anyway in case capture is switching between two windows within the same process.
+				// Otherwise we'd risk clobbering the cached handle from the new capturing window
+				CaptureWindowHandle = (HWND)lParam;
+			}
+			break;
+
 #if WITH_ACCESSIBILITY && UE_WINDOWS_USING_UIA
 		case WM_GETOBJECT:
 		{
diff --git a/Engine/Source/Runtime/ApplicationCore/Public/Windows/WindowsApplication.h b/Engine/Source/Runtime/ApplicationCore/Public/Windows/WindowsApplication.h
--- a/Engine/Source/Runtime/ApplicationCore/Public/Windows/WindowsApplication.h
+++ b/Engine/Source/Runtime/ApplicationCore/Public/Windows/WindowsApplication.h
@@ -535,6 +535,7 @@ private:
 	bool bUsingHighPrecisionMouseInput;
 	bool bCachedUseWorkerThreadForRawInput;
 	HWND HighPrecisionMouseInputWindowHandle;
+	HWND CaptureWindowHandle;
 
 	UE::FEventCount AsyncCursorEventCount;
 	std::atomic<bool> bAsyncCursorTaskActive = false;
diff --git a/Engine/Source/Runtime/Slate/Private/Framework/Application/SlateApplication.cpp b/Engine/Source/Runtime/Slate/Private/Framework/Application/SlateApplication.cpp
--- a/Engine/Source/Runtime/Slate/Private/Framework/Application/SlateApplication.cpp
+++ b/Engine/Source/Runtime/Slate/Private/Framework/Application/SlateApplication.cpp
@@ -3481,7 +3481,7 @@ void FSlateApplication::ProcessReply( const FWidgetPath& CurrentEventPath, const
 		{
 			// If the mouse is being captured or released, toggle high precision raw input if requested by the reply.
 			// Raw input is only used with mouse capture
-			if (TheReply.ShouldUseHighPrecisionMouse())
+			if (TheReply.ShouldUseHighPrecisionMouse() && !PlatformApplication->IsUsingHighPrecisionMouseMode())
 			{
 				const TSharedRef< SWindow> Window = CurrentEventPath.GetWindow();
 				PlatformApplication->SetCapture(Window->GetNativeWindow());
