UE-306958 Set the async cursor pos cache directly when setting the position from the api

#tests Ran the automated test that was failing. Verified through manual breakpoints that we're hitting what we expect
#rnx

diff --git a/Engine/Source/Runtime/ApplicationCore/Private/Windows/WindowsApplication.cpp b/Engine/Source/Runtime/ApplicationCore/Private/Windows/WindowsApplication.cpp
--- a/Engine/Source/Runtime/ApplicationCore/Private/Windows/WindowsApplication.cpp
+++ b/Engine/Source/Runtime/ApplicationCore/Private/Windows/WindowsApplication.cpp
@@ -1588,19 +1588,33 @@ bool FWindowsApplication::IsCursorDirectlyOverSlateWindow() const
 	return false;
 }
 
-bool FWindowsApplication::StaticGetCursorPosition(POINT& OutSursorPos)
+bool FWindowsApplication::StaticGetCursorPosition(POINT& OutCursorPos)
 {
-	return WindowsApplication ? WindowsApplication->GetCursorPosition(OutSursorPos) : false;
+	return WindowsApplication ? WindowsApplication->GetCursorPosition(OutCursorPos) : false;
 }
 
-bool FWindowsApplication::GetCursorPosition(POINT& OutSursorPos)
+bool FWindowsApplication::GetCursorPosition(POINT& OutCursorPos)
 {
 	FRWScopeLock Lock(AsyncCursorRWLock, SLT_ReadOnly);
-	OutSursorPos = AsyncCursorPos;
+	OutCursorPos = AsyncCursorPos;
 	return bAsyncCursorPosValid;
 }
 
-void FWindowsApplication::SetCapture(const TSharedPtr< FGenericWindow >& InWindow)
+void FWindowsApplication::ForceUpdateAsyncCursor(POINT position, HWND windowUnder)
+{
+	if (WindowsApplication == nullptr)
+	{
+		return;
+	}
+	FRWScopeLock Lock(WindowsApplication->AsyncCursorRWLock, SLT_Write);
+	if (WindowsApplication->bAsyncCursorPosValid)
+	{
+		WindowsApplication->AsyncCursorWindowHandle = windowUnder;
+		WindowsApplication->AsyncCursorPos = position;
+	}
+}
+
+void FWindowsApplication::SetCapture( const TSharedPtr< FGenericWindow >& InWindow )
 {
 	HWND DesiredCaptureHandle = (InWindow.IsValid() ? (HWND)InWindow->GetOSWindowHandle() : NULL);
 
diff --git a/Engine/Source/Runtime/ApplicationCore/Private/Windows/WindowsCursor.cpp b/Engine/Source/Runtime/ApplicationCore/Private/Windows/WindowsCursor.cpp
--- a/Engine/Source/Runtime/ApplicationCore/Private/Windows/WindowsCursor.cpp
+++ b/Engine/Source/Runtime/ApplicationCore/Private/Windows/WindowsCursor.cpp
@@ -241,7 +241,16 @@ FVector2D FWindowsCursor::GetPosition() const
 
 void FWindowsCursor::SetPosition( const int32 X, const int32 Y )
 {
-	::SetCursorPos( X, Y );
+	if (::SetCursorPos(X, Y))
+	{
+		POINT newPosition;
+		if (!::GetCursorPos(&newPosition))
+		{
+			return;
+		}
+		HWND newWindow = ::WindowFromPoint(newPosition);
+		FWindowsApplication::ForceUpdateAsyncCursor(newPosition, newWindow);
+	}
 }
 
 void FWindowsCursor::SetType( const EMouseCursor::Type InNewCursor )
diff --git a/Engine/Source/Runtime/ApplicationCore/Public/Windows/WindowsApplication.h b/Engine/Source/Runtime/ApplicationCore/Public/Windows/WindowsApplication.h
--- a/Engine/Source/Runtime/ApplicationCore/Public/Windows/WindowsApplication.h
+++ b/Engine/Source/Runtime/ApplicationCore/Public/Windows/WindowsApplication.h
@@ -381,6 +381,7 @@ public:
 	APPLICATIONCORE_API virtual bool IsCursorDirectlyOverSlateWindow() const override;
 	static bool StaticGetCursorPosition(POINT& OutSursorPos);
 	bool GetCursorPosition(POINT& OutSursorPos);
+	static void ForceUpdateAsyncCursor(POINT position, HWND windowUnder);
 	APPLICATIONCORE_API virtual FPlatformRect GetWorkArea( const FPlatformRect& CurrentWindow ) const override;
 	APPLICATIONCORE_API virtual void GetInitialDisplayMetrics( FDisplayMetrics& OutDisplayMetrics ) const override;
 	APPLICATIONCORE_API virtual EWindowTitleAlignment::Type GetWindowTitleAlignment() const override;
