UE-306958 Set the async cursor pos cache directly when setting the position from the api

#tests Ran the automated test that was failing. Verified through manual breakpoints that we're hitting what we expect
#rnx

diff --git a/Engine/Source/Runtime/ApplicationCore/Private/Windows/WindowsApplication.cpp b/Engine/Source/Runtime/ApplicationCore/Private/Windows/WindowsApplication.cpp
--- a/Engine/Source/Runtime/ApplicationCore/Private/Windows/WindowsApplication.cpp
+++ b/Engine/Source/Runtime/ApplicationCore/Private/Windows/WindowsApplication.cpp
@@ -1614,6 +1614,20 @@ void FWindowsApplication::ForceUpdateAsyncCursor(POINT position, HWND windowUnde
 	}
 }
 
+
+{
+	if (WindowsApplication == nullptr)
+	{
+		return;
+	}
+	FRWScopeLock Lock(WindowsApplication->AsyncCursorRWLock, SLT_Write);
+	if (WindowsApplication->bAsyncCursorPosValid)
+	{
+		WindowsApplication->AsyncCursorWindowHandle = windowUnder;
+		WindowsApplication->AsyncCursorPos = position;
+	}
+}
+
 void FWindowsApplication::SetCapture( const TSharedPtr< FGenericWindow >& InWindow )
 {
 	HWND DesiredCaptureHandle = (InWindow.IsValid() ? (HWND)InWindow->GetOSWindowHandle() : NULL);
