Fix bad merge from robomerge

#tests no tests
#rnx

diff --git a/Engine/Source/Runtime/ApplicationCore/Private/Windows/WindowsApplication.cpp b/Engine/Source/Runtime/ApplicationCore/Private/Windows/WindowsApplication.cpp
--- a/Engine/Source/Runtime/ApplicationCore/Private/Windows/WindowsApplication.cpp
+++ b/Engine/Source/Runtime/ApplicationCore/Private/Windows/WindowsApplication.cpp
@@ -1614,20 +1614,6 @@ void FWindowsApplication::ForceUpdateAsyncCursor(POINT position, HWND windowUnde
 	}
 }
 
-
-{
-	if (WindowsApplication == nullptr)
-	{
-		return;
-	}
-	FRWScopeLock Lock(WindowsApplication->AsyncCursorRWLock, SLT_Write);
-	if (WindowsApplication->bAsyncCursorPosValid)
-	{
-		WindowsApplication->AsyncCursorWindowHandle = windowUnder;
-		WindowsApplication->AsyncCursorPos = position;
-	}
-}
-
 void FWindowsApplication::SetCapture( const TSharedPtr< FGenericWindow >& InWindow )
 {
 	HWND DesiredCaptureHandle = (InWindow.IsValid() ? (HWND)InWindow->GetOSWindowHandle() : NULL);
