Ensure cached async cursor info is updated within FWindowsCursor::SetPosition unconditionally

Also added comment about why we query cursor info in this situation outside the RWLock.

diff --git a/Engine/Source/Runtime/ApplicationCore/Private/Windows/WindowsApplication.cpp b/Engine/Source/Runtime/ApplicationCore/Private/Windows/WindowsApplication.cpp
--- a/Engine/Source/Runtime/ApplicationCore/Private/Windows/WindowsApplication.cpp
+++ b/Engine/Source/Runtime/ApplicationCore/Private/Windows/WindowsApplication.cpp
@@ -1600,20 +1600,32 @@ bool FWindowsApplication::GetCursorPosition(POINT& OutCursorPos)
 	return bAsyncCursorPosValid;
 }
 
-void FWindowsApplication::ForceUpdateAsyncCursor(POINT position, HWND windowUnder)
+void FWindowsApplication::StaticForceUpdateAsyncCursor()
 {
-	if (WindowsApplication == nullptr)
-	{
-		return;
-	}
-	FRWScopeLock Lock(WindowsApplication->AsyncCursorRWLock, SLT_Write);
-	if (WindowsApplication->bAsyncCursorPosValid)
+	if (WindowsApplication)
 	{
-		WindowsApplication->AsyncCursorWindowHandle = windowUnder;
-		WindowsApplication->AsyncCursorPos = position;
+		WindowsApplication->ForceUpdateAsyncCursor();
 	}
 }
 
+void FWindowsApplication::ForceUpdateAsyncCursor()
+{
+	// Querying outside of the lock can cause reordering:
+	// - Query for cursor info here
+	// - Async task queries for cursor info and updates members under lock
+	// - Members are updated here under lock with stale info
+	// The reason we don't query under lock here is that Win32 calls from the game thread can enter our WndProc and access AsyncCursorRWLock recursively and hang.
+	// We could add lock tracking to avoid grabbing the recursive lock if allowing a little reordering turns out to be fatal.
+	POINT CursorPos;
+	const BOOL bCursorPosValid = ::GetCursorPos(&CursorPos);
+	const HWND CursorWindowHandle = bCursorPosValid ? ::WindowFromPoint(CursorPos) : NULL;
+
+	FRWScopeLock Lock(AsyncCursorRWLock, SLT_Write);
+	bAsyncCursorPosValid = bCursorPosValid;
+	AsyncCursorPos = CursorPos;
+	AsyncCursorWindowHandle = CursorWindowHandle;
+}
+
 void FWindowsApplication::SetCapture( const TSharedPtr< FGenericWindow >& InWindow )
 {
 	HWND DesiredCaptureHandle = (InWindow.IsValid() ? (HWND)InWindow->GetOSWindowHandle() : NULL);
diff --git a/Engine/Source/Runtime/ApplicationCore/Private/Windows/WindowsCursor.cpp b/Engine/Source/Runtime/ApplicationCore/Private/Windows/WindowsCursor.cpp
--- a/Engine/Source/Runtime/ApplicationCore/Private/Windows/WindowsCursor.cpp
+++ b/Engine/Source/Runtime/ApplicationCore/Private/Windows/WindowsCursor.cpp
@@ -239,17 +239,11 @@ FVector2D FWindowsCursor::GetPosition() const
 	return FVector2D(OutCursorPos.x, OutCursorPos.y);
 }
 
-void FWindowsCursor::SetPosition( const int32 X, const int32 Y )
+void FWindowsCursor::SetPosition(const int32 X, const int32 Y)
 {
 	if (::SetCursorPos(X, Y))
 	{
-		POINT newPosition;
-		if (!::GetCursorPos(&newPosition))
-		{
-			return;
-		}
-		HWND newWindow = ::WindowFromPoint(newPosition);
-		FWindowsApplication::ForceUpdateAsyncCursor(newPosition, newWindow);
+		FWindowsApplication::StaticForceUpdateAsyncCursor();
 	}
 }
 
diff --git a/Engine/Source/Runtime/ApplicationCore/Public/Windows/WindowsApplication.h b/Engine/Source/Runtime/ApplicationCore/Public/Windows/WindowsApplication.h
--- a/Engine/Source/Runtime/ApplicationCore/Public/Windows/WindowsApplication.h
+++ b/Engine/Source/Runtime/ApplicationCore/Public/Windows/WindowsApplication.h
@@ -381,7 +381,8 @@ public:
 	APPLICATIONCORE_API virtual bool IsCursorDirectlyOverSlateWindow() const override;
 	static bool StaticGetCursorPosition(POINT& OutSursorPos);
 	bool GetCursorPosition(POINT& OutSursorPos);
-	static void ForceUpdateAsyncCursor(POINT position, HWND windowUnder);
+	static void StaticForceUpdateAsyncCursor();
+	void ForceUpdateAsyncCursor();
 	APPLICATIONCORE_API virtual FPlatformRect GetWorkArea( const FPlatformRect& CurrentWindow ) const override;
 	APPLICATIONCORE_API virtual void GetInitialDisplayMetrics( FDisplayMetrics& OutDisplayMetrics ) const override;
 	APPLICATIONCORE_API virtual EWindowTitleAlignment::Type GetWindowTitleAlignment() const override;
