Fix race condiition with iOS rotation delay seen on iPad and low-end iOS devices

diff --git a/Engine/Source/Runtime/Apple/MetalRHI/Private/MetalViewport.cpp b/Engine/Source/Runtime/Apple/MetalRHI/Private/MetalViewport.cpp
--- a/Engine/Source/Runtime/Apple/MetalRHI/Private/MetalViewport.cpp
+++ b/Engine/Source/Runtime/Apple/MetalRHI/Private/MetalViewport.cpp
@@ -294,6 +294,7 @@ void FMetalViewport::Resize(uint32 InSizeX, uint32 InSizeY, bool bInIsFullscreen
 		SizeY = InSizeY;
 		Format = InFormat;
 		OrientationMask = CachedOrientationMask;
+		FPlatformMisc::LowLevelOutputDebugStringf(TEXT("[rotation] FMetalViewport::Resize: OrientationMask %d"), OrientationMask);
 #endif
 	}
 #endif
diff --git a/Engine/Source/Runtime/ApplicationCore/Private/IOS/IOSAppDelegate.cpp b/Engine/Source/Runtime/ApplicationCore/Private/IOS/IOSAppDelegate.cpp
--- a/Engine/Source/Runtime/ApplicationCore/Private/IOS/IOSAppDelegate.cpp
+++ b/Engine/Source/Runtime/ApplicationCore/Private/IOS/IOSAppDelegate.cpp
@@ -105,6 +105,21 @@ FAutoConsoleVariableRef CVarDelayRotationUntilPresent(
 		GIOSDelayRotationUntilPresent,
 		TEXT("Whether to delay device rotation handling until we have a frame available at the new orientation (default=1).\n"),
 		ECVF_Default);
+
+#if !UE_BUILD_SHIPPING
+static FAutoConsoleCommand CForceRotation(
+	TEXT("ios.ForceRotation"),
+	TEXT("Force a rotation with the integer representing a EDeviceScreenOrientation value passed as a parameter."),
+	FConsoleCommandWithArgsDelegate::CreateLambda([](const TArray<FString>& Args)
+	{
+		if (Args.Num() > 0)
+		{
+			int32 RequestRotation = (int32)FCString::Atod(*Args[0]);
+			FPlatformMisc::SetAllowedDeviceOrientation((EDeviceScreenOrientation)RequestRotation);
+		}
+	}),
+	ECVF_Default);
+#endif
 #endif
 
 
@@ -1454,114 +1469,106 @@ static FAutoConsoleVariableRef CVarGEnableThermalsReport(
 	{
 		return;
 	}
-		
-	FPlatformMisc::LowLevelOutputDebugStringf(TEXT("[rotation] didRotate windowScene=%d device=%d bPlatformInit=%d\n"), (int)[self.Window.windowScene interfaceOrientation], (int)[UIDevice currentDevice].orientation, bPlatformInit);
+
+	if(!self.bPlatformInit)
+	{
+		FPlatformMisc::LowLevelOutputDebugStringf(TEXT("[rotation] didRotate ignoring call before bPlatformInit self.InterfaceOrientation=%d windowScene=%d device=%d\n"), self.InterfaceOrientation, (int)[self.Window.windowScene interfaceOrientation], (int)[UIDevice currentDevice].orientation);
+		return;
+	}
+
+	FPlatformMisc::LowLevelOutputDebugStringf(TEXT("[rotation] didRotate self.InterfaceOrientation=%d windowScene=%d device=%d\n"), self.InterfaceOrientation, (int)[self.Window.windowScene interfaceOrientation], (int)[UIDevice currentDevice].orientation);
 	
-    if (self.bPlatformInit)
-    {
-		UIInterfaceOrientationMask PotentialMask;
-		UIInterfaceOrientation NewOrientation = self.InterfaceOrientation;
-		if (GIOSDelayRotationUntilPresent)
-		{
-			UIInterfaceOrientationMask AllowedMask = [self.IOSController supportedInterfaceOrientations_Internal];
+	UIInterfaceOrientationMask PotentialMask;
+	UIInterfaceOrientation NewOrientation = self.InterfaceOrientation;
+	bool bForceUpdate = false;
+	if (GIOSDelayRotationUntilPresent)
+	{
+		// Force a refresh if the current and cached orientations are not the same. This can happen at startup.
+		bForceUpdate = [self.Window.windowScene interfaceOrientation] != self.InterfaceOrientation;
+
+		UIInterfaceOrientationMask AllowedMask = [self.IOSController supportedInterfaceOrientations_Internal];
 
-			// See if we support the current device orientation
-			UIInterfaceOrientation PotentialOrientation = [IOSAppDelegate GetUIInterfaceOrientationFromUIDeviceOrientation: [UIDevice currentDevice].orientation];	
+		// See if we support the current device orientation
+		UIInterfaceOrientation PotentialOrientation = [IOSAppDelegate GetUIInterfaceOrientationFromUIDeviceOrientation: [UIDevice currentDevice].orientation];	
+		PotentialMask = [IOSAppDelegate GetMaskFromUIInterfaceOrientation: PotentialOrientation];
+		if (PotentialOrientation != UIInterfaceOrientationUnknown && (PotentialMask & AllowedMask) != 0)
+		{
+			FPlatformMisc::LowLevelOutputDebugStringf(TEXT("[rotation] didRotate: using device orientation\n"));
+			NewOrientation = PotentialOrientation;
+		}
+		else
+		{
+			// If not, check if we support the current windowScene interfaceOrientation. It could have been programmatically set.
+			PotentialOrientation = [self.Window.windowScene interfaceOrientation];
 			PotentialMask = [IOSAppDelegate GetMaskFromUIInterfaceOrientation: PotentialOrientation];
-			if (PotentialOrientation != UIInterfaceOrientationUnknown && (PotentialMask & AllowedMask) != 0)
+			if ((PotentialMask & AllowedMask) != 0)
 			{
-				FPlatformMisc::LowLevelOutputDebugStringf(TEXT("[rotation] didRotate: using device orientation\n"));
 				NewOrientation = PotentialOrientation;
+				FPlatformMisc::LowLevelOutputDebugStringf(TEXT("[rotation] didRotate: using windowScene interfaceOrientation\n"));
 			}
 			else
-			{
-				// If not, check if we support the current windowScene interfaceOrientation. It could have been programmatically set.
-				PotentialOrientation = [self.Window.windowScene interfaceOrientation];
-				PotentialMask = [IOSAppDelegate GetMaskFromUIInterfaceOrientation: PotentialOrientation];
-				if ((PotentialMask & AllowedMask) != 0)
+			{			
+				// Force rotation to the first allowed orientation
+				if ((AllowedMask & UIInterfaceOrientationMaskLandscape) == UIInterfaceOrientationMaskLandscape)
 				{
-					NewOrientation = PotentialOrientation;
-					FPlatformMisc::LowLevelOutputDebugStringf(TEXT("[rotation] didRotate: using windowScene interfaceOrientation\n"));
-				}
-				else
-				{			
-					// Force rotation to the first allowed orientation
-					if ((AllowedMask & UIInterfaceOrientationMaskLandscape) == UIInterfaceOrientationMaskLandscape)
-					{
-						FString PreferredLandscapeOrientation;
-						GConfig->GetString(TEXT("/Script/IOSRuntimeSettings.IOSRuntimeSettings"), TEXT("PreferredLandscapeOrientation"), PreferredLandscapeOrientation, GEngineIni);
-						
-						// Preferred default is LandscapeRight as per 
-						if (PreferredLandscapeOrientation.Equals("LandscapeLeft"))
-						{
-							NewOrientation =  UIInterfaceOrientationLandscapeLeft;
-							PotentialMask = UIInterfaceOrientationMaskLandscapeLeft;
-						}
-						else
-						{
-							NewOrientation =  UIInterfaceOrientationLandscapeRight;
-							PotentialMask = UIInterfaceOrientationMaskLandscapeRight;
-						}
-					}
-					else if (AllowedMask & UIInterfaceOrientationMaskLandscapeRight)
-					{
-						NewOrientation =  UIInterfaceOrientationLandscapeRight;
-						PotentialMask = UIInterfaceOrientationMaskLandscapeRight;
-					}
-					else if (AllowedMask & UIInterfaceOrientationMaskLandscapeLeft)
+					FString PreferredLandscapeOrientation;
+					GConfig->GetString(TEXT("/Script/IOSRuntimeSettings.IOSRuntimeSettings"), TEXT("PreferredLandscapeOrientation"), PreferredLandscapeOrientation, GEngineIni);
+					
+					// Preferred default is LandscapeRight as per 
+					if (PreferredLandscapeOrientation.Equals("LandscapeLeft"))
 					{
 						NewOrientation =  UIInterfaceOrientationLandscapeLeft;
 						PotentialMask = UIInterfaceOrientationMaskLandscapeLeft;
 					}
-					else if (AllowedMask & UIInterfaceOrientationMaskPortrait)
-					{
-						NewOrientation = UIInterfaceOrientationPortrait;
-						PotentialMask = UIInterfaceOrientationMaskPortrait;
-					}
-					else if (AllowedMask & UIInterfaceOrientationMaskPortraitUpsideDown)
+					else
 					{
-						NewOrientation = UIInterfaceOrientationPortraitUpsideDown;
-						PotentialMask = UIInterfaceOrientationMaskPortraitUpsideDown;
+						NewOrientation =  UIInterfaceOrientationLandscapeRight;
+						PotentialMask = UIInterfaceOrientationMaskLandscapeRight;
 					}
-					FPlatformMisc::LowLevelOutputDebugStringf(TEXT("[rotation] didRotate: using first allowed orientation\n"));
 				}
+				else if (AllowedMask & UIInterfaceOrientationMaskLandscapeRight)
+				{
+					NewOrientation =  UIInterfaceOrientationLandscapeRight;
+					PotentialMask = UIInterfaceOrientationMaskLandscapeRight;
+				}
+				else if (AllowedMask & UIInterfaceOrientationMaskLandscapeLeft)
+				{
+					NewOrientation =  UIInterfaceOrientationLandscapeLeft;
+					PotentialMask = UIInterfaceOrientationMaskLandscapeLeft;
+				}
+				else if (AllowedMask & UIInterfaceOrientationMaskPortrait)
+				{
+					NewOrientation = UIInterfaceOrientationPortrait;
+					PotentialMask = UIInterfaceOrientationMaskPortrait;
+				}
+				else if (AllowedMask & UIInterfaceOrientationMaskPortraitUpsideDown)
+				{
+					NewOrientation = UIInterfaceOrientationPortraitUpsideDown;
+					PotentialMask = UIInterfaceOrientationMaskPortraitUpsideDown;
+				}
+				FPlatformMisc::LowLevelOutputDebugStringf(TEXT("[rotation] didRotate: using first allowed orientation %d mask %d (AllowedMask was %d)\n"), NewOrientation, PotentialMask, AllowedMask);
 			}
 		}
-		else
-		{
-			NewOrientation = [self.Window.windowScene interfaceOrientation];
-		}
-
-		if (self.InterfaceOrientation != NewOrientation)
-		{
-			FPlatformMisc::LowLevelOutputDebugStringf(TEXT("[rotation] didRotate: rotating to new orientation %d\n"), (int)NewOrientation);
-
-			self.InterfaceOrientation = NewOrientation;
-			FIOSApplication* Application = self.IOSApplication;
-			FIOSView* View = self.IOSView;
-			IOSViewController* Controller = self.IOSController;
-			
-			FFunctionGraphTask::CreateAndDispatchWhenReady([NewOrientation, Application, View, Controller, PotentialMask]()
-			{
-				Application->OrientationChanged(NewOrientation);
-			}, TStatId(), NULL, ENamedThreads::GameThread);
-		}
 	}
 	else
 	{
-		UIInterfaceOrientation NewOrientation = [self.Window.windowScene interfaceOrientation];
-		self.InterfaceOrientation = NewOrientation;
+		NewOrientation = [self.Window.windowScene interfaceOrientation];
+	}
 
-		[FAppleAsyncTask CreateTaskWithBlock : ^ bool(void)
+	if (self.InterfaceOrientation != NewOrientation || bForceUpdate)
+	{
+		static int32 Count=0;
+		int LocalCount = ++Count;
+		FPlatformMisc::LowLevelOutputDebugStringf(TEXT("[rotation] didRotate: rotating to new orientation %d, count %d\n"), (int32)NewOrientation, LocalCount);
+
+		self.InterfaceOrientation = NewOrientation;
+		FIOSApplication* Application = self.IOSApplication;
+		
+		FFunctionGraphTask::CreateAndDispatchWhenReady([NewOrientation, Application, LocalCount]()
 		{
-			FIOSApplication* Application = [IOSAppDelegate GetDelegate].IOSApplication;
+			FPlatformMisc::LowLevelOutputDebugStringf(TEXT("[rotation] didRotate: calling OrientationChanged(%d) from GT, count %d"), (int32)NewOrientation, LocalCount);
 			Application->OrientationChanged(NewOrientation);
-			FCoreDelegates::ApplicationReceivedScreenOrientationChangedNotificationDelegate.Broadcast((int32)[IOSAppDelegate ConvertFromUIInterfaceOrientation:NewOrientation]);
-
-			// we also want to fire off the safe frame event
-			FCoreDelegates::OnSafeFrameChangedEvent.Broadcast();
-			return true;
-		}];
+		}, TStatId(), NULL, ENamedThreads::GameThread);
 	}
 #endif
 }
diff --git a/Engine/Source/Runtime/ApplicationCore/Private/IOS/IOSView.cpp b/Engine/Source/Runtime/ApplicationCore/Private/IOS/IOSView.cpp
--- a/Engine/Source/Runtime/ApplicationCore/Private/IOS/IOSView.cpp
+++ b/Engine/Source/Runtime/ApplicationCore/Private/IOS/IOSView.cpp
@@ -349,8 +349,11 @@ MTL::Device* GMetalDevice = nullptr;
 	// IOSApplication::OrientationChange calls into Slate which issues rendering commands
 	if(![IOSAppDelegate GetDelegate].bIsSuspended)
 	{
-		UIInterfaceOrientation CurrentOrientation = [self.window.windowScene interfaceOrientation];
-		FIOSApplication::OrientationChanged(CurrentOrientation);
+		if (!GIOSDelayRotationUntilPresent)
+		{
+			UIInterfaceOrientation CurrentOrientation = [self.window.windowScene interfaceOrientation];
+			FIOSApplication::OrientationChanged(CurrentOrientation);
+		}
 	}
 #endif
 }
@@ -1133,7 +1136,6 @@ self.accessibilityElements = @[Window.accessibilityContainer];
  {
 	IOSAppDelegate* AppDelegate = [IOSAppDelegate GetDelegate];
 	const FIOSView* View = [AppDelegate IOSView];
-	extern bool GIOSDelayRotationUntilPresent;
 	if (GIOSDelayRotationUntilPresent && AppDelegate.bPlatformInit && View != nil && (int)View->LastFrameInterfaceOrientationMask != 0)
 	{
 		FPlatformMisc::LowLevelOutputDebugStringf(TEXT("[rotation] supportedInterfaceOrientations: allowing %d\n"), (int)View->LastFrameInterfaceOrientationMask);
diff --git a/Engine/Source/Runtime/ApplicationCore/Private/IOS/IOSWindow.cpp b/Engine/Source/Runtime/ApplicationCore/Private/IOS/IOSWindow.cpp
--- a/Engine/Source/Runtime/ApplicationCore/Private/IOS/IOSWindow.cpp
+++ b/Engine/Source/Runtime/ApplicationCore/Private/IOS/IOSWindow.cpp
@@ -136,7 +136,6 @@ FPlatformRect FIOSWindow::GetScreenRect()
 		ScreenRect.Right = FMath::TruncToInt((Frame.origin.x + Frame.size.width) * Scale);
 
 #if !PLATFORM_TVOS && !PLATFORM_VISIONOS
-		extern bool GIOSDelayRotationUntilPresent;
 		if (GIOSDelayRotationUntilPresent)
 		{	
 			int32 WindowWidth = ScreenRect.Right - ScreenRect.Left;
diff --git a/Engine/Source/Runtime/ApplicationCore/Public/IOS/IOSAppDelegate.h b/Engine/Source/Runtime/ApplicationCore/Public/IOS/IOSAppDelegate.h
--- a/Engine/Source/Runtime/ApplicationCore/Public/IOS/IOSAppDelegate.h
+++ b/Engine/Source/Runtime/ApplicationCore/Public/IOS/IOSAppDelegate.h
@@ -22,6 +22,10 @@
 
 #ifndef SWIFT_IMPORT
 
+#if !PLATFORM_TVOS
+extern bool GIOSDelayRotationUntilPresent;
+#endif
+
 enum class EAudioFeature : uint8
 {
 	Playback, // Audio Not affected by the ringer switch
diff --git a/Engine/Source/Runtime/Launch/Private/IOS/LaunchIOS.cpp b/Engine/Source/Runtime/Launch/Private/IOS/LaunchIOS.cpp
--- a/Engine/Source/Runtime/Launch/Private/IOS/LaunchIOS.cpp
+++ b/Engine/Source/Runtime/Launch/Private/IOS/LaunchIOS.cpp
@@ -425,7 +425,14 @@ void FAppEntry::PlatformInit()
 	FDisplayMetrics::RebuildDisplayMetrics(DisplayMetrics);
 	FSystemResolution::RequestResolutionChange(DisplayMetrics.PrimaryDisplayWidth, DisplayMetrics.PrimaryDisplayHeight, EWindowMode::Fullscreen);
 	IConsoleManager::Get().CallAllConsoleVariableSinks();
-	AppDelegate.bPlatformInit = true;
+
+	dispatch_async(dispatch_get_main_queue(), ^{
+		AppDelegate.bPlatformInit = true;
+#if !PLATFORM_TVOS && !PLATFORM_VISIONOS
+		// Update rotation state now the game thread and taskgraph are ready
+		[AppDelegate didRotate:nil];
+#endif
+	});
 }
 
 extern TcpConsoleListener *ConsoleListener;
