Fix issue with safe zones with new rotation code
#rnx

diff --git a/Engine/Source/Runtime/ApplicationCore/Private/IOS/IOSApplication.cpp b/Engine/Source/Runtime/ApplicationCore/Private/IOS/IOSApplication.cpp
--- a/Engine/Source/Runtime/ApplicationCore/Private/IOS/IOSApplication.cpp
+++ b/Engine/Source/Runtime/ApplicationCore/Private/IOS/IOSApplication.cpp
@@ -214,8 +214,59 @@ void FDisplayMetrics::RebuildDisplayMetrics(FDisplayMetrics& OutDisplayMetrics)
 void FIOSApplication::CacheDisplayMetrics(UIInterfaceOrientation Orientation)
 {
 	CachedOrientation = Orientation;
-	//todo: rotate insets to match given orientation
-	CachedInsets = [[[[UIApplication sharedApplication] delegate] window] safeAreaInsets];
+	
+	// Insets are given in whatever actual orientation the window has
+	UIWindow* Window = [[[UIApplication sharedApplication] delegate] window];
+	UIEdgeInsets OrientatedInsets = [Window safeAreaInsets];
+
+	if (GIOSDelayRotationUntilPresent)
+	{
+		// After queueing a new orientation, the actual rotation hasn't occurred yet. 
+		// We can attempt to adjust our current insets to the new orientation, but they are not entirely accurate as iOS adds an inset for the home swipe bar to the current orientation bottom.
+		UIEdgeInsets PortraitInsets;
+
+		// Convert current orientation to portrait
+		switch ([Window.windowScene interfaceOrientation])
+		{
+		case UIInterfaceOrientationPortraitUpsideDown:
+			PortraitInsets = UIEdgeInsetsMake(OrientatedInsets.bottom, OrientatedInsets.right, OrientatedInsets.top, OrientatedInsets.left);
+			break;
+		case UIInterfaceOrientationLandscapeLeft:
+			PortraitInsets = UIEdgeInsetsMake(OrientatedInsets.right, OrientatedInsets.top, OrientatedInsets.left, OrientatedInsets.bottom);
+			break;
+		case UIInterfaceOrientationLandscapeRight:
+			PortraitInsets = UIEdgeInsetsMake(OrientatedInsets.left, OrientatedInsets.bottom, OrientatedInsets.right, OrientatedInsets.top);
+			break;
+		case UIInterfaceOrientationPortrait:
+		default:
+			PortraitInsets = OrientatedInsets;
+			break;	
+		}
+
+		// Convert portrait to new orientation and cache
+		switch (Orientation)
+		{
+		case UIInterfaceOrientationPortraitUpsideDown:
+			CachedInsets = UIEdgeInsetsMake(PortraitInsets.bottom, PortraitInsets.right, PortraitInsets.top, PortraitInsets.left);
+			break;
+		case UIInterfaceOrientationLandscapeLeft:
+			CachedInsets = UIEdgeInsetsMake(PortraitInsets.left, PortraitInsets.bottom, PortraitInsets.right, PortraitInsets.top);
+			break;
+		case UIInterfaceOrientationLandscapeRight:
+			CachedInsets = UIEdgeInsetsMake(PortraitInsets.right, PortraitInsets.top, PortraitInsets.left, PortraitInsets.bottom);
+			break;
+		case UIInterfaceOrientationPortrait:
+		default:
+			CachedInsets = PortraitInsets;
+			break;	
+		}
+	}
+	else
+	{
+		CachedInsets = OrientatedInsets;
+	}
+
+	FPlatformMisc::LowLevelOutputDebugStringf(TEXT("[rotation] CachedInsets top:%f left:%f bottom:%f right:%f\n"), CachedInsets.top, CachedInsets.left, CachedInsets.bottom, CachedInsets.right);
 }
 #endif
 
@@ -258,6 +309,28 @@ void FIOSApplication::OrientationChanged(UIInterfaceOrientation NewOrientation)
 		return true;
 	 }];
 }
+
+void FIOSApplication::UpdateSafeZoneAfterRotation()
+{
+	// this function is called from the UI thread
+	UIWindow* Window = [[[UIApplication sharedApplication] delegate] window];
+	CachedInsets = [Window safeAreaInsets];
+
+	// notify the game thread of the safe zone change
+	[FIOSAsyncTask CreateTaskWithBlock : ^ bool(void)
+	{
+	 	IOSAppDelegate* AppDelegate = [IOSAppDelegate GetDelegate];
+		FIOSApplication* App = AppDelegate.IOSApplication;
+
+		FDisplayMetrics DisplayMetrics;
+		FDisplayMetrics::RebuildDisplayMetrics(DisplayMetrics);
+		App->BroadcastDisplayMetricsChanged(DisplayMetrics);
+		FCoreDelegates::OnSafeFrameChangedEvent.Broadcast();
+		FPlatformMisc::LowLevelOutputDebugStringf(TEXT("[rotation] UpdateSafeZoneAfterRotation top:%f left:%f bottom:%f right:%f\n"), CachedInsets.top, CachedInsets.left, CachedInsets.bottom, CachedInsets.right);
+		return true;
+	}];
+}
+
 #endif
 
 
diff --git a/Engine/Source/Runtime/ApplicationCore/Private/IOS/IOSView.cpp b/Engine/Source/Runtime/ApplicationCore/Private/IOS/IOSView.cpp
--- a/Engine/Source/Runtime/ApplicationCore/Private/IOS/IOSView.cpp
+++ b/Engine/Source/Runtime/ApplicationCore/Private/IOS/IOSView.cpp
@@ -1236,7 +1236,18 @@ self.accessibilityElements = @[Window.accessibilityContainer];
 	
 	FPlatformMisc::LowLevelOutputDebugStringf(TEXT("[rotation] viewWillTransitionToSize\n"));
 
-	[super viewWillTransitionToSize : size withTransitionCoordinator : coordinator];	
+	[super viewWillTransitionToSize : size withTransitionCoordinator : coordinator];
+
+#if !PLATFORM_TVOS && !PLATFORM_VISIONOS
+	if (GIOSDelayRotationUntilPresent)
+	{
+		[coordinator animateAlongsideTransition:nil completion:^(id<UIViewControllerTransitionCoordinatorContext>  _Nonnull context)
+		{
+			// rotation has completed, we need to update the actual safe zone
+			FIOSApplication::UpdateSafeZoneAfterRotation();
+		}];
+	}
+#endif
 }
 
 #if !PLATFORM_TVOS && !PLATFORM_VISIONOS
diff --git a/Engine/Source/Runtime/ApplicationCore/Public/IOS/IOSApplication.h b/Engine/Source/Runtime/ApplicationCore/Public/IOS/IOSApplication.h
--- a/Engine/Source/Runtime/ApplicationCore/Public/IOS/IOSApplication.h
+++ b/Engine/Source/Runtime/ApplicationCore/Public/IOS/IOSApplication.h
@@ -34,6 +34,7 @@ public:
 #if !PLATFORM_TVOS
 #if !PLATFORM_VISIONOS
 	static void OrientationChanged(UIInterfaceOrientation NewOrientation);
+	static void UpdateSafeZoneAfterRotation();
 #endif
 	static UIInterfaceOrientation CachedOrientation;
 #endif
