Fix iOS compile warning due to deprecated class

#rnx

diff --git a/Engine/Source/Runtime/ApplicationCore/Private/IOS/IOSApplication.cpp b/Engine/Source/Runtime/ApplicationCore/Private/IOS/IOSApplication.cpp
--- a/Engine/Source/Runtime/ApplicationCore/Private/IOS/IOSApplication.cpp
+++ b/Engine/Source/Runtime/ApplicationCore/Private/IOS/IOSApplication.cpp
@@ -317,7 +317,7 @@ void FIOSApplication::UpdateSafeZoneAfterRotation()
 	CachedInsets = [Window safeAreaInsets];
 
 	// notify the game thread of the safe zone change
-	[FIOSAsyncTask CreateTaskWithBlock : ^ bool(void)
+	[FAppleAsyncTask CreateTaskWithBlock : ^ bool(void)
 	{
 	 	IOSAppDelegate* AppDelegate = [IOSAppDelegate GetDelegate];
 		FIOSApplication* App = AppDelegate.IOSApplication;
