Allowing BGDL to use expensive and constrained networks when allowed by the user.

Also configuring `multipathServiceType` to `NSURLSessionMultipathServiceTypeAggregate` to improve throughput as explained by Apple's docs:
"A service that aggregates the capacities of other Multipath options in an attempt to increase throughput and minimize latency."

#rnx

diff --git a/Engine/Source/Runtime/ApplicationCore/Private/IOS/IOSBackgroundURLSessionHandler.cpp b/Engine/Source/Runtime/ApplicationCore/Private/IOS/IOSBackgroundURLSessionHandler.cpp
--- a/Engine/Source/Runtime/ApplicationCore/Private/IOS/IOSBackgroundURLSessionHandler.cpp
+++ b/Engine/Source/Runtime/ApplicationCore/Private/IOS/IOSBackgroundURLSessionHandler.cpp
@@ -496,11 +496,12 @@ static constexpr NSInteger HTTPStatusCodeErrorServer = 500;
 
 	// Set session to allow cellular and instead control this on NSMutableURLRequest level because this value cannot be changed after session is created.
 	Configuration.allowsCellularAccess = YES;
+	Configuration.allowsExpensiveNetworkAccess = YES;
+	Configuration.allowsConstrainedNetworkAccess = YES;
 
 	Configuration.networkServiceType = bUseForegroundSession ? NSURLNetworkServiceTypeDefault : NSURLNetworkServiceTypeBackground;
 
-	// TODO Is this any use for us? Needs entitlement.
-	//Configuration.multipathServiceType = NSURLSessionMultipathServiceTypeAggregate;
+	Configuration.multipathServiceType = NSURLSessionMultipathServiceTypeAggregate;
 
 	Configuration.HTTPMaximumConnectionsPerHost = _MaximumConnectionsPerHost;
 
@@ -829,6 +830,8 @@ static constexpr NSInteger HTTPStatusCodeErrorServer = 500;
 {
 	NSMutableURLRequest* URLRequest = [NSMutableURLRequest requestWithURL:URL];
 	URLRequest.allowsCellularAccess = self.AllowCellular;
+	URLRequest.allowsConstrainedNetworkAccess = self.AllowCellular;
+	URLRequest.allowsExpensiveNetworkAccess = self.AllowCellular;
 
 	NSURLSessionDownloadTask* Task = [_Session downloadTaskWithRequest:URLRequest];
 	[Task setPriority:Priority];
@@ -1051,7 +1054,8 @@ static constexpr NSInteger HTTPStatusCodeErrorServer = 500;
             }
         }
 #if !UE_BUILD_SHIPPING
-		else {
+		else 
+		{
             UE_DNLD_LOG(@"Response Duration: Not Available");
         }
 #endif
