Fix other Apple platforms compilation due to CL 44863788

#rnx

diff --git a/Engine/Source/Runtime/ApplicationCore/Private/IOS/IOSBackgroundURLSessionHandler.cpp b/Engine/Source/Runtime/ApplicationCore/Private/IOS/IOSBackgroundURLSessionHandler.cpp
--- a/Engine/Source/Runtime/ApplicationCore/Private/IOS/IOSBackgroundURLSessionHandler.cpp
+++ b/Engine/Source/Runtime/ApplicationCore/Private/IOS/IOSBackgroundURLSessionHandler.cpp
@@ -501,7 +501,11 @@ static constexpr NSInteger HTTPStatusCodeErrorServer = 500;
 
 	Configuration.networkServiceType = bUseForegroundSession ? NSURLNetworkServiceTypeDefault : NSURLNetworkServiceTypeBackground;
 
+	// this api is not available on any other Apple platform other than iOS
+	// as this file is compiled for tvOS as well, we need the #if PLATFORM_IOS
+#if PLATFORM_IOS
 	Configuration.multipathServiceType = NSURLSessionMultipathServiceTypeAggregate;
+#endif // PLATFORM_IOS
 
 	Configuration.HTTPMaximumConnectionsPerHost = _MaximumConnectionsPerHost;
 
