Fix tvOS compilation

#rnx

diff --git a/Engine/Source/Runtime/ApplicationCore/Private/IOS/IOSBackgroundURLSessionHandler.cpp b/Engine/Source/Runtime/ApplicationCore/Private/IOS/IOSBackgroundURLSessionHandler.cpp
--- a/Engine/Source/Runtime/ApplicationCore/Private/IOS/IOSBackgroundURLSessionHandler.cpp
+++ b/Engine/Source/Runtime/ApplicationCore/Private/IOS/IOSBackgroundURLSessionHandler.cpp
@@ -502,10 +502,10 @@ static constexpr NSInteger HTTPStatusCodeErrorServer = 500;
 	Configuration.networkServiceType = bUseForegroundSession ? NSURLNetworkServiceTypeDefault : NSURLNetworkServiceTypeBackground;
 
 	// this api is not available on any other Apple platform other than iOS
-	// as this file is compiled for tvOS as well, we need the #if PLATFORM_IOS
-#if PLATFORM_IOS
+	// as this file is compiled for tvOS as well, we need the #if !PLATFORM_TVOS
+#if !PLATFORM_TVOS
 	Configuration.multipathServiceType = NSURLSessionMultipathServiceTypeAggregate;
-#endif // PLATFORM_IOS
+#endif // !PLATFORM_TVOS
 
 	Configuration.HTTPMaximumConnectionsPerHost = _MaximumConnectionsPerHost;
 
