Fix popup menus on Linux to correctly parent themselves when created.

#rnx

diff --git a/Engine/Source/Runtime/ApplicationCore/Private/Linux/LinuxWindow.cpp b/Engine/Source/Runtime/ApplicationCore/Private/Linux/LinuxWindow.cpp
--- a/Engine/Source/Runtime/ApplicationCore/Private/Linux/LinuxWindow.cpp
+++ b/Engine/Source/Runtime/ApplicationCore/Private/Linux/LinuxWindow.cpp
@@ -202,6 +202,7 @@ void FLinuxWindow::Initialize( FLinuxApplication* const Application, const TShar
 	{
 		SDL_SetBooleanProperty( Props, SDL_PROP_WINDOW_CREATE_TRANSPARENT_BOOLEAN, true );
 		SDL_SetBooleanProperty( Props, SDL_PROP_WINDOW_CREATE_BORDERLESS_BOOLEAN, true );
+		SDL_SetPointerProperty( Props, SDL_PROP_WINDOW_CREATE_PARENT_POINTER, InParent->GetHWnd() );
 		bIsPopupWindow = true;
 		UE_LOG(LogLinuxWindowType, Verbose, TEXT("*** New Window is a Popup Menu Window ***"));
 	}
