Deferred handling of notifications if engine isn't ready to receive.

diff --git a/Engine/Source/Runtime/ApplicationCore/Private/IOS/IOSAppDelegate.cpp b/Engine/Source/Runtime/ApplicationCore/Private/IOS/IOSAppDelegate.cpp
--- a/Engine/Source/Runtime/ApplicationCore/Private/IOS/IOSAppDelegate.cpp
+++ b/Engine/Source/Runtime/ApplicationCore/Private/IOS/IOSAppDelegate.cpp
@@ -234,6 +234,10 @@ bool FIOSCoreDelegates::PassesPushNotificationFilters(NSDictionary* Payload)
 	});
 }
 
+@interface IOSNotificationHandler : NSObject
++(void)handleWhenReady:(UNNotification*)notification withCompletionHandler:(void (^)()) completionHandler;
+@end
+
 @interface IOSAppDelegate()
 
 // move private things from header to here
@@ -2003,11 +2007,11 @@ extern double GCStartTime;
 	});
 }
 
-void HandleReceivedNotification(UNNotification* notification)
+static bool HandleReceivedNotification(UNNotification* notification)
 {
 	bool IsLocal = false;
-	
-	if ([IOSAppDelegate GetDelegate].bEngineInit)
+	bool IsReadyToHandle = [IOSAppDelegate GetDelegate].bEngineInit;
+	if (IsReadyToHandle)
 	{
 		NSString* NotificationType = (NSString*)[notification.request.content.userInfo objectForKey: @"NotificationType"];
 		if(NotificationType != nullptr)
@@ -2068,6 +2072,7 @@ void HandleReceivedNotification(UNNotification* notification)
 														   }, TStatId(), NULL, ENamedThreads::GameThread);
 		}
 	}
+	return IsReadyToHandle;
 }
 
 - (void)userNotificationCenter:(UNUserNotificationCenter *)center
@@ -2075,9 +2080,7 @@ void HandleReceivedNotification(UNNotification* notification)
 		 withCompletionHandler:(void (^)(UNNotificationPresentationOptions options))completionHandler
 {
 	// Received notification while app is in the foreground
-	HandleReceivedNotification(notification);
-	
-	completionHandler(UNNotificationPresentationOptionNone);
+	[IOSNotificationHandler handleWhenReady:notification withCompletionHandler:^{ completionHandler(UNNotificationPresentationOptionNone); }];
 }
 
 - (void)userNotificationCenter:(UNUserNotificationCenter *)center
@@ -2108,9 +2111,7 @@ didReceiveNotificationResponse:(UNNotificationResponse *)response
 		}
 	}
 	
-	HandleReceivedNotification(response.notification);
-	
-	completionHandler();
+	[IOSNotificationHandler handleWhenReady:response.notification withCompletionHandler:completionHandler];
 }
 
 #endif
@@ -2312,3 +2313,53 @@ CORE_API bool IOSShowDashboardUI()
 }
 
 @end
+
+@implementation IOSNotificationHandler
+{
+	UNNotification* notification;
+	void (^completionHandler)();
+	int retryCount;
+}
+
++(void)handleWhenReady:(UNNotification*)notification withCompletionHandler:(void (^)()) completionHandler
+{
+	[[[[IOSNotificationHandler alloc] initWithNotification:notification andCompletionHandler:completionHandler] autorelease] handle];
+}
+
+-(id)initWithNotification:(UNNotification*)_notification andCompletionHandler:(void (^)()) _completionHandler
+{
+	self = [super init];
+	if(self)
+	{
+		completionHandler = [_completionHandler copy];
+		notification = [_notification retain];
+		retryCount = 0;
+	}
+	return self;
+}
+
+-(void)dealloc
+{
+	[completionHandler release];
+	[notification release];
+	[super dealloc];
+}
+
+-(void)handle
+{
+	if (HandleReceivedNotification(notification))
+	{
+		completionHandler();
+	}
+	// retry for up to 10s
+	else if(++retryCount <= 100)
+	{
+		[self performSelector:@selector(handle) withObject:nil afterDelay:0.1];
+	}
+	else
+	{
+		UE_LOG(LogIOS, Warning, TEXT("Failed to deliver notification after retrying for 10 seconds."));
+
+	}
+}
+@end
