Fix TVOS/VisionOS compilation. Push Notification are only supported on iOS - properly stub out handler for non iOS plarforms.

diff --git a/Engine/Source/Runtime/ApplicationCore/Private/IOS/IOSAppDelegate.cpp b/Engine/Source/Runtime/ApplicationCore/Private/IOS/IOSAppDelegate.cpp
--- a/Engine/Source/Runtime/ApplicationCore/Private/IOS/IOSAppDelegate.cpp
+++ b/Engine/Source/Runtime/ApplicationCore/Private/IOS/IOSAppDelegate.cpp
@@ -2114,6 +2114,8 @@ didReceiveNotificationResponse:(UNNotificationResponse *)response
 	[IOSNotificationHandler handleWhenReady:response.notification withCompletionHandler:completionHandler];
 }
 
+#else
+static bool HandleReceivedNotification(UNNotification* notification) { return true; }
 #endif
 
 /**
