Add expectedSize to iOS download request

#rnx

diff --git a/Engine/Plugins/Experimental/DefaultInstallBundleManager/Source/Private/InstallBundleManagerUtil.cpp b/Engine/Plugins/Experimental/DefaultInstallBundleManager/Source/Private/InstallBundleManagerUtil.cpp
--- a/Engine/Plugins/Experimental/DefaultInstallBundleManager/Source/Private/InstallBundleManagerUtil.cpp
+++ b/Engine/Plugins/Experimental/DefaultInstallBundleManager/Source/Private/InstallBundleManagerUtil.cpp
@@ -752,7 +752,7 @@ namespace InstallBundleManagerAnalytics
 
 		AnalyticsProvider->RecordEvent(TEXT("InstallBundleManager.BundleRequestStartedDetailed"), MakeAnalyticsEventAttributeArray(
 			TEXT("BundleName"), BundleName,
-			TEXT("PreviousManfest"), PreviousManifest,
+			TEXT("PreviousManifest"), PreviousManifest,
 			TEXT("CurrentManifest"), CurrentManifest,
 			TEXT("CurrentVersion"), CurrentVersion));
 	}
diff --git a/Engine/Source/Runtime/ApplicationCore/Private/IOS/IOSBackgroundURLSessionHandler.cpp b/Engine/Source/Runtime/ApplicationCore/Private/IOS/IOSBackgroundURLSessionHandler.cpp
--- a/Engine/Source/Runtime/ApplicationCore/Private/IOS/IOSBackgroundURLSessionHandler.cpp
+++ b/Engine/Source/Runtime/ApplicationCore/Private/IOS/IOSBackgroundURLSessionHandler.cpp
@@ -54,6 +54,7 @@ static void LogIOSBackgroundDownloadMessage(NSString* Message)
 @property (nonatomic, retain) NSMutableArray<__kindof NSURL*>* URLs;
 @property (nonatomic, retain) NSMutableArray<__kindof NSNumber*>* RetryCountPerURL;
 @property (nonatomic) uint64 ExpectedResultSize;
+@property (nonatomic) uint64 ExpectedResultSizeFromHeadRequest;
 
 // assumes all URL's have same content path but different domain
 + (FBackgroundNSURLSessionDownloadTaskData* _Nonnull)TaskDataWithURLs:(NSArray<__kindof NSURL*>*)URLs WithRetryCount:(NSInteger)RetryCount WithExpectedResultSize:(uint64)ExpectedResultSize;
@@ -77,6 +78,7 @@ NSString* const SerializationKeyCDNs = @"c";
 NSString* const SerializationKeyPath = @"p";
 NSString* const SerializationKeyRetryCountPerURL = @"r";
 NSString* const SerializationKeyExpectedResultSize = @"s";
+NSString* const SerializationKeyExpectedResultSizeFromHeaderRequest = @"h";
 
 + (FBackgroundNSURLSessionDownloadTaskData* _Nonnull)TaskDataWithURLs:(NSArray<__kindof NSURL*>*)URLs WithRetryCount:(NSInteger)RetryCount WithExpectedResultSize:(uint64)ExpectedResultSize
 {
@@ -119,10 +121,11 @@ NSString* const SerializationKeyExpectedResultSize = @"s";
 	NSString* Path = [Dict valueForKey:SerializationKeyPath];
 	NSArray<__kindof NSNumber*>* RetryCountPerURL = [Dict valueForKey:SerializationKeyRetryCountPerURL];
 	NSNumber* ExpectedResultSize = [Dict valueForKey:SerializationKeyExpectedResultSize]; // can be nil
+	NSNumber* ExpectedResultSizeFromHeadRequest = [Dict valueForKey:SerializationKeyExpectedResultSizeFromHeaderRequest];
 
 	if (Error != nil || Version == nil || Version.intValue != 1 || CDNs == nil || Path == nil || RetryCountPerURL == nil)
 	{
-		UE_DNLD_LOG(@"Failed to deserialize task state '%@' due to '%@', %u, %u, %u, %u, %u, %u",
+		UE_DNLD_LOG(@"Failed to deserialize task state '%@' due to '%@', %u, %u, %u, %u, %u, %u, %u",
 			SerializedData,
 			Error != nil ? Error.localizedDescription : @"",
 			Version == nil ? 1 : 0,
@@ -130,7 +133,8 @@ NSString* const SerializationKeyExpectedResultSize = @"s";
 			CDNs == nil ? 1 : 0,
 			Path == nil ? 1 : 0,
 			RetryCountPerURL == nil ? 1 : 0,
-			ExpectedResultSize == nil ? 1 : 0
+			ExpectedResultSize == nil ? 1 : 0,
+			ExpectedResultSizeFromHeadRequest == nil ? 1 : 0
 		);
 		[Data release];
 		return nil;
@@ -145,6 +149,7 @@ NSString* const SerializationKeyExpectedResultSize = @"s";
 	NSAssert(Data.URLs.count > 0, @"URLs should be non empty");
 	NSAssert(Data.URLs.count == Data.RetryCountPerURL.count, @"URLs and RetryCountPerURL arrays should have same size");
 	Data.ExpectedResultSize = ExpectedResultSize != nil ? ExpectedResultSize.unsignedLongLongValue : 0;
+	Data.ExpectedResultSizeFromHeadRequest = ExpectedResultSizeFromHeadRequest != nil ? ExpectedResultSizeFromHeadRequest.unsignedLongLongValue : 0;
 
 	return [Data autorelease];
 }
@@ -183,6 +188,10 @@ NSString* const SerializationKeyExpectedResultSize = @"s";
 	{
 		[Dict setValue:[NSNumber numberWithUnsignedLongLong:self.ExpectedResultSize] forKey:SerializationKeyExpectedResultSize];
 	}
+	if (self.ExpectedResultSizeFromHeadRequest != 0)
+	{
+		[Dict setValue:[NSNumber numberWithUnsignedLongLong:self.ExpectedResultSizeFromHeadRequest] forKey:SerializationKeyExpectedResultSizeFromHeaderRequest];
+	}
 
 	NSError* Error = nil;
 	NSData* Data = [NSJSONSerialization dataWithJSONObject:Dict options:NSJSONWritingSortedKeys error:&Error];
@@ -373,6 +382,7 @@ enum class EBackgroundNSURLCDNInfoResponse : uint32 // Beware these constants de
 @implementation FBackgroundNSURLSession
 {
 	NSURLSession* _Session;
+	NSURLSession* _HeadRequestSession;		// Used by CDN reordering and HEAD request for expected size
 	NSMutableDictionary<__kindof NSNumber*, __kindof NSURLSessionDownloadTask*>* _AllDownloads;
 	NSUInteger _NextDownloadId;
 	std::promise<void> _AllDownloadsPromise;
@@ -382,7 +392,7 @@ enum class EBackgroundNSURLCDNInfoResponse : uint32 // Beware these constants de
 	BackgroundHttpFileHashHelperPtr _HelperPtr;
 	int32 _MaximumConnectionsPerHost;
 	int32 _RetryResumeDataLimit;
-	int32 _CDNReorderingTimeout;
+	int32 _HeadRequestTimeout;
 	bool _bCDNReorderByPingTime;
 	double _CheckForForegroundStaleDownloadsWithInterval;
 	double _ForegroundStaleDownloadTimeout;
@@ -401,6 +411,7 @@ NSProgressUserInfoKey const NSProgressDownloadResultTempFilePath = @"com.epicgam
 static constexpr NSInteger HTTPStatusCodeSuccessCreated = 201;
 static constexpr NSInteger HTTPStatusCodeErrorBadRequest = 400;
 static constexpr NSInteger HTTPStatusCodeErrorServer = 500;
+static constexpr NSInteger TypicalSizeForDownloadResponseHeader = 400; // bytes
 
 + (FBackgroundNSURLSession*)Shared
 {
@@ -434,13 +445,13 @@ static constexpr NSInteger HTTPStatusCodeErrorServer = 500;
 - (void)Initialize
 {
 	bool bUseForegroundSession = false;
-	bool bDiscretionary = false;
+	bool bDiscretionary = true;	// Let iOS schedule background downloads, now that we have the expected sizes filled in
 	bool bShouldSendLaunchEvents = true;
 	_MaximumConnectionsPerHost = 6;
 	double TimeoutIntervalForRequest = 120.0; // Note, ignored in background sessions (if bUseForegroundSession is false).
 	double TimeoutIntervalForResource = 60.0 * 60.0;
 	_RetryResumeDataLimit = 3;
-	_CDNReorderingTimeout = 400;
+	_HeadRequestTimeout = 400;
 	_bCDNReorderByPingTime = false;
 	_CheckForForegroundStaleDownloadsWithInterval = 1.0; // how often to check for stale downloads, <=0.0 to disable
 	_ForegroundStaleDownloadTimeout = 30.0; // If download hasn't received any bytes for this duration, cancel and retry if possible
@@ -453,7 +464,7 @@ static constexpr NSInteger HTTPStatusCodeErrorServer = 500;
 	GConfig->GetDouble(TEXT("BackgroundHttp.iOSSettings"), TEXT("BackgroundReceiveTimeout"), TimeoutIntervalForRequest, GEngineIni);
 	GConfig->GetDouble(TEXT("BackgroundHttp.iOSSettings"), TEXT("BackgroundHttpResourceTimeout"), TimeoutIntervalForResource, GEngineIni);
 	GConfig->GetInt(TEXT("BackgroundHttp.iOSSettings"), TEXT("RetryResumeDataLimit"), _RetryResumeDataLimit, GEngineIni);
-	GConfig->GetInt(TEXT("BackgroundHttp.iOSSettings"), TEXT("CDNReorderingTimeout"), _CDNReorderingTimeout, GEngineIni);
+	GConfig->GetInt(TEXT("BackgroundHttp.iOSSettings"), TEXT("CDNReorderingTimeout"), _HeadRequestTimeout, GEngineIni);
 	GConfig->GetBool(TEXT("BackgroundHttp.iOSSettings"), TEXT("bCDNReorderByPingTime"), _bCDNReorderByPingTime, GEngineIni);
 	GConfig->GetDouble(TEXT("BackgroundHttp.iOSSettings"), TEXT("CheckForForegroundStaleDownloadsWithInterval"), _CheckForForegroundStaleDownloadsWithInterval, GEngineIni);
 	GConfig->GetDouble(TEXT("BackgroundHttp.iOSSettings"), TEXT("ForegroundStaleDownloadTimeout"), _ForegroundStaleDownloadTimeout, GEngineIni);
@@ -466,7 +477,7 @@ static constexpr NSInteger HTTPStatusCodeErrorServer = 500;
 	UE_DNLD_LOG(@"TimeoutIntervalForRequest=%f", TimeoutIntervalForRequest);
 	UE_DNLD_LOG(@"TimeoutIntervalForResource=%f", TimeoutIntervalForResource);
 	UE_DNLD_LOG(@"RetryResumeDataLimit=%i", _RetryResumeDataLimit);
-	UE_DNLD_LOG(@"CDNReorderingTimeout=%i", _CDNReorderingTimeout);
+	UE_DNLD_LOG(@"HeadRequestTimeout=%i", _HeadRequestTimeout);
 	UE_DNLD_LOG(@"CDNReorderByPingTime=%u", _bCDNReorderByPingTime ? 1 : 0);
 	UE_DNLD_LOG(@"CheckForForegroundStaleDownloadsWithInterval=%f", _CheckForForegroundStaleDownloadsWithInterval);
 	UE_DNLD_LOG(@"ForegroundStaleDownloadTimeout=%f", _ForegroundStaleDownloadTimeout);
@@ -548,6 +559,17 @@ static constexpr NSInteger HTTPStatusCodeErrorServer = 500;
 
 		_AllDownloadsPromise.set_value();
 	}];
+	
+	Configuration = [NSURLSessionConfiguration ephemeralSessionConfiguration];	// doesn’t store caches, credential stores, or any session-related data to disk
+	Configuration.discretionary = NO;
+	// Don't go over cellular here, during first download attempt we would try to over non-cellular connections first.
+	// Hence it makes more sense to prioritize CDN's that are reachable via non-cellular connections.
+	Configuration.allowsCellularAccess = NO;
+	Configuration.networkServiceType = NSURLNetworkServiceTypeResponsiveData;
+	Configuration.timeoutIntervalForRequest = (double)_HeadRequestTimeout / 1000 ;
+	Configuration.timeoutIntervalForResource = (double)_HeadRequestTimeout / 1000;
+	Configuration.HTTPMaximumConnectionsPerHost = _MaximumConnectionsPerHost;
+	_HeadRequestSession = [[NSURLSession sessionWithConfiguration:Configuration] retain];
 
 	const FString& DirectoryPath = FBackgroundHttpFileHashHelper::GetTemporaryRootPath();
 	if (ensureAlwaysMsgf(!DirectoryPath.IsEmpty(), TEXT("Invalid FBackgroundHttpFileHashHelper::GetTemporaryRootPath()")))
@@ -572,6 +594,9 @@ static constexpr NSInteger HTTPStatusCodeErrorServer = 500;
 {
 	[_Session release];
 	_Session = nil;
+	
+	[_HeadRequestSession release];
+	_HeadRequestSession = nil;
 
 	[_AllDownloads release];
 	_AllDownloads = nil;
@@ -635,7 +660,7 @@ static constexpr NSInteger HTTPStatusCodeErrorServer = 500;
 
 - (NSMutableArray<__kindof NSURL*>*)ReorderCDNsByReachability:(NSMutableArray<__kindof NSURL*>*)URLs
 {
-	if (_CDNReorderingTimeout == 0 || [URLs count] == 0)
+	if (_HeadRequestTimeout == 0 || [URLs count] == 0)
 	{
 		return URLs;
 	}
@@ -645,22 +670,7 @@ static constexpr NSInteger HTTPStatusCodeErrorServer = 500;
 		if ([_CDNInfo count] == 0)
 		{
 			UE_DNLD_LOG(@"Starting to check for CDN reachability");
-
-			const double CDNReorderingTimeoutInSeconds = (double)_CDNReorderingTimeout / 1000;
-
-			// Creating a temporary foreground NSURLSession for pinging CDN's
-			NSURLSessionConfiguration* Configuration = [NSURLSessionConfiguration defaultSessionConfiguration];
-			Configuration.discretionary = NO;
-			// Don't go over cellular here, during first download attempt we would try to over non-cellular connections first.
-			// Hence it makes more sense to prioritize CDN's that are reachable via non-cellular connections.
-			Configuration.allowsCellularAccess = NO;
-			Configuration.networkServiceType = NSURLNetworkServiceTypeResponsiveData;
-			Configuration.timeoutIntervalForRequest = CDNReorderingTimeoutInSeconds;
-			Configuration.timeoutIntervalForResource = CDNReorderingTimeoutInSeconds;
-			Configuration.HTTPMaximumConnectionsPerHost = _MaximumConnectionsPerHost;
-
-			NSURLSession* Session = [NSURLSession sessionWithConfiguration:Configuration];
-
+			
 			std::shared_ptr<std::atomic<int32>> PendingTasks = std::make_shared<std::atomic<int32>>();
 			std::shared_ptr<std::promise<void>> PendingTasksFinished = std::make_shared<std::promise<void>>();
 			std::shared_ptr<std::atomic<bool>> WaitingForTasksCompletionHandlers = std::make_shared<std::atomic<bool>>(true);
@@ -678,7 +688,7 @@ static constexpr NSInteger HTTPStatusCodeErrorServer = 500;
 				UE_DNLD_LOG(@"Create data task for '%@'", Request.URL.absoluteString);
 
 				// Note, completion handler might be invoked after end of this method.
-				NSURLSessionDataTask* Task = [Session dataTaskWithRequest:Request completionHandler:^(NSData* _Nullable Data, NSURLResponse* _Nullable Response, NSError* _Nullable Error)
+				NSURLSessionDataTask* Task = [_HeadRequestSession dataTaskWithRequest:Request completionHandler:^(NSData* _Nullable Data, NSURLResponse* _Nullable Response, NSError* _Nullable Error)
 				{
 					// The delegate can be invoked from other thread way past invalidateAndCancel,
 					// we cannot modify _CDNInfo without synchronization.
@@ -735,11 +745,11 @@ static constexpr NSInteger HTTPStatusCodeErrorServer = 500;
 				[Task resume];
 			}
 
-			PendingTasksFinished->get_future().wait_for(std::chrono::milliseconds(_CDNReorderingTimeout));
+			PendingTasksFinished->get_future().wait_for(std::chrono::milliseconds(_HeadRequestTimeout));
 			UE_DNLD_LOG(@"Finished waiting for CDN reachability");
 
 			WaitingForTasksCompletionHandlers->store(false);
-			[Session invalidateAndCancel];
+			[_HeadRequestSession resetWithCompletionHandler:^{}];
 
 			for (NSUInteger URLIndex = 0; URLIndex < [URLs count]; URLIndex++)
 			{
@@ -840,6 +850,18 @@ static constexpr NSInteger HTTPStatusCodeErrorServer = 500;
 	NSURLSessionDownloadTask* Task = [_Session downloadTaskWithRequest:URLRequest];
 	[Task setPriority:Priority];
 	[Task setTaskDescription:[TaskData ToSerializedString]];
+	if (TaskData.ExpectedResultSizeFromHeadRequest > 0)
+	{
+		if (TaskData.ExpectedResultSizeFromHeadRequest != TaskData.ExpectedResultSize)
+		{
+			UE_DNLD_LOG(@"Mismatch expectedSize from '%@' between HEAD request %llu and config %llu", URL.absoluteString, TaskData.ExpectedResultSizeFromHeadRequest, TaskData.ExpectedResultSize);
+		}
+		[Task setCountOfBytesClientExpectsToReceive:TaskData.ExpectedResultSizeFromHeadRequest + TypicalSizeForDownloadResponseHeader];
+	}
+	else
+	{
+		[Task setCountOfBytesClientExpectsToReceive:TaskData.ExpectedResultSize + TypicalSizeForDownloadResponseHeader];
+	}
 
 	UE_DNLD_LOG(@"CreateDownloadForURL '%@' with taskIdentifier %lu", Task.taskDescription, Task.taskIdentifier);
 	return Task;
@@ -923,9 +945,49 @@ static constexpr NSInteger HTTPStatusCodeErrorServer = 500;
 	}
 
 	NSURL* URL = [TaskData GetFirstURL];
-	NSURLSessionDownloadTask* Task = [self CreateDownloadForURL:URL WithPriority:Priority WithTaskData:TaskData];
-	const NSUInteger DownloadId = [self EnsureTaskIsTracked:Task];
-	[Task resume];
+	
+	NSURLSessionDownloadTask* DownloadTask = [self CreateDownloadForURL:URL WithPriority:Priority WithTaskData:TaskData];
+	const NSUInteger DownloadId = [self EnsureTaskIsTracked:DownloadTask];
+	
+	// Do not start the Download yet, use HEAD request to get more accurate expected size first, also may trigger CDN cache fetching
+	NSMutableURLRequest* Request = [NSMutableURLRequest requestWithURL:URL];
+	[Request setHTTPMethod:@"HEAD"];
+
+	UE_DNLD_LOG(@"Create HEAD request for '%@'", Request.URL.absoluteString);
+
+	NSURLSessionDataTask* HeadRequestTask = [_HeadRequestSession dataTaskWithRequest:Request completionHandler:^(NSData* _Nullable Data, NSURLResponse* _Nullable Response, NSError* _Nullable Error)
+	{
+		if (Response != nil && [Response isKindOfClass:[NSHTTPURLResponse class]])
+		{
+			NSHTTPURLResponse* HTTPResponse = (NSHTTPURLResponse*)Response;
+			[TaskData setExpectedResultSizeFromHeadRequest: (NSUInteger)[[HTTPResponse valueForHTTPHeaderField:@"Content-Length"] integerValue]];
+			UE_DNLD_LOG(@"Finished HEAD request for '%@' (host '%@') with status code %li and Content-Length %llu", Request.URL.absoluteString, Request.URL.host, HTTPResponse.statusCode, TaskData.ExpectedResultSizeFromHeadRequest);
+		}
+		else
+		{
+			UE_DNLD_LOG(@"Finished HEAD request for '%@' with error '%@' %i", Request.URL.absoluteString, (Error != nil ? Error.localizedDescription : @"nil"), (Error != nil ? (int32)Error.code : 0));
+		}
+		[DownloadTask setTaskDescription:[TaskData ToSerializedString]];
+		
+		if (TaskData.ExpectedResultSizeFromHeadRequest > 0)
+		{
+			if (TaskData.ExpectedResultSizeFromHeadRequest != TaskData.ExpectedResultSize)
+			{
+				UE_DNLD_LOG(@"Mismatch expectedSize from '%@' between HEAD request %llu and config %llu", URL.absoluteString, TaskData.ExpectedResultSizeFromHeadRequest, TaskData.ExpectedResultSize);
+			}
+			[DownloadTask setCountOfBytesClientExpectsToReceive:TaskData.ExpectedResultSizeFromHeadRequest + TypicalSizeForDownloadResponseHeader];
+		}
+		else
+		{
+			UE_DNLD_LOG(@"ExpectedSize from HEAD request is missing for '%@', using size from DownloadConfig %llu", URL.absoluteString, TaskData.ExpectedResultSize);
+			[DownloadTask setCountOfBytesClientExpectsToReceive:TaskData.ExpectedResultSize + TypicalSizeForDownloadResponseHeader];
+		}
+		
+		[DownloadTask resume];
+	}];
+	
+	[HeadRequestTask setPriority:NSURLSessionTaskPriorityHigh];
+	[HeadRequestTask resume];
 
 	return DownloadId;
 }
