[Input] Input devices from IInputModules can now be used to skip startup movies.

Previously we registered input devices the first time we called PollGameDeviceState after GIsRunning is true. Waiting for GIsRunning was used to ensure modules were loaded (e.g., see CL 28678218). Since GIsRunning is not set to true until after we startup movies finish, these input devices could not be used to skip.

This change no longer waits for GIsRunning to be set. Instead, input devices from loaded modules are registered on the first call to PollGameDeviceState. Any modules loaded after that point will invoke the OnInputDeviceModuleRegistered callback when loaded.

Theoretically, we could move the initial registration to happen earlier, rather than deferring until the first call to PollGameDeviceState. However, there are initialization order issues that would need to be addressed to make that change. That is why this preserves the existing behavior that defers the initial registration until it's first needed.

diff --git a/Engine/Source/Runtime/ApplicationCore/Private/Android/AndroidApplication.cpp b/Engine/Source/Runtime/ApplicationCore/Private/Android/AndroidApplication.cpp
--- a/Engine/Source/Runtime/ApplicationCore/Private/Android/AndroidApplication.cpp
+++ b/Engine/Source/Runtime/ApplicationCore/Private/Android/AndroidApplication.cpp
@@ -54,6 +54,11 @@ FAndroidApplication::FAndroidApplication(TSharedPtr<class FAndroidInputInterface
 	_application = this;
 }
 
+FAndroidApplication::~FAndroidApplication()
+{
+	IModularFeatures::Get().OnModularFeatureRegistered().RemoveAll(this);
+}
+
 void FAndroidApplication::SetMessageHandler( const TSharedRef< FGenericApplicationMessageHandler >& InMessageHandler )
 {
 	GenericApplication::SetMessageHandler(InMessageHandler);
@@ -64,22 +69,25 @@ void FAndroidApplication::AddExternalInputDevice(TSharedPtr<IInputDevice> InputD
 {
 	if (InputDevice.IsValid())
 	{
-		InputInterface->AddExternalInputDevice(InputDevice);
+		InputInterface->AddExternalInputDevice(MoveTemp(InputDevice));
 	}
 }
 
 void FAndroidApplication::PollGameDeviceState( const float TimeDelta )
 {
-	// initialize any externally-implemented input devices (we delay load initialize the array so any plugins have had time to load)
-	if (!bHasLoadedInputPlugins && GIsRunning)
+	if (!bHasLoadedInputPlugins)
 	{
+		// Create all externally-implemented input devices from loaded modules.
 		TArray<IInputDeviceModule*> PluginImplementations = IModularFeatures::Get().GetModularFeatureImplementations<IInputDeviceModule>(IInputDeviceModule::GetModularFeatureName());
 		for (auto InputPluginIt = PluginImplementations.CreateIterator(); InputPluginIt; ++InputPluginIt)
 		{
 			TSharedPtr<IInputDevice> Device = (*InputPluginIt)->CreateInputDevice(MessageHandler);
-			AddExternalInputDevice(Device);
+			AddExternalInputDevice(MoveTemp(Device));
 		}
 
+		// Modules loaded after this point are handled by this callback
+		IModularFeatures::Get().OnModularFeatureRegistered().AddRaw(this, &FAndroidApplication::OnInputDeviceModuleRegistered);
+
 		bHasLoadedInputPlugins = true;
 	}
 
@@ -222,3 +230,13 @@ void FAndroidApplication::OnWindowSizeChanged()
 	bWindowSizeChanged = true;
 }
 
+void FAndroidApplication::OnInputDeviceModuleRegistered(const FName& Type, IModularFeature* ModularFeature)
+{
+	if (Type == IInputDeviceModule::GetModularFeatureName())
+	{
+		IInputDeviceModule* Module = static_cast<IInputDeviceModule*>(ModularFeature);
+
+		TSharedPtr<IInputDevice> Device = Module->CreateInputDevice(MessageHandler);
+		AddExternalInputDevice(MoveTemp(Device));
+	}
+}
diff --git a/Engine/Source/Runtime/ApplicationCore/Private/IOS/IOSApplication.cpp b/Engine/Source/Runtime/ApplicationCore/Private/IOS/IOSApplication.cpp
--- a/Engine/Source/Runtime/ApplicationCore/Private/IOS/IOSApplication.cpp
+++ b/Engine/Source/Runtime/ApplicationCore/Private/IOS/IOSApplication.cpp
@@ -38,6 +38,8 @@ FIOSApplication::FIOSApplication()
 
 FIOSApplication::~FIOSApplication()
 {
+	IModularFeatures::Get().OnModularFeatureRegistered().RemoveAll(this);
+
 #if WITH_ACCESSIBILITY
 	if (AccessibilityAnnouncementDelayTimer)
 	{
@@ -88,7 +90,7 @@ void FIOSApplication::AddExternalInputDevice(TSharedPtr<IInputDevice> InputDevic
 {
 	if (InputDevice.IsValid())
 	{
-		ExternalInputDevices.Add(InputDevice);
+		ExternalInputDevices.Add(MoveTemp(InputDevice));
 	}
 }
 
@@ -96,16 +98,19 @@ void FIOSApplication::PollGameDeviceState( const float TimeDelta )
 {
 	@autoreleasepool
 	{
-		// initialize any externally-implemented input devices (we delay load initialize the array so any plugins have had time to load)
-		if (!bHasLoadedInputPlugins && GIsRunning)
+		// Create all externally-implemented input devices from loaded modules.
+		if (!bHasLoadedInputPlugins)
 		{
 			TArray<IInputDeviceModule*> PluginImplementations = IModularFeatures::Get().GetModularFeatureImplementations<IInputDeviceModule>(IInputDeviceModule::GetModularFeatureName());
 			for (auto InputPluginIt = PluginImplementations.CreateIterator(); InputPluginIt; ++InputPluginIt)
 			{
 				TSharedPtr<IInputDevice> Device = (*InputPluginIt)->CreateInputDevice(MessageHandler);
-				AddExternalInputDevice(Device);
+				AddExternalInputDevice(MoveTemp(Device));
 			}
 
+			// Modules loaded after this point are handled by this callback
+			IModularFeatures::Get().OnModularFeatureRegistered().AddRaw(this, &FIOSApplication::OnInputDeviceModuleRegistered);
+
 			bHasLoadedInputPlugins = true;
 		}
 
@@ -357,6 +362,17 @@ TSharedRef<FIOSWindow> FIOSApplication::FindWindowByAppDelegateView()
     return Windows[0];
 }
 
+void FIOSApplication::OnInputDeviceModuleRegistered(const FName& Type, IModularFeature* ModularFeature)
+{
+	if (Type == IInputDeviceModule::GetModularFeatureName())
+	{
+		IInputDeviceModule* Module = static_cast<IInputDeviceModule*>(ModularFeature);
+
+		TSharedPtr<IInputDevice> Device = Module->CreateInputDevice(MessageHandler);
+		AddExternalInputDevice(MoveTemp(Device));
+	}
+}
+
 #if WITH_ACCESSIBILITY
 
 float GIOSAccessibleAnnouncementDelay = 0.1f;
diff --git a/Engine/Source/Runtime/ApplicationCore/Private/Linux/LinuxApplication.cpp b/Engine/Source/Runtime/ApplicationCore/Private/Linux/LinuxApplication.cpp
--- a/Engine/Source/Runtime/ApplicationCore/Private/Linux/LinuxApplication.cpp
+++ b/Engine/Source/Runtime/ApplicationCore/Private/Linux/LinuxApplication.cpp
@@ -135,6 +135,8 @@ FLinuxApplication::~FLinuxApplication()
 		GConfig->GetFloat(TEXT("X11.Tweaks"), TEXT("MouseWheelScrollAcceleration"), fMouseWheelScrollAccel, GEngineIni);
 		GConfig->Flush(false, GEngineIni);
 	}
+
+	IModularFeatures::Get().OnModularFeatureRegistered().RemoveAll(this);
 }
 
 void FLinuxApplication::DestroyApplication()
@@ -1207,9 +1209,9 @@ void FLinuxApplication::PollGameDeviceState( const float TimeDelta )
 		}
 	}
 
-	// initialize any externally-implemented input devices (we delay load initialize the array so any plugins have had time to load)
-	if (!bHasLoadedInputPlugins && GIsRunning)
+	if (!bHasLoadedInputPlugins)
 	{
+		// Create all externally-implemented input devices from loaded modules.
 		TArray<IInputDeviceModule*> PluginImplementations = IModularFeatures::Get().GetModularFeatureImplementations<IInputDeviceModule>(IInputDeviceModule::GetModularFeatureName());
 		for (auto InputPluginIt = PluginImplementations.CreateIterator(); InputPluginIt; ++InputPluginIt)
 		{
@@ -1217,10 +1219,13 @@ void FLinuxApplication::PollGameDeviceState( const float TimeDelta )
 			if (Device.IsValid())
 			{
 				UE_LOG(LogInit, Log, TEXT("Adding external input plugin."));
-				ExternalInputDevices.Add(Device);
+				ExternalInputDevices.Add(MoveTemp(Device));
 			}
 		}
 
+		// Modules loaded after this point are handled by this callback
+		IModularFeatures::Get().OnModularFeatureRegistered().AddRaw(this, &FLinuxApplication::OnInputDeviceModuleRegistered);
+
 		bHasLoadedInputPlugins = true;
 	}
 	
@@ -2116,3 +2121,19 @@ void FLinuxApplication::RemoveGameController(SDL_JoystickID Id)
 
 	DeviceMapper.Internal_MapInputDeviceToUser(DeviceId, PlatformUserId, EInputDeviceConnectionState::Disconnected);
 }
+
+void FLinuxApplication::OnInputDeviceModuleRegistered(const FName& Type, IModularFeature* ModularFeature)
+{
+	if (Type == IInputDeviceModule::GetModularFeatureName())
+	{
+		IInputDeviceModule* Module = static_cast<IInputDeviceModule*>(ModularFeature);
+
+		TSharedPtr<IInputDevice> Device = Module->CreateInputDevice(MessageHandler);
+		if (Device.IsValid())
+		{
+			UE_LOG(LogInit, Log, TEXT("Adding external input plugin."));
+			ExternalInputDevices.Add(MoveTemp(Device));
+		}
+	}
+}
+
diff --git a/Engine/Source/Runtime/ApplicationCore/Private/Mac/MacApplication.cpp b/Engine/Source/Runtime/ApplicationCore/Private/Mac/MacApplication.cpp
--- a/Engine/Source/Runtime/ApplicationCore/Private/Mac/MacApplication.cpp
+++ b/Engine/Source/Runtime/ApplicationCore/Private/Mac/MacApplication.cpp
@@ -200,6 +200,8 @@ FMacApplication::FMacApplication()
 
 FMacApplication::~FMacApplication()
 {
+	IModularFeatures::Get().OnModularFeatureRegistered().RemoveAll(this);
+
 	if (MenuBarShutdownFunc)
 	{
 		MenuBarShutdownFunc();
@@ -305,9 +307,9 @@ void FMacApplication::SetAccessibleMessageHandler(const TSharedRef<FGenericAcces
 
 void FMacApplication::PollGameDeviceState(const float TimeDelta)
 {
-	// initialize any externally-implemented input devices (we delay load initialize the array so any plugins have had time to load)
-	if (!bHasLoadedInputPlugins && GIsRunning)
+	if (!bHasLoadedInputPlugins)
 	{
+		// Create all externally-implemented input devices from loaded modules.
 		TArray<IInputDeviceModule*> PluginImplementations = IModularFeatures::Get().GetModularFeatureImplementations<IInputDeviceModule>( IInputDeviceModule::GetModularFeatureName() );
 		for( auto InputPluginIt = PluginImplementations.CreateIterator(); InputPluginIt; ++InputPluginIt )
 		{
@@ -315,10 +317,13 @@ void FMacApplication::PollGameDeviceState(const float TimeDelta)
 			if (Device.IsValid())
 			{
 				UE_LOG(LogInit, Log, TEXT("Adding external input plugin."));
-				ExternalInputDevices.Add(Device);
+				ExternalInputDevices.Add(MoveTemp(Device));
 			}
 		}
 
+		// Modules loaded after this point are handled by this callback
+		IModularFeatures::Get().OnModularFeatureRegistered().AddRaw(this, &FMacApplication::OnInputDeviceModuleRegistered);
+
 		bHasLoadedInputPlugins = true;
 	}
 
@@ -2375,6 +2380,21 @@ void FDisplayMetrics::RebuildDisplayMetrics(FDisplayMetrics& OutDisplayMetrics)
 	OutDisplayMetrics.ApplyDefaultSafeZones();
 }
 
+void FMacApplication::OnInputDeviceModuleRegistered(const FName& Type, IModularFeature* ModularFeature)
+{
+	if (Type == IInputDeviceModule::GetModularFeatureName())
+	{
+		IInputDeviceModule* Module = static_cast<IInputDeviceModule*>(ModularFeature);
+
+		TSharedPtr<IInputDevice> Device = Module->CreateInputDevice(MessageHandler);
+		if (Device.IsValid())
+		{
+			UE_LOG(LogInit, Log, TEXT("Adding external input plugin."));
+			ExternalInputDevices.Add(MoveTemp(Device));
+		}
+	}
+}
+
 #if WITH_ACCESSIBILITY
 float GMacAccessibleAnnouncementDelay = 0.1f;
 FAutoConsoleVariableRef MacAccessibleAnnouncementDealyRef(
diff --git a/Engine/Source/Runtime/ApplicationCore/Private/Null/NullApplication.cpp b/Engine/Source/Runtime/ApplicationCore/Private/Null/NullApplication.cpp
--- a/Engine/Source/Runtime/ApplicationCore/Private/Null/NullApplication.cpp
+++ b/Engine/Source/Runtime/ApplicationCore/Private/Null/NullApplication.cpp
@@ -124,6 +124,7 @@ FNullApplication::FNullApplication()
 
 FNullApplication::~FNullApplication()
 {
+	IModularFeatures::Get().OnModularFeatureRegistered().RemoveAll(this);
 }
 
 void FNullApplication::DestroyApplication()
@@ -217,9 +218,9 @@ void FNullApplication::ProcessDeferredEvents(const float TimeDelta)
 void FNullApplication::PollGameDeviceState(const float TimeDelta)
 {
 	IPlatformInputDeviceMapper& Mapper = IPlatformInputDeviceMapper::Get();
-	// initialize any externally-implemented input devices (we delay load initialize the array so any plugins have had time to load)
-	if (!bHasLoadedInputPlugins && GIsRunning)
+	if (!bHasLoadedInputPlugins)
 	{
+		// Create all externally-implemented input devices from loaded modules.
 		TArray<IInputDeviceModule*> PluginImplementations = IModularFeatures::Get().GetModularFeatureImplementations<IInputDeviceModule>(IInputDeviceModule::GetModularFeatureName());
 		for (auto InputPluginIt = PluginImplementations.CreateIterator(); InputPluginIt; ++InputPluginIt)
 		{
@@ -227,10 +228,13 @@ void FNullApplication::PollGameDeviceState(const float TimeDelta)
 			if (Device.IsValid())
 			{
 				UE_LOG(LogInit, Log, TEXT("Adding external input plugin."));
-				ExternalInputDevices.Add(Device);
+				ExternalInputDevices.Add(MoveTemp(Device));
 			}
 		}
 
+		// Modules loaded after this point are handled by this callback
+		IModularFeatures::Get().OnModularFeatureRegistered().AddRaw(this, &FNullApplication::OnInputDeviceModuleRegistered);
+
 		bHasLoadedInputPlugins = true;
 	}
 
@@ -415,6 +419,21 @@ bool FNullApplication::ParseResolution(const TCHAR* InResolution, uint32& OutX,
 	return false;
 }
 
+void FNullApplication::OnInputDeviceModuleRegistered(const FName& Type, IModularFeature* ModularFeature)
+{
+	if (Type == IInputDeviceModule::GetModularFeatureName())
+	{
+		IInputDeviceModule* Module = static_cast<IInputDeviceModule*>(ModularFeature);
+
+		TSharedPtr<IInputDevice> Device = Module->CreateInputDevice(MessageHandler);
+		if (Device.IsValid())
+		{
+			UE_LOG(LogInit, Log, TEXT("Adding external input plugin."));
+			ExternalInputDevices.Add(MoveTemp(Device));
+		}
+	}
+}
+
 void FNullPlatformDisplayMetrics::RebuildDisplayMetrics(struct FDisplayMetrics& OutDisplayMetrics)
 {
 	int32 X, Y, Width, Height;
diff --git a/Engine/Source/Runtime/ApplicationCore/Private/Windows/WindowsApplication.cpp b/Engine/Source/Runtime/ApplicationCore/Private/Windows/WindowsApplication.cpp
--- a/Engine/Source/Runtime/ApplicationCore/Private/Windows/WindowsApplication.cpp
+++ b/Engine/Source/Runtime/ApplicationCore/Private/Windows/WindowsApplication.cpp
@@ -1461,6 +1461,8 @@ FWindowsApplication::~FWindowsApplication()
 		UE::WindowsApplication::RawInputThread::FWindowsRawInputRunnable::Shutdown();
 	}
 
+	IModularFeatures::Get().OnModularFeatureRegistered().RemoveAll(this);
+
 	FCoreDelegates::OnShutdownAfterError.RemoveAll(this);
 
 	if (TextInputMethodSystem.IsValid())
@@ -3878,6 +3880,17 @@ LRESULT CALLBACK FWindowsApplication::HandleLowLevelMouseFilterHook(int nCode, W
 
 #undef MOUSEEVENTF_FROMTOUCH
 
+void FWindowsApplication::OnInputDeviceModuleRegistered(const FName& Type, IModularFeature* ModularFeature)
+{
+	if (Type == IInputDeviceModule::GetModularFeatureName())
+	{
+		IInputDeviceModule* Module = static_cast<IInputDeviceModule*>(ModularFeature);
+
+		TSharedPtr<IInputDevice> Device = Module->CreateInputDevice(MessageHandler);
+		AddExternalInputDevice(MoveTemp(Device));
+	}
+}
+
 void FWindowsApplication::DeferMessage( TSharedPtr<FWindowsWindow>& NativeWindow, HWND InHWnd, uint32 InMessage, WPARAM InWParam, LPARAM InLParam, int32 MouseX, int32 MouseY, uint32 RawInputFlags )
 {
 	if( GPumpingMessagesOutsideOfMainLoop && bAllowedToDeferMessageProcessing )
@@ -3944,16 +3957,19 @@ void FWindowsApplication::PollGameDeviceState( const float TimeDelta )
 		return;
 	}
 
-	// initialize any externally-implemented input devices (we delay load initialize the array so any plugins have had time to load)
-	if (!bHasLoadedInputPlugins && GIsRunning )
+	if (!bHasLoadedInputPlugins)
 	{
+		// Create all externally-implemented input devices from loaded modules.
 		TArray<IInputDeviceModule*> PluginImplementations = IModularFeatures::Get().GetModularFeatureImplementations<IInputDeviceModule>( IInputDeviceModule::GetModularFeatureName() );
 		for( auto InputPluginIt = PluginImplementations.CreateIterator(); InputPluginIt; ++InputPluginIt )
 		{
 			TSharedPtr<IInputDevice> Device = (*InputPluginIt)->CreateInputDevice(MessageHandler);
-			AddExternalInputDevice(Device);			
+			AddExternalInputDevice(MoveTemp(Device));			
 		}
 
+		// Modules loaded after this point are handled by this callback
+		IModularFeatures::Get().OnModularFeatureRegistered().AddRaw(this, &FWindowsApplication::OnInputDeviceModuleRegistered);
+
 		bHasLoadedInputPlugins = true;
 	}
 
@@ -4038,7 +4054,7 @@ void FWindowsApplication::AddExternalInputDevice(TSharedPtr<IInputDevice> InputD
 {
 	if (InputDevice.IsValid())
 	{
-		ExternalInputDevices.Add(InputDevice);
+		ExternalInputDevices.Add(MoveTemp(InputDevice));
 	}
 }
 
diff --git a/Engine/Source/Runtime/ApplicationCore/Public/Android/AndroidApplication.h b/Engine/Source/Runtime/ApplicationCore/Public/Android/AndroidApplication.h
--- a/Engine/Source/Runtime/ApplicationCore/Public/Android/AndroidApplication.h
+++ b/Engine/Source/Runtime/ApplicationCore/Public/Android/AndroidApplication.h
@@ -70,7 +70,7 @@ public:
 
 public:	
 	
-	virtual ~FAndroidApplication() {}
+	virtual ~FAndroidApplication() override;
 
 	void SetMessageHandler( const TSharedRef< FGenericApplicationMessageHandler >& InMessageHandler );
 
@@ -105,6 +105,8 @@ protected:
 
 private:
 
+	void OnInputDeviceModuleRegistered(const FName& Type, IModularFeature* ModularFeature);
+
 	TSharedPtr< class FAndroidInputInterface > InputInterface;
 	bool bHasLoadedInputPlugins;
 
diff --git a/Engine/Source/Runtime/ApplicationCore/Public/IOS/IOSApplication.h b/Engine/Source/Runtime/ApplicationCore/Public/IOS/IOSApplication.h
--- a/Engine/Source/Runtime/ApplicationCore/Public/IOS/IOSApplication.h
+++ b/Engine/Source/Runtime/ApplicationCore/Public/IOS/IOSApplication.h
@@ -55,6 +55,8 @@ private:
 	void OnAccessibleEventRaised(const FAccessibleEventArgs& Args);
 #endif
 
+	void OnInputDeviceModuleRegistered(const FName& Type, IModularFeature* ModularFeature);
+
 private:
 
 	TSharedPtr< class FIOSInputInterface > InputInterface;
diff --git a/Engine/Source/Runtime/ApplicationCore/Public/Linux/LinuxApplication.h b/Engine/Source/Runtime/ApplicationCore/Public/Linux/LinuxApplication.h
--- a/Engine/Source/Runtime/ApplicationCore/Public/Linux/LinuxApplication.h
+++ b/Engine/Source/Runtime/ApplicationCore/Public/Linux/LinuxApplication.h
@@ -232,6 +232,8 @@ private:
 	void AddGameController(SDL_JoystickID Id);
 	void RemoveGameController(SDL_JoystickID Id);
 
+	void OnInputDeviceModuleRegistered(const FName& Type, IModularFeature* ModularFeature);
+
 	/** Stores context information about a currently active touch. */
 	struct FTouchContext
 	{
diff --git a/Engine/Source/Runtime/ApplicationCore/Public/Mac/MacApplication.h b/Engine/Source/Runtime/ApplicationCore/Public/Mac/MacApplication.h
--- a/Engine/Source/Runtime/ApplicationCore/Public/Mac/MacApplication.h
+++ b/Engine/Source/Runtime/ApplicationCore/Public/Mac/MacApplication.h
@@ -330,6 +330,8 @@ private:
 	void OnAccessibleEventRaised(const FAccessibleEventArgs& Args);
 	#endif
 
+	void OnInputDeviceModuleRegistered(const FName& Type, IModularFeature* ModularFeature);
+
 private:
 
 	bool bUsingHighPrecisionMouseInput;
diff --git a/Engine/Source/Runtime/ApplicationCore/Public/Null/NullApplication.h b/Engine/Source/Runtime/ApplicationCore/Public/Null/NullApplication.h
--- a/Engine/Source/Runtime/ApplicationCore/Public/Null/NullApplication.h
+++ b/Engine/Source/Runtime/ApplicationCore/Public/Null/NullApplication.h
@@ -20,6 +20,7 @@
 #include "Null/NullCursor.h"
 
 class IInputDevice;
+class IModularFeature;
 
 struct FNullPlatformDisplayMetrics : public FDisplayMetrics
 {
@@ -113,6 +114,8 @@ private:
 	/** Handles parsing the work area resolution from the command line */
 	APPLICATIONCORE_API bool ParseResolution(const TCHAR* InResolution, uint32& OutX, uint32& OutY);
 
+	void OnInputDeviceModuleRegistered(const FName& Type, IModularFeature* ModularFeature);
+
 public:
 	virtual IInputInterface* GetInputInterface() override
 	{
diff --git a/Engine/Source/Runtime/ApplicationCore/Public/Windows/WindowsApplication.h b/Engine/Source/Runtime/ApplicationCore/Public/Windows/WindowsApplication.h
--- a/Engine/Source/Runtime/ApplicationCore/Public/Windows/WindowsApplication.h
+++ b/Engine/Source/Runtime/ApplicationCore/Public/Windows/WindowsApplication.h
@@ -526,6 +526,8 @@ private:
 	/** Unregisters for raw input */
 	void UnregisterHighPrecisionMouseMode();
 
+	void OnInputDeviceModuleRegistered(const FName& Type, IModularFeature* ModularFeature);
+
 private:
 
 	static APPLICATIONCORE_API const FIntPoint MinimizedWindowPosition;
