Adds missing forward declaration of IModularFeature.

#rnx

diff --git a/Engine/Source/Runtime/ApplicationCore/Public/Linux/LinuxApplication.h b/Engine/Source/Runtime/ApplicationCore/Public/Linux/LinuxApplication.h
--- a/Engine/Source/Runtime/ApplicationCore/Public/Linux/LinuxApplication.h
+++ b/Engine/Source/Runtime/ApplicationCore/Public/Linux/LinuxApplication.h
@@ -19,6 +19,7 @@
 #include "Linux/LinuxCursor.h"
 
 class IInputDevice;
+class IModularFeature;
 
 class FLinuxApplication : public GenericApplication, public FSelfRegisteringExec, public IInputInterface
 {
