Adds forward declarations.

Adds them to all code impacted by CL 45168962, to address the Mac non-unity build failure and avoid these problems in the future. Linux was already addressed in CL 45179276.

#rnx

diff --git a/Engine/Source/Runtime/ApplicationCore/Public/Android/AndroidApplication.h b/Engine/Source/Runtime/ApplicationCore/Public/Android/AndroidApplication.h
--- a/Engine/Source/Runtime/ApplicationCore/Public/Android/AndroidApplication.h
+++ b/Engine/Source/Runtime/ApplicationCore/Public/Android/AndroidApplication.h
@@ -9,6 +9,8 @@
 #include "Android/AndroidJavaEnv.h"
 #endif
 
+class IModularFeature;
+
 namespace FAndroidAppEntry
 {
 	void PlatformInit();
diff --git a/Engine/Source/Runtime/ApplicationCore/Public/IOS/IOSApplication.h b/Engine/Source/Runtime/ApplicationCore/Public/IOS/IOSApplication.h
--- a/Engine/Source/Runtime/ApplicationCore/Public/IOS/IOSApplication.h
+++ b/Engine/Source/Runtime/ApplicationCore/Public/IOS/IOSApplication.h
@@ -6,6 +6,7 @@
 #include "GenericPlatform/GenericApplication.h"
 
 class FIOSWindow;
+class IModularFeature;
 
 class FIOSApplication : public GenericApplication
 {
diff --git a/Engine/Source/Runtime/ApplicationCore/Public/Mac/MacApplication.h b/Engine/Source/Runtime/ApplicationCore/Public/Mac/MacApplication.h
--- a/Engine/Source/Runtime/ApplicationCore/Public/Mac/MacApplication.h
+++ b/Engine/Source/Runtime/ApplicationCore/Public/Mac/MacApplication.h
@@ -8,6 +8,8 @@
 #include "Mac/MacTextInputMethodSystem.h"
 #include "GenericPlatform/IInputInterface.h"
 
+class IModularFeature;
+
 struct FDeferredMacEvent
 {
 	FDeferredMacEvent()
diff --git a/Engine/Source/Runtime/ApplicationCore/Public/Windows/WindowsApplication.h b/Engine/Source/Runtime/ApplicationCore/Public/Windows/WindowsApplication.h
--- a/Engine/Source/Runtime/ApplicationCore/Public/Windows/WindowsApplication.h
+++ b/Engine/Source/Runtime/ApplicationCore/Public/Windows/WindowsApplication.h
@@ -22,6 +22,7 @@ class ITextInputMethodSystem;
 enum class FForceFeedbackChannelType;
 struct FForceFeedbackValues;
 struct FHapticFeedbackValues;
+class IModularFeature;
 
 DECLARE_LOG_CATEGORY_EXTERN(LogWindowsDesktop, Log, All);
 
