Adding FInputDeviceScope information to FAppleControllerInterface and FIOSInputInterface to allow developers to identify which device is currently sending events through FInputDeviceScope::GetCurrent().

#tests tested in iOS with Dual Sense and Xbox controllers

diff --git a/Engine/Source/Runtime/ApplicationCore/Private/Apple/AppleControllerInterface.cpp b/Engine/Source/Runtime/ApplicationCore/Private/Apple/AppleControllerInterface.cpp
--- a/Engine/Source/Runtime/ApplicationCore/Private/Apple/AppleControllerInterface.cpp
+++ b/Engine/Source/Runtime/ApplicationCore/Private/Apple/AppleControllerInterface.cpp
@@ -9,6 +9,19 @@ DEFINE_LOG_CATEGORY(LogAppleController);
 
 #define APPLE_CONTROLLER_DEBUG 0
 
+namespace
+{
+	// Controller Type Strings
+	static const FName UnassignedControllerTypeName = TEXT("Unassigned");
+	static const FName SiriRemoteControllerTypeName = TEXT("SiriRemote");
+	static const FName ExtendedGamepadControllerTypeName = TEXT("ExtendedGamepad");
+	static const FName XboxControllerTypeName = TEXT("Xbox Wireless Controller");
+	static const FName DualShockControllerTypeName = TEXT("PS4 Wireless Controller");
+	static const FName DualSenseControllerTypeName = TEXT("PS5 Wireless Controller");
+}
+
+FString FAppleControllerInterface::HardwareDeviceIdentifier_DefaultGamepad;
+
 TSharedRef< FAppleControllerInterface > FAppleControllerInterface::Create(  const TSharedRef< FGenericApplicationMessageHandler >& InMessageHandler )
 {
 	return MakeShareable( new FAppleControllerInterface( InMessageHandler ) );
@@ -53,6 +66,8 @@ FAppleControllerInterface::FAppleControllerInterface( const TSharedRef< FGeneric
 		   [GCController startWirelessControllerDiscoveryWithCompletionHandler:^{ }];
 		});
 	}
+
+	HardwareDeviceIdentifier_DefaultGamepad = TEXT("Gamepad");
 }
 
 void FAppleControllerInterface::SetMessageHandler( const TSharedRef< FGenericApplicationMessageHandler >& InMessageHandler )
@@ -151,6 +166,33 @@ void FAppleControllerInterface::SetControllerType(uint32 ControllerIndex)
     }
 }
 
+FName FAppleControllerInterface::GetControllerTypeName(const ControllerType InControllerType)
+{
+	FName ControllerTypeName = UnassignedControllerTypeName;
+	switch (InControllerType)
+	{
+		case ControllerType::SiriRemote:
+			ControllerTypeName = SiriRemoteControllerTypeName;
+			break;
+		case ControllerType::ExtendedGamepad:
+			ControllerTypeName = ExtendedGamepadControllerTypeName;
+			break;
+		case ControllerType::XboxGamepad:
+			ControllerTypeName = XboxControllerTypeName;
+			break;
+		case ControllerType::DualShockGamepad:
+			ControllerTypeName = DualShockControllerTypeName;
+			break;
+		case ControllerType::DualSenseGamepad:
+			ControllerTypeName = DualSenseControllerTypeName;
+			break;
+		default:
+			break;
+	}
+
+	return ControllerTypeName;
+}
+
 void FAppleControllerInterface::SetCurrentController(GCController* Controller)
 {
     int32 ControllerIndex = 0;
@@ -271,6 +313,14 @@ void FAppleControllerInterface::SendControllerEvents()
 			continue;
 		}
 
+		IPlatformInputDeviceMapper& DeviceMapper = IPlatformInputDeviceMapper::Get();
+		FPlatformUserId UserId = FGenericPlatformMisc::GetPlatformUserForUserIndex(Controller.PlayerIndex);
+		FInputDeviceId DeviceId = DeviceMapper.GetPrimaryInputDeviceForUser(UserId);
+		///The FInputDeviceScope::HardwareDeviceIdentifier has to be one of values in UInputPlatformSettings::HardwareDevices.
+		// This is a temp solution with a hardcoded string which can be mapped to FHardwareDeviceIdentifier::DefaultGamepad.
+		// TODO: Future improvement is needed to acquire them by values in the iOS input device info.
+		FInputDeviceScope InputScope(nullptr, GetControllerTypeName(Controller.ControllerType), DeviceId.GetId(), HardwareDeviceIdentifier_DefaultGamepad);
+		
 		HandleButtonGamepad(FGamepadKeyNames::FaceButtonBottom, i);
 		HandleButtonGamepad(FGamepadKeyNames::FaceButtonLeft, i);
 		HandleButtonGamepad(FGamepadKeyNames::FaceButtonRight, i);
diff --git a/Engine/Source/Runtime/ApplicationCore/Private/Apple/AppleControllerInterface.h b/Engine/Source/Runtime/ApplicationCore/Private/Apple/AppleControllerInterface.h
--- a/Engine/Source/Runtime/ApplicationCore/Private/Apple/AppleControllerInterface.h
+++ b/Engine/Source/Runtime/ApplicationCore/Private/Apple/AppleControllerInterface.h
@@ -94,6 +94,7 @@ public:
 
 	const ControllerType GetControllerType(uint32 ControllerIndex);
     void SetControllerType(uint32 ControllerIndex);
+	FName GetControllerTypeName(const ControllerType InControllerType);
     
     GCControllerButtonInput* GetGCControllerButton(const FGamepadKeyNames::Type& ButtonKey, uint32 ControllerIndex);
     
@@ -150,4 +151,6 @@ protected:
 	
 	// should we allow controllers to send input
 	bool bAllowControllers;
+
+	static FString HardwareDeviceIdentifier_DefaultGamepad;
 };
diff --git a/Engine/Source/Runtime/ApplicationCore/Private/IOS/IOSInputInterface.cpp b/Engine/Source/Runtime/ApplicationCore/Private/IOS/IOSInputInterface.cpp
--- a/Engine/Source/Runtime/ApplicationCore/Private/IOS/IOSInputInterface.cpp
+++ b/Engine/Source/Runtime/ApplicationCore/Private/IOS/IOSInputInterface.cpp
@@ -142,6 +142,8 @@ static TArray<TouchInput> TouchInputStack;
 static TArray<int32> KeyInputStack;
 
 bool FIOSInputInterface::bKeyboardInhibited = false;
+FName FIOSInputInterface::InputClassName_DefaultMobileTouch;
+FString FIOSInputInterface::HardwareDeviceIdentifier_DefaultMobileTouch;
 
 TSharedRef< FIOSInputInterface > FIOSInputInterface::Create(  const TSharedRef< FGenericApplicationMessageHandler >& InMessageHandler )
 {
@@ -245,6 +247,9 @@ FIOSInputInterface::FIOSInputInterface( const TSharedRef< FGenericApplicationMes
 #else
     HapticFeedbackSupportLevel = 0;
 #endif
+
+	InputClassName_DefaultMobileTouch = TEXT("DefaultMobileTouch");
+	HardwareDeviceIdentifier_DefaultMobileTouch = TEXT("MobileTouch");
 }
 
 template <EIOSEventType U, EIOSEventType V>
@@ -398,6 +403,11 @@ void FIOSInputInterface::ProcessTouchesAndKeys(uint32 ControllerId, const TArray
     {
         const TouchInput& Touch = InTouchInputStack[i];
         
+		///The FInputDeviceScope::HardwareDeviceIdentifier has to be one of values in UInputPlatformSettings::HardwareDevices.
+		// This is a temp solution with a hardcoded string which can be mapped to FHardwareDeviceIdentifier::DefaultMobileTouch.
+		// TODO: Future improvement is needed to acquire them by values in the iOS input device info.
+		FInputDeviceScope InputScope(nullptr, InputClassName_DefaultMobileTouch, 0, HardwareDeviceIdentifier_DefaultMobileTouch);
+		
         // send input to handler
         if (Touch.Type == TouchBegan)
         {
diff --git a/Engine/Source/Runtime/ApplicationCore/Public/IOS/IOSInputInterface.h b/Engine/Source/Runtime/ApplicationCore/Public/IOS/IOSInputInterface.h
--- a/Engine/Source/Runtime/ApplicationCore/Public/IOS/IOSInputInterface.h
+++ b/Engine/Source/Runtime/ApplicationCore/Public/IOS/IOSInputInterface.h
@@ -219,4 +219,7 @@ private:
         float ScrollDeltaY;
         bool bHaveMouse;
         static bool bKeyboardInhibited;
+	
+	static FName InputClassName_DefaultMobileTouch;
+	static FString HardwareDeviceIdentifier_DefaultMobileTouch;
 };
