Fixing GetApplicationWillTerminateDelegate iOS call bug and bFocused uninitialized bug

diff --git a/Engine/Source/Runtime/ApplicationCore/Private/IOS/IOSAppDelegate.cpp b/Engine/Source/Runtime/ApplicationCore/Private/IOS/IOSAppDelegate.cpp
--- a/Engine/Source/Runtime/ApplicationCore/Private/IOS/IOSAppDelegate.cpp
+++ b/Engine/Source/Runtime/ApplicationCore/Private/IOS/IOSAppDelegate.cpp
@@ -1832,13 +1832,10 @@ extern double GCStartTime;
 	 Save data if appropriate.
 	 See also applicationDidEnterBackground:.
 	 */
-	AsyncTask(ENamedThreads::Type::GameThread, []
-    {
-		if (!UE::IsGCLocked())
-		{
-			FCoreDelegates::GetApplicationWillTerminateDelegate().Broadcast();
-		}
-	});
+	if (!UE::IsGCLocked())
+	{
+		FCoreDelegates::GetApplicationWillTerminateDelegate().Broadcast();
+	}
     
     // note that we are shutting down
     // TODO: fix the reason why we are hanging when asked to shutdown
diff --git a/Engine/Source/Runtime/WebBrowser/Private/Apple/ApplePlatformWebBrowser.cpp b/Engine/Source/Runtime/WebBrowser/Private/Apple/ApplePlatformWebBrowser.cpp
--- a/Engine/Source/Runtime/WebBrowser/Private/Apple/ApplePlatformWebBrowser.cpp
+++ b/Engine/Source/Runtime/WebBrowser/Private/Apple/ApplePlatformWebBrowser.cpp
@@ -72,6 +72,7 @@ class SAppleWebBrowserWidget : public SViewport
 		WebBrowserWindowPtr = Args._WebBrowserWindow;
 		Is3DBrowser = false;
 		bAcquiredInitialFocus = false;
+		bFocused = false;
 
 		bool bEnableFloatingCloseButton = false;
 		GConfig->GetBool(TEXT("Browser"), TEXT("bEnableFloatingCloseButton"), bEnableFloatingCloseButton, GEngineIni);
