Add InstallBundleManager.CDN.ResponseTime event

#rnx

diff --git a/Engine/Plugins/Experimental/DefaultInstallBundleManager/Source/Private/InstallBundleManagerUtil.cpp b/Engine/Plugins/Experimental/DefaultInstallBundleManager/Source/Private/InstallBundleManagerUtil.cpp
--- a/Engine/Plugins/Experimental/DefaultInstallBundleManager/Source/Private/InstallBundleManagerUtil.cpp
+++ b/Engine/Plugins/Experimental/DefaultInstallBundleManager/Source/Private/InstallBundleManagerUtil.cpp
@@ -824,9 +824,8 @@ namespace InstallBundleManagerAnalytics
 		JsonWriter.WriteArrayStart();
 		if (ChunkDownloadInfos)
 		{
-			for (int i = 0; i < ChunkDownloadInfos->Num(); i++)
+			for (const auto& ChunkInfo : *ChunkDownloadInfos)
 			{
-				const FInstallBundleChunkDownloadInfo& ChunkInfo = (*ChunkDownloadInfos)[i];
 				JsonWriter.WriteObjectStart();
 
 				JsonWriter.WriteValue(TEXT("ChunkName"), ChunkInfo.ChunkName);
diff --git a/Engine/Source/Runtime/ApplicationCore/Private/IOS/IOSBackgroundURLSessionHandler.cpp b/Engine/Source/Runtime/ApplicationCore/Private/IOS/IOSBackgroundURLSessionHandler.cpp
--- a/Engine/Source/Runtime/ApplicationCore/Private/IOS/IOSBackgroundURLSessionHandler.cpp
+++ b/Engine/Source/Runtime/ApplicationCore/Private/IOS/IOSBackgroundURLSessionHandler.cpp
@@ -1867,9 +1867,9 @@ void FBackgroundURLSessionHandler::SaveBackgroundHttpFileHashHelperState()
 	}
 }
 
-TArray<FString> FBackgroundURLSessionHandler::GetCDNOrderArray()
+TArray<TTuple<FString, uint32, FString>> FBackgroundURLSessionHandler::GetCDNAnalyticsData()
 {
-	TArray<FString> Result;
+	TArray<TTuple<FString, uint32, FString>> Result;
 
 	@autoreleasepool
 	{
@@ -1877,7 +1877,20 @@ TArray<FString> FBackgroundURLSessionHandler::GetCDNOrderArray()
 
 		for (FBackgroundNSURLCDNInfo* Info in CDNInfo)
 		{
-			Result.Add(FString(Info.CDNAbsoluteURL));
+			FString StatusString;
+			switch (Info.Response) {
+				case EBackgroundNSURLCDNInfoResponse::Ok:
+					StatusString = TEXT("OK");
+					break;
+				case EBackgroundNSURLCDNInfoResponse::Timeout:
+					StatusString = TEXT("Timeout");
+					break;
+				case EBackgroundNSURLCDNInfoResponse::Error:
+				default:
+					StatusString = TEXT("Error");
+					break;
+			}
+			Result.Emplace(TTuple<FString, uint32, FString>(Info.CDNAbsoluteURL, (uint32)(Info.ResponseTime * MSEC_PER_SEC), StatusString));
 		}
 	}
 
diff --git a/Engine/Source/Runtime/ApplicationCore/Public/IOS/IOSBackgroundURLSessionHandler.h b/Engine/Source/Runtime/ApplicationCore/Public/IOS/IOSBackgroundURLSessionHandler.h
--- a/Engine/Source/Runtime/ApplicationCore/Public/IOS/IOSBackgroundURLSessionHandler.h
+++ b/Engine/Source/Runtime/ApplicationCore/Public/IOS/IOSBackgroundURLSessionHandler.h
@@ -119,10 +119,10 @@ public:
 	// To be used by ApplePlatformBackgroundHttpManager
 	static void SaveBackgroundHttpFileHashHelperState();
 
-	// Returns an ordered list of CDNs used to issue actual downloads.
+	// Returns an ordered list of CDNs used to issue actual downloads, including AbosluteURL, ResponseTime (in ms) and Status
 	// A list of URLS provided to CreateOrFindDownload might change order if CDNReorderingTimeout > 0 to ensure better success rate.
 	// List is empty before first CreateOrFindDownload call.
-	static TArray<FString> GetCDNOrderArray();
+	static TArray<TTuple<FString, uint32, FString>> GetCDNAnalyticsData();
 
 #if !UE_BUILD_SHIPPING
 	static void GetDownloadDebugText(const uint64 DownloadId, TArray<FString>& Output);
