Adding physical screen size calculation methods for Android and iOS.

diff --git a/Engine/Source/Runtime/ApplicationCore/Private/Android/AndroidPlatformApplicationMisc.cpp b/Engine/Source/Runtime/ApplicationCore/Private/Android/AndroidPlatformApplicationMisc.cpp
--- a/Engine/Source/Runtime/ApplicationCore/Private/Android/AndroidPlatformApplicationMisc.cpp
+++ b/Engine/Source/Runtime/ApplicationCore/Private/Android/AndroidPlatformApplicationMisc.cpp
@@ -225,6 +225,46 @@ EScreenPhysicalAccuracy FAndroidApplicationMisc::ComputePhysicalScreenDensity(in
 		}
 	}
 
+#if USE_ANDROID_JNI
+	const EScreenPhysicalAccuracy Accuracy  = GetNativePhysicalScreenDensity(OutScreenDensity);
+	OutScreenDensity = (int32)((float)OutScreenDensity * GetWindowUpscaleFactor());
+	return Accuracy;
+#else
+	return EScreenPhysicalAccuracy::Unknown;
+#endif
+}
+
+EScreenPhysicalAccuracy FAndroidApplicationMisc::ComputePhysicalScreenDimensions(FVector2f& OutScreenDimensions, float& OutScreenSize)
+{
+	FIntVector2 NativeScreenResolution;
+	GetNativePhysicalScreenResolution(NativeScreenResolution);
+
+	int32 ScreenDensity;
+	const EScreenPhysicalAccuracy Accuracy = GetNativePhysicalScreenDensity(ScreenDensity);
+
+	if (ScreenDensity <= 0 || Accuracy == EScreenPhysicalAccuracy::Unknown)
+	{
+		UE_LOG(LogAndroid, Warning, TEXT("ComputePhysicalScreenDimensions : Incorrect ScreenDensity or Unknown ScreenPhysicalAccuracy. ScreenDensity: %d, Accuracy: %d"),
+			ScreenDensity,
+			Accuracy);
+
+		return Accuracy;
+	}
+
+	OutScreenDimensions = FVector2f(NativeScreenResolution) / (float)ScreenDensity;
+	OutScreenSize = OutScreenDimensions.Size();
+
+	UE_LOG(LogAndroid, Display, TEXT("ComputePhysicalScreenDimensions : NativeScreenResolution: %s, ScreenDensity: %d, OutScreenDimensions (inches): %s, OutScreenSize (inches): %f"),
+		*NativeScreenResolution.ToString(),
+		ScreenDensity,
+		*OutScreenDimensions.ToString(),
+		OutScreenSize);
+
+	return Accuracy;
+}
+
+EScreenPhysicalAccuracy FAndroidApplicationMisc::GetNativePhysicalScreenDensity(int32& OutScreenDensity)
+{
 #if USE_ANDROID_JNI
 	extern FString AndroidThunkCpp_GetMetaDataString(const FString& Key);
 	FString DPIStrings = AndroidThunkCpp_GetMetaDataString(TEXT("unreal.displaymetrics.dpi"));
@@ -242,9 +282,26 @@ EScreenPhysicalAccuracy FAndroidApplicationMisc::ComputePhysicalScreenDensity(in
 		return EScreenPhysicalAccuracy::Unknown;
 	}
 
-	OutScreenDensity = (int32)((float)OutScreenDensity * GetWindowUpscaleFactor());
 	return EScreenPhysicalAccuracy::Approximation;
 #else
 	return EScreenPhysicalAccuracy::Unknown;
 #endif
 }
+
+void FAndroidApplicationMisc::GetNativePhysicalScreenResolution(FIntVector2& OutScreenResolution)
+{
+#if USE_ANDROID_JNI
+	extern FString AndroidThunkCpp_GetMetaDataString(const FString& Key);
+	FString DisplayString = AndroidThunkCpp_GetMetaDataString(FString(TEXT("unreal.display")));
+	TArray<FString> DisplayValues;
+	DisplayString.ParseIntoArray(DisplayValues, TEXT(";"));
+
+	FString WidthString = DisplayValues[DisplayValues.IndexOfByKey(TEXT("realWidth")) + 1];
+	FString HeightString = DisplayValues[DisplayValues.IndexOfByKey(TEXT("realHeight")) + 1];
+
+	int Width = FCString::Atoi(*WidthString);
+	int Height = FCString::Atoi(*HeightString);
+
+	OutScreenResolution = FIntVector2(Width, Height);
+#endif
+}
diff --git a/Engine/Source/Runtime/ApplicationCore/Private/GenericPlatform/GenericPlatformApplicationMisc.cpp b/Engine/Source/Runtime/ApplicationCore/Private/GenericPlatform/GenericPlatformApplicationMisc.cpp
--- a/Engine/Source/Runtime/ApplicationCore/Private/GenericPlatform/GenericPlatformApplicationMisc.cpp
+++ b/Engine/Source/Runtime/ApplicationCore/Private/GenericPlatform/GenericPlatformApplicationMisc.cpp
@@ -18,6 +18,10 @@ bool FGenericPlatformApplicationMisc::CachedPhysicalScreenData = false;
 EScreenPhysicalAccuracy FGenericPlatformApplicationMisc::CachedPhysicalScreenAccuracy = EScreenPhysicalAccuracy::Unknown;
 int32 FGenericPlatformApplicationMisc::CachedPhysicalScreenDensity = 0;
 
+EScreenPhysicalAccuracy FGenericPlatformApplicationMisc::CachedPhysicalScreenDimensionsAccuracy = EScreenPhysicalAccuracy::Unknown;
+FVector2f FGenericPlatformApplicationMisc::CachedPhysicalScreenDimensions = FVector2f::ZeroVector;
+float FGenericPlatformApplicationMisc::CachedPhysicalScreenSize = 0.f;
+
 static int32 bEnableHighDPIAwareness = 1;
 
 FAutoConsoleVariableRef FGenericPlatformApplicationMisc::CVarEnableHighDPIAwareness(
@@ -115,27 +119,63 @@ void FGenericPlatformApplicationMisc:: ClipboardPaste(class FString& Dest)
 	Dest = FString();
 }
 
+void FGenericPlatformApplicationMisc::ComputePhysicalScreenData()
+{
+	CachedPhysicalScreenAccuracy = FPlatformApplicationMisc::ComputePhysicalScreenDensity(CachedPhysicalScreenDensity);
+	if (CachedPhysicalScreenAccuracy == EScreenPhysicalAccuracy::Unknown)
+	{
+		// If the screen density is unknown we use 96, which is the default on Windows,
+		// but it's also the assumed DPI scale of Slate internally for fonts.
+		CachedPhysicalScreenDensity = 96;
+	}
+
+	CachedPhysicalScreenDimensionsAccuracy = FPlatformApplicationMisc::ComputePhysicalScreenDimensions(CachedPhysicalScreenDimensions, CachedPhysicalScreenSize);
+
+	CachedPhysicalScreenData = true;
+}
+
+EScreenPhysicalAccuracy FGenericPlatformApplicationMisc::ComputePhysicalScreenDensity(int32& ScreenDensity)
+{
+	ScreenDensity = 0;
+	return EScreenPhysicalAccuracy::Unknown;
+}
+
+EScreenPhysicalAccuracy FGenericPlatformApplicationMisc::ComputePhysicalScreenDimensions(FVector2f& OutScreenDimensions, float& OutScreenSize)
+{
+	OutScreenDimensions = FVector2f::ZeroVector;
+	OutScreenSize = 0.f;
+	return EScreenPhysicalAccuracy::Unknown;
+}
+
 EScreenPhysicalAccuracy FGenericPlatformApplicationMisc::GetPhysicalScreenDensity(int32& ScreenDensity)
 {
 	if (!CachedPhysicalScreenData)
 	{
-		CachedPhysicalScreenData = true;
-		CachedPhysicalScreenAccuracy = FPlatformApplicationMisc::ComputePhysicalScreenDensity(CachedPhysicalScreenDensity);
-		if (CachedPhysicalScreenAccuracy == EScreenPhysicalAccuracy::Unknown)
-		{
-			// If the screen density is unknown we use 96, which is the default on Windows,
-			// but it's also the assumed DPI scale of Slate internally for fonts.
-			CachedPhysicalScreenDensity = 96;
-		}
+		ComputePhysicalScreenData();
 	}
 
 	ScreenDensity = CachedPhysicalScreenDensity;
 	return CachedPhysicalScreenAccuracy;
 }
 
-EScreenPhysicalAccuracy FGenericPlatformApplicationMisc::ComputePhysicalScreenDensity(int32& ScreenDensity)
+EScreenPhysicalAccuracy FGenericPlatformApplicationMisc::GetPhysicalScreenDimensions(FVector2f& OutScreenDimensions)
 {
-	ScreenDensity = 0;
-	return EScreenPhysicalAccuracy::Unknown;
+	if (!CachedPhysicalScreenData)
+	{
+		ComputePhysicalScreenData();
+	}
+
+	OutScreenDimensions = CachedPhysicalScreenDimensions;
+	return CachedPhysicalScreenDimensionsAccuracy;
 }
 
+EScreenPhysicalAccuracy FGenericPlatformApplicationMisc::GetPhysicalScreenSize(float& OutScreenSize)
+{
+	if (!CachedPhysicalScreenData)
+	{
+		ComputePhysicalScreenData();
+	}
+
+	OutScreenSize = CachedPhysicalScreenSize;
+	return CachedPhysicalScreenDimensionsAccuracy;
+}
diff --git a/Engine/Source/Runtime/ApplicationCore/Private/IOS/IOSPlatformApplicationMisc.cpp b/Engine/Source/Runtime/ApplicationCore/Private/IOS/IOSPlatformApplicationMisc.cpp
--- a/Engine/Source/Runtime/ApplicationCore/Private/IOS/IOSPlatformApplicationMisc.cpp
+++ b/Engine/Source/Runtime/ApplicationCore/Private/IOS/IOSPlatformApplicationMisc.cpp
@@ -315,6 +315,43 @@ EScreenPhysicalAccuracy FIOSPlatformApplicationMisc::ComputePhysicalScreenDensit
 	return EScreenPhysicalAccuracy::Approximation;
 }
 
+EScreenPhysicalAccuracy FIOSPlatformApplicationMisc::ComputePhysicalScreenDimensions(FVector2f& OutScreenDimensions, float& OutScreenSize)
+{
+	EScreenPhysicalAccuracy Accuracy = EScreenPhysicalAccuracy::Unknown;
+	int32 ScreenDensity = CVarPhysicalScreenDensity.GetValueOnAnyThread();
+	if (ScreenDensity > 0)
+	{
+		Accuracy = EScreenPhysicalAccuracy::Truth;
+	}
+	else
+	{
+		UE_LOG(LogIOS, Warning, TEXT("ComputePhysicalScreenDimensions : Approximating ScreenDensity due to CVarPhysicalScreenDensity not set or unknown (%d)."), ScreenDensity);
+
+		const double NativeScale = [[UIScreen mainScreen] scale];
+		ScreenDensity = FMath::TruncToInt(163 * NativeScale);
+		if (ScreenDensity <= 0)
+		{
+			UE_LOG(LogIOS, Warning, TEXT("ComputePhysicalScreenDimensions : Incorrect ScreenDensity resulting in Unknown ScreenPhysicalAccuracy. ScreenDensity: %d"), ScreenDensity);
+			return EScreenPhysicalAccuracy::Unknown;
+		}
+
+		Accuracy = EScreenPhysicalAccuracy::Approximation;
+	}
+
+	CGRect NativeBounds = [[UIScreen mainScreen] nativeBounds];
+	const FVector2f NativeResolution = FVector2f((float)NativeBounds.size.width, (float)NativeBounds.size.height);
+	OutScreenDimensions = NativeResolution / (float)ScreenDensity;
+	OutScreenSize = OutScreenDimensions.Size();
+
+	UE_LOG(LogIOS, Display, TEXT("ComputePhysicalScreenDimensions - NativeResolution: %s, ScreenDensity: %d, OutScreenDimensions (inches): %s, OutScreenSize (inches): %f"),
+		*NativeResolution.ToString(),
+		ScreenDensity,
+		*OutScreenDimensions.ToString(),
+		OutScreenSize);
+
+	return Accuracy;
+}
+
 bool FIOSPlatformApplicationMisc::RequiresVirtualKeyboard()
 {
 #if !PLATFORM_TVOS
diff --git a/Engine/Source/Runtime/ApplicationCore/Public/Android/AndroidPlatformApplicationMisc.h b/Engine/Source/Runtime/ApplicationCore/Public/Android/AndroidPlatformApplicationMisc.h
--- a/Engine/Source/Runtime/ApplicationCore/Public/Android/AndroidPlatformApplicationMisc.h
+++ b/Engine/Source/Runtime/ApplicationCore/Public/Android/AndroidPlatformApplicationMisc.h
@@ -22,7 +22,11 @@ struct FAndroidApplicationMisc : public FGenericPlatformApplicationMisc
 	static APPLICATIONCORE_API void ClipboardCopy(const TCHAR* Str);
 	static APPLICATIONCORE_API void ClipboardPaste(class FString& Dest);
 	static APPLICATIONCORE_API EScreenPhysicalAccuracy ComputePhysicalScreenDensity(int32& OutScreenDensity);
+	static APPLICATIONCORE_API EScreenPhysicalAccuracy ComputePhysicalScreenDimensions(FVector2f& OutScreenDimensions, float& OutScreenSize);
 	static APPLICATIONCORE_API void EnableMotionData(bool bEnable);
+
+	static EScreenPhysicalAccuracy GetNativePhysicalScreenDensity(int32& OutScreenDensity);
+	static void GetNativePhysicalScreenResolution(FIntVector2& OutScreenResolution);
 };
 
 typedef FAndroidApplicationMisc FPlatformApplicationMisc;
diff --git a/Engine/Source/Runtime/ApplicationCore/Public/GenericPlatform/GenericPlatformApplicationMisc.h b/Engine/Source/Runtime/ApplicationCore/Public/GenericPlatform/GenericPlatformApplicationMisc.h
--- a/Engine/Source/Runtime/ApplicationCore/Public/GenericPlatform/GenericPlatformApplicationMisc.h
+++ b/Engine/Source/Runtime/ApplicationCore/Public/GenericPlatform/GenericPlatformApplicationMisc.h
@@ -302,11 +302,19 @@ struct FGenericPlatformApplicationMisc
 	 */
 	static APPLICATIONCORE_API EScreenPhysicalAccuracy GetPhysicalScreenDensity(int32& OutScreenDensity);
 
+	static APPLICATIONCORE_API EScreenPhysicalAccuracy GetPhysicalScreenDimensions(FVector2f& OutScreenDimensions);
+
+	static APPLICATIONCORE_API EScreenPhysicalAccuracy GetPhysicalScreenSize(float& OutScreenSize);
+
 	/**
 	 * Gets the physical size of the screen if possible.  Some platforms lie, some platforms don't know.
 	 */
 	static APPLICATIONCORE_API EScreenPhysicalAccuracy ComputePhysicalScreenDensity(int32& OutScreenDensity);
 
+	static APPLICATIONCORE_API EScreenPhysicalAccuracy ComputePhysicalScreenDimensions(FVector2f& OutScreenDimensions, float& OutScreenSize);
+
+	static void ComputePhysicalScreenData();
+
 	/**
 	 * If we know or can approximate the pixel density of the screen we will convert the incoming inches
 	 * to pixels on the device.  If the accuracy is unknown OutPixels will be set to 0.
@@ -388,4 +396,8 @@ protected:
 	static APPLICATIONCORE_API FAutoConsoleVariableRef CVarEnableHighDPIAwareness;
 	static APPLICATIONCORE_API FAutoConsoleVariableRef CVarAllowVirtualKeyboard;
 
+	static APPLICATIONCORE_API EScreenPhysicalAccuracy CachedPhysicalScreenDimensionsAccuracy;
+	static APPLICATIONCORE_API FVector2f CachedPhysicalScreenDimensions;
+	static APPLICATIONCORE_API float CachedPhysicalScreenSize;
+
 };
diff --git a/Engine/Source/Runtime/ApplicationCore/Public/IOS/IOSPlatformApplicationMisc.h b/Engine/Source/Runtime/ApplicationCore/Public/IOS/IOSPlatformApplicationMisc.h
--- a/Engine/Source/Runtime/ApplicationCore/Public/IOS/IOSPlatformApplicationMisc.h
+++ b/Engine/Source/Runtime/ApplicationCore/Public/IOS/IOSPlatformApplicationMisc.h
@@ -31,6 +31,7 @@ struct APPLICATIONCORE_API FIOSPlatformApplicationMisc : public FGenericPlatform
 	static void ClipboardPaste(class FString& Dest);
 
 	static EScreenPhysicalAccuracy ComputePhysicalScreenDensity(int32& ScreenDensity);
+	static EScreenPhysicalAccuracy ComputePhysicalScreenDimensions(FVector2f& OutScreenDimensions, float& OutScreenSize);
 
     static bool RequiresVirtualKeyboard();
     
