DirectInstall support added to BPS source

#redirectsaresafe

#FailedValidations Preflight

diff --git a/Engine/Source/Runtime/ApplicationCore/Private/IOS/IOSBackgroundURLSessionHandler.cpp b/Engine/Source/Runtime/ApplicationCore/Private/IOS/IOSBackgroundURLSessionHandler.cpp
--- a/Engine/Source/Runtime/ApplicationCore/Private/IOS/IOSBackgroundURLSessionHandler.cpp
+++ b/Engine/Source/Runtime/ApplicationCore/Private/IOS/IOSBackgroundURLSessionHandler.cpp
@@ -166,16 +166,20 @@ NSString* const SerializationKeyExpectedResultSizeFromHeaderRequest = @"h";
 	NSMutableDictionary* Dict = [NSMutableDictionary dictionaryWithCapacity:4];
 
 	NSString* Path = [[self.URLs firstObject] path];
+	NSString* Query = [[self.URLs firstObject] query];
+	if (Query && Query.length > 0)
+	{
+		Path = [[Path stringByAppendingString:@"?"] stringByAppendingString:Query];
+	}
 
+	// Decompose URL's that we assume to always have the same path+query into just a path(with query attached if present),
+	// and all the different CDN's. This reduces the total data stored in the session.
 	NSMutableArray<__kindof NSString*>* CDNs = [NSMutableArray arrayWithCapacity:self.URLs.count];
 	for (NSURL* URL in self.URLs)
 	{
 		NSString* URLString = URL.absoluteString;
-		if (![URLString hasSuffix:Path])
-		{
-			UE_DNLD_LOG(@"Expected all URLs have same path but got '%@' when expected path '%@'", URLString, Path);
-			continue;
-		}
+		checkf([URLString hasSuffix:Path], TEXT("Expected all URLs have same path(+query), but have %s vs %s"), UTF8_TO_TCHAR([Path UTF8String]), UTF8_TO_TCHAR([URLString UTF8String]));
+
 		NSString* Domain = [[URLString componentsSeparatedByString:URL.path] firstObject];
 		[CDNs addObject:Domain];
 	}
@@ -897,10 +901,11 @@ static constexpr NSInteger TypicalSizeForDownloadResponseHeader = 400; // bytes
 		if (AssetPath == nil)
 		{
 			AssetPath = URLValue.path;
-		}
-		if (![URLValue.absoluteString hasSuffix:AssetPath])
-		{
-			UE_DNLD_LOG(@"Expected all URLs have same path but got '%@' when expected path '%@'", URLString, AssetPath);
+			NSString* Query = URLValue.query;
+			if (Query != nil)
+			{
+				AssetPath = [[AssetPath stringByAppendingString:@"?"] stringByAppendingString:Query];
+			}
 		}
 		[URLs addObject:URLValue];
 	}
