Fixing visionos build break in IOSPlatformApplicationMisc.cpp ComputePhysicalScreenDimensions
-Maybe a windowed visionos app could have meaningful values here, but for an immersive style app I think the concept is invalid.  For now we are just returning 0/unknown.

diff --git a/Engine/Source/Runtime/ApplicationCore/Private/IOS/IOSPlatformApplicationMisc.cpp b/Engine/Source/Runtime/ApplicationCore/Private/IOS/IOSPlatformApplicationMisc.cpp
--- a/Engine/Source/Runtime/ApplicationCore/Private/IOS/IOSPlatformApplicationMisc.cpp
+++ b/Engine/Source/Runtime/ApplicationCore/Private/IOS/IOSPlatformApplicationMisc.cpp
@@ -317,6 +317,11 @@ EScreenPhysicalAccuracy FIOSPlatformApplicationMisc::ComputePhysicalScreenDensit
 
 EScreenPhysicalAccuracy FIOSPlatformApplicationMisc::ComputePhysicalScreenDimensions(FVector2f& OutScreenDimensions, float& OutScreenSize)
 {
+#if PLATFORM_VISIONOS
+	OutScreenDimensions = FVector2f::ZeroVector;
+	OutScreenSize = 0.f;
+	return EScreenPhysicalAccuracy::Unknown;
+#else
 	EScreenPhysicalAccuracy Accuracy = EScreenPhysicalAccuracy::Unknown;
 	int32 ScreenDensity = CVarPhysicalScreenDensity.GetValueOnAnyThread();
 	if (ScreenDensity > 0)
@@ -350,6 +355,7 @@ EScreenPhysicalAccuracy FIOSPlatformApplicationMisc::ComputePhysicalScreenDimens
 		OutScreenSize);
 
 	return Accuracy;
+#endif
 }
 
 bool FIOSPlatformApplicationMisc::RequiresVirtualKeyboard()
