SDL3 reversed return value for SDL_SetClipboardText.  Now returns true on success.

#rnx

diff --git a/Engine/Source/Runtime/ApplicationCore/Private/Linux/LinuxPlatformApplicationMisc.cpp b/Engine/Source/Runtime/ApplicationCore/Private/Linux/LinuxPlatformApplicationMisc.cpp
--- a/Engine/Source/Runtime/ApplicationCore/Private/Linux/LinuxPlatformApplicationMisc.cpp
+++ b/Engine/Source/Runtime/ApplicationCore/Private/Linux/LinuxPlatformApplicationMisc.cpp
@@ -524,7 +524,7 @@ float FLinuxPlatformApplicationMisc::GetDPIScaleFactorAtPoint(float X, float Y)
 
 void FLinuxPlatformApplicationMisc::ClipboardCopy(const TCHAR* Str)
 {
-	if (SDL_WasInit(SDL_INIT_VIDEO) && SDL_SetClipboardText(TCHAR_TO_UTF8(Str)))
+	if (SDL_WasInit(SDL_INIT_VIDEO) && !SDL_SetClipboardText(TCHAR_TO_UTF8(Str)))
 	{
 		UE_LOG(LogInit, Warning, TEXT("Error copying clipboard contents: %s\n"), UTF8_TO_TCHAR(SDL_GetError()));
 	}
