Fixed the freeze of Android Media Player playback on Xiaomi 13T.

diff --git a/Engine/Source/Runtime/ApplicationCore/Private/Android/AndroidJavaMediaFrameData.cpp b/Engine/Source/Runtime/ApplicationCore/Private/Android/AndroidJavaMediaFrameData.cpp
--- a/Engine/Source/Runtime/ApplicationCore/Private/Android/AndroidJavaMediaFrameData.cpp
+++ b/Engine/Source/Runtime/ApplicationCore/Private/Android/AndroidJavaMediaFrameData.cpp
@@ -90,6 +90,7 @@ jobject FAndroidJavaMediaFrameData::Extract()
 void FAndroidJavaMediaFrameData::CleanUp()
 {
 	Fence = nullptr;
+	Texture = nullptr;
 
 	if (FrameDataGlobalRef != nullptr)
 	{
@@ -102,13 +103,13 @@ void FAndroidJavaMediaFrameData::CleanUp()
 
 bool FAndroidJavaMediaFrameData::ExtractToTextureVulkan(FRHICommandListImmediate& RHICmdList, FTextureRHIRef& InDstTexture, IMediaTextureSample* TextureSample)
 {
-	if (!FrameDataGlobalRef)
+	check(IsInRenderingThread());
+
+	if (GDynamicRHI->RHIIsRenderingSuspended() || FrameDataGlobalRef == nullptr)
 	{
 		return false;
 	}
 
-	check(IsInRenderingThread());
-
 	bool Converted = false;
 
 	jobject HardwareBufferObj = Extract();
@@ -118,8 +119,8 @@ bool FAndroidJavaMediaFrameData::ExtractToTextureVulkan(FRHICommandListImmediate
 		AHardwareBuffer_acquire(HardwareBuffer);
 
 		IVulkanDynamicRHI* RHI = GetIVulkanDynamicRHI();
-		FTextureRHIRef SrcTexture = RHI->RHICreateTexture2DFromAndroidHardwareBuffer(HardwareBuffer);
-		if (SrcTexture.IsValid())
+		Texture = RHI->RHICreateTexture2DFromAndroidHardwareBuffer(HardwareBuffer);
+		if (Texture.IsValid())
 		{
 			FGraphicsPipelineStateInitializer GraphicsPSOInit;
 			FRHITexture* RenderTarget = InDstTexture;
@@ -175,7 +176,7 @@ bool FAndroidJavaMediaFrameData::ExtractToTextureVulkan(FRHICommandListImmediate
 				TShaderMapRef<FVYUConvertPS> ConvertShader(ShaderMap);
 				GraphicsPSOInit.BoundShaderState.PixelShaderRHI = ConvertShader.GetPixelShader();
 				SetGraphicsPipelineState(RHICmdList, GraphicsPSOInit, 0);
-				SetShaderParametersLegacyPS(RHICmdList, ConvertShader, SrcTexture, Dim, YUVMtx, TextureSample->GetEncodingType(), ColorSpaceMtx, 
+				SetShaderParametersLegacyPS(RHICmdList, ConvertShader, Texture, Dim, YUVMtx, TextureSample->GetEncodingType(), ColorSpaceMtx,
 					MediaShaders::EToneMapMethod::None, TextureSample->GetScaleRotation(), TextureSample->GetOffset());
 
 				// draw full size quad into render target
@@ -220,8 +221,8 @@ bool FAndroidJavaMediaFrameData::ExtractToTextureOES(FRHICommandListImmediate& R
 		AHardwareBuffer_acquire(HardwareBuffer);
 
 		IOpenGLDynamicRHI* RHI = GetIOpenGLDynamicRHI();
-		FTextureRHIRef SrcTexture = RHI->RHICreateTexture2DFromAndroidHardwareBuffer(RHICmdList, HardwareBuffer);
-		if (SrcTexture.IsValid())
+		Texture = RHI->RHICreateTexture2DFromAndroidHardwareBuffer(RHICmdList, HardwareBuffer);
+		if (Texture.IsValid())
 		{
 			FGraphicsPipelineStateInitializer GraphicsPSOInit;
 			FRHITexture* RenderTarget = InDstTexture.GetReference();
@@ -254,7 +255,7 @@ bool FAndroidJavaMediaFrameData::ExtractToTextureOES(FRHICommandListImmediate& R
 				TShaderMapRef<FReadTextureExternalPS> CopyShader(ShaderMap);
 				GraphicsPSOInit.BoundShaderState.PixelShaderRHI = CopyShader.GetPixelShader();
 				SetGraphicsPipelineState(RHICmdList, GraphicsPSOInit, 0);
-				SetShaderParametersLegacyPS(RHICmdList, CopyShader, SrcTexture, InSamplerState, TextureSample->GetScaleRotation(), TextureSample->GetOffset());
+				SetShaderParametersLegacyPS(RHICmdList, CopyShader, Texture, InSamplerState, TextureSample->GetScaleRotation(), TextureSample->GetOffset());
 
 				// draw full size quad into render target
 				FBufferRHIRef VertexBuffer = CreateTempMediaVertexBuffer(RHICmdList);
diff --git a/Engine/Source/Runtime/ApplicationCore/Public/Android/AndroidJavaMediaFrameData.h b/Engine/Source/Runtime/ApplicationCore/Public/Android/AndroidJavaMediaFrameData.h
--- a/Engine/Source/Runtime/ApplicationCore/Public/Android/AndroidJavaMediaFrameData.h
+++ b/Engine/Source/Runtime/ApplicationCore/Public/Android/AndroidJavaMediaFrameData.h
@@ -42,6 +42,7 @@ public:
 	bool ExtractToTextureOES(FRHICommandListImmediate& RHICmdList, FTextureRHIRef& InDstTexture, IMediaTextureSample* TextureSample);
 
 private:
+	TRefCountPtr<FRHITexture> Texture;
 	TRefCountPtr<FRHIGPUFence> Fence;
 
 	/** The java frame data from the surface provider*/
