Windows: Validate Visual C++ redistributable version on editor launch, similar to BootstrapPackagedGame verification.

* Mostly applicable to users building from source, since installation from EGS will install redistributables automatically.
* Warn users on launch of the issue given potential crashes, for instance in third-party libraries (i.e. access violations with std::mutex).

diff --git a/Engine/Source/Programs/Windows/BootstrapPackagedGame/BootstrapPackagedGame.Build.cs b/Engine/Source/Programs/Windows/BootstrapPackagedGame/BootstrapPackagedGame.Build.cs
--- a/Engine/Source/Programs/Windows/BootstrapPackagedGame/BootstrapPackagedGame.Build.cs
+++ b/Engine/Source/Programs/Windows/BootstrapPackagedGame/BootstrapPackagedGame.Build.cs
@@ -6,6 +6,10 @@ public class BootstrapPackagedGame : ModuleRules
 {
 	public BootstrapPackagedGame(ReadOnlyTargetRules Target) : base(Target)
 	{
+		bTreatAsEngineModule = true; // For internal headers
+
+		PrivateIncludePathModuleNames.Add("ApplicationCore");
+
 		PublicSystemLibraries.Add("shlwapi.lib");
 		PublicSystemLibraries.Add("version.lib");
 	}
diff --git a/Engine/Source/Programs/Windows/BootstrapPackagedGame/Private/BootstrapPackagedGame.cpp b/Engine/Source/Programs/Windows/BootstrapPackagedGame/Private/BootstrapPackagedGame.cpp
--- a/Engine/Source/Programs/Windows/BootstrapPackagedGame/Private/BootstrapPackagedGame.cpp
+++ b/Engine/Source/Programs/Windows/BootstrapPackagedGame/Private/BootstrapPackagedGame.cpp
@@ -2,33 +2,14 @@
 
 #include "BootstrapPackagedGame.h"
 
+#include "Windows/WindowsRedistributableValidation.h"
+
 #define IDI_EXEC_FILE 201
 #define IDI_EXEC_ARGS 202
 #define IDI_EXEC_FILE_ARM64 203
 #define IDI_EXEC_FILE_ARM64EC 204
 #define IDI_EXEC_FILE_OVERRIDE 205
 
-struct VersionInfo
-{
-	DWORD Major = 0;
-	DWORD Minor = 0;
-	DWORD Bld = 0;
-	DWORD Rbld = 0;
-};
-
-// This minimum should match the version installed by
-// Engine/Extras/Redist/en-us/vc_redist.x64.exe
-static const VersionInfo MinRedistVersion = { 14, 44, 35211, 0 };
-
-bool IsVersionValid(const VersionInfo& Version, const VersionInfo& MinVersion)
-{
-	if (Version.Major > MinVersion.Major) return true;
-	if (Version.Major == MinVersion.Major && Version.Minor > MinVersion.Minor) return true;
-	if (Version.Major == MinVersion.Major && Version.Minor == MinVersion.Minor && Version.Bld > MinVersion.Bld) return true;
-	if (Version.Major == MinVersion.Major && Version.Minor == MinVersion.Minor && Version.Bld == MinVersion.Bld && Version.Rbld >= MinVersion.Rbld) return true;
-	return false;
-}
-
 WCHAR* ReadResourceString(HMODULE ModuleHandle, LPCWSTR Name)
 {
 	WCHAR* Result = NULL;
diff --git a/Engine/Source/Runtime/ApplicationCore/Internal/Windows/WindowsRedistributableValidation.h b/Engine/Source/Runtime/ApplicationCore/Internal/Windows/WindowsRedistributableValidation.h
new file mode 100644
--- /dev/null
+++ b/Engine/Source/Runtime/ApplicationCore/Internal/Windows/WindowsRedistributableValidation.h
@@ -0,0 +1,46 @@
+// Copyright Epic Games, Inc. All Rights Reserved.
+
+#pragma once
+
+/**
+ * Visual C++ redistributable version information.
+ */
+struct VersionInfo
+{
+	constexpr VersionInfo() = default;
+	
+	constexpr VersionInfo(unsigned long InMajor, unsigned long InMinor, unsigned long InBld, unsigned long InRbld)
+		: Major(InMajor)
+		, Minor(InMinor)
+		, Bld(InBld)
+		, Rbld(InRbld)
+	{}
+
+	constexpr VersionInfo(unsigned long long InVersion)
+		: Major((InVersion >> 48) & 0xffff)
+		, Minor((InVersion >> 32) & 0xffff)
+		, Bld((InVersion >> 16) & 0xffff)
+		, Rbld((InVersion >> 0) & 0xffff)
+	{ }
+
+	unsigned long Major = 0;
+	unsigned long Minor = 0;
+	unsigned long Bld = 0;
+	unsigned long Rbld = 0;
+};
+
+// This minimum should match the version installed by
+// Engine/Extras/Redist/en-us/vc_redist.x64.exe
+constexpr VersionInfo MinRedistVersion = { 14, 44, 35211, 0 };
+
+/** Valid if the specified version is greater or equal than the minimum version. */
+inline bool IsVersionValid(const VersionInfo& Version, const VersionInfo& MinVersion)
+{
+	if (Version.Major > MinVersion.Major) return true;
+	if (Version.Major == MinVersion.Major && Version.Minor > MinVersion.Minor) return true;
+	if (Version.Major == MinVersion.Major && Version.Minor == MinVersion.Minor && Version.Bld > MinVersion.Bld) return true;
+	if (Version.Major == MinVersion.Major && Version.Minor == MinVersion.Minor && Version.Bld == MinVersion.Bld && Version.Rbld >= MinVersion.Rbld) return true;
+	return false;
+}
+
+
diff --git a/Engine/Source/Runtime/ApplicationCore/Private/Windows/WindowsPlatformApplicationMisc.cpp b/Engine/Source/Runtime/ApplicationCore/Private/Windows/WindowsPlatformApplicationMisc.cpp
--- a/Engine/Source/Runtime/ApplicationCore/Private/Windows/WindowsPlatformApplicationMisc.cpp
+++ b/Engine/Source/Runtime/ApplicationCore/Private/Windows/WindowsPlatformApplicationMisc.cpp
@@ -18,6 +18,11 @@
 #include "Templates/RefCounting.h"
 #include "Null/NullPlatformApplicationMisc.h"
 
+#if WITH_EDITOR
+#include "Misc/MessageDialog.h"
+#include "Windows/WindowsRedistributableValidation.h"
+#endif
+
 // Resource includes.
 #include "Runtime/Launch/Resources/Windows/resource.h"
 
@@ -29,12 +34,69 @@ THIRD_PARTY_INCLUDES_START
 #include "Windows/HideWindowsPlatformTypes.h"
 THIRD_PARTY_INCLUDES_END
 
+#define LOCTEXT_NAMESPACE "WindowsPlatformApplicationMisc"
+
 typedef HRESULT(STDAPICALLTYPE *GetDpiForMonitorProc)(HMONITOR Monitor, int32 DPIType, uint32 *DPIX, uint32 *DPIY);
 APPLICATIONCORE_API GetDpiForMonitorProc GetDpiForMonitor;
 
 void FWindowsPlatformApplicationMisc::PreInit()
 {
 	FApp::SetHasFocusFunction(&FWindowsPlatformApplicationMisc::IsThisApplicationForeground);
+
+#if WITH_EDITOR
+	auto IsDllVersionValid = [](const TCHAR* ExecDirectory, const TCHAR* Name) -> bool
+		{
+			const FString FilePath = FPaths::Combine(ExecDirectory, Name);
+			if (!FilePath.IsEmpty())
+			{
+				const uint64 FileVersion = FWindowsPlatformMisc::GetFileVersion(FilePath);
+				if (FileVersion != 0)
+				{
+					const VersionInfo CurrentVersionInfo{ FileVersion };
+
+					return IsVersionValid(CurrentVersionInfo, MinRedistVersion);
+				}
+			}
+
+			return false;
+		};
+
+	// Reproducing verifications done in BootstrapPackagedGame.cpp
+	bool bHasValidRedistributable = false;
+
+	const FString SystemRoot = FWindowsPlatformMisc::GetEnvironmentVariable(TEXT("SystemRoot"));
+	const FString System32Path = FPaths::Combine(SystemRoot, TEXT("system32"));
+
+	if (IsDllVersionValid(*System32Path, TEXT("msvcp140_2.dll")) &&
+		IsDllVersionValid(*System32Path, TEXT("vcruntime140_1.dll")))
+	{
+		bHasValidRedistributable = true;
+	}
+	else if (IsDllVersionValid(TEXT(""), TEXT("msvcp140_2.dll")) &&
+		IsDllVersionValid(TEXT(""), TEXT("vcruntime140_1.dll")))
+	{
+		bHasValidRedistributable = true;
+	}
+
+	if (!bHasValidRedistributable)
+	{
+		const FString ArchitectureName = FWindowsPlatformMisc::GetHostArchitecture();
+		const FString VersionName = FString::Printf(TEXT("%lu.%lu.%lu.%lu"),
+			static_cast<unsigned long>(MinRedistVersion.Major),
+			static_cast<unsigned long>(MinRedistVersion.Minor),
+			static_cast<unsigned long>(MinRedistVersion.Bld),
+			static_cast<unsigned long>(MinRedistVersion.Rbld)
+		);
+
+		// Note: Message dialog is automatically disabled in unattended mode, see implementation.
+		FFormatNamedArguments Arguments;
+		Arguments.Add(TEXT("Arch"), FText::FromString(ArchitectureName));
+		Arguments.Add(TEXT("Ver"), FText::FromString(VersionName));
+		FMessageDialog::Open(EAppMsgType::Ok, FText::Format(LOCTEXT("FWindowsPlatformApplicationMiscRedist", "Invalid or outdated Visual C++ redistributable.\n\nPlease install version {Ver} (or above) from Engine\\Extras\\Redist\\en-us\\vc_redist.{Arch}.exe."), Arguments));
+
+		UE_LOG(LogWindows, Error, TEXT("Invalid or outdated Visual C++ redistributable. Please install version %s (or above) from Engine\\Extras\\Redist\\en-us\\vc_redist.%s.exe."), *VersionName, *ArchitectureName);
+	}
+#endif	//WITH_EDITOR
 }
 
 void FWindowsPlatformApplicationMisc::LoadStartupModules()
@@ -461,3 +523,4 @@ void FWindowsPlatformApplicationMisc::ClipboardPaste(class FString& Result)
 }
 
 UE_ENABLE_OPTIMIZATION_SHIP
+#undef LOCTEXT_NAMESPACE
