Windows: Fixed static analysis warning from global object variable declared in the header (caused by previous CL 47042812).

#rnx

diff --git a/Engine/Source/Runtime/ApplicationCore/Internal/Windows/WindowsRedistributableValidation.h b/Engine/Source/Runtime/ApplicationCore/Internal/Windows/WindowsRedistributableValidation.h
--- a/Engine/Source/Runtime/ApplicationCore/Internal/Windows/WindowsRedistributableValidation.h
+++ b/Engine/Source/Runtime/ApplicationCore/Internal/Windows/WindowsRedistributableValidation.h
@@ -31,7 +31,7 @@ struct VersionInfo
 
 // This minimum should match the version installed by
 // Engine/Extras/Redist/en-us/vc_redist.x64.exe
-constexpr VersionInfo MinRedistVersion = { 14, 44, 35211, 0 };
+inline constexpr VersionInfo MinRedistVersion = { 14, 44, 35211, 0 };
 
 /** Valid if the specified version is greater or equal than the minimum version. */
 inline bool IsVersionValid(const VersionInfo& Version, const VersionInfo& MinVersion)
