Windows: Urgent disabling of visual c++ redistributable validation added in CL 47042812 per remaining Horde issues.

#rnx

diff --git a/Engine/Source/Runtime/ApplicationCore/Private/Windows/WindowsPlatformApplicationMisc.cpp b/Engine/Source/Runtime/ApplicationCore/Private/Windows/WindowsPlatformApplicationMisc.cpp
--- a/Engine/Source/Runtime/ApplicationCore/Private/Windows/WindowsPlatformApplicationMisc.cpp
+++ b/Engine/Source/Runtime/ApplicationCore/Private/Windows/WindowsPlatformApplicationMisc.cpp
@@ -42,61 +42,6 @@ APPLICATIONCORE_API GetDpiForMonitorProc GetDpiForMonitor;
 void FWindowsPlatformApplicationMisc::PreInit()
 {
 	FApp::SetHasFocusFunction(&FWindowsPlatformApplicationMisc::IsThisApplicationForeground);
-
-#if WITH_EDITOR
-	auto IsDllVersionValid = [](const TCHAR* ExecDirectory, const TCHAR* Name) -> bool
-		{
-			const FString FilePath = FPaths::Combine(ExecDirectory, Name);
-			if (!FilePath.IsEmpty())
-			{
-				const uint64 FileVersion = FWindowsPlatformMisc::GetFileVersion(FilePath);
-				if (FileVersion != 0)
-				{
-					const VersionInfo CurrentVersionInfo{ FileVersion };
-
-					return IsVersionValid(CurrentVersionInfo, MinRedistVersion);
-				}
-			}
-
-			return false;
-		};
-
-	// Reproducing verifications done in BootstrapPackagedGame.cpp
-	bool bHasValidRedistributable = false;
-
-	const FString SystemRoot = FWindowsPlatformMisc::GetEnvironmentVariable(TEXT("SystemRoot"));
-	const FString System32Path = FPaths::Combine(SystemRoot, TEXT("system32"));
-
-	if (IsDllVersionValid(*System32Path, TEXT("msvcp140_2.dll")) &&
-		IsDllVersionValid(*System32Path, TEXT("vcruntime140_1.dll")))
-	{
-		bHasValidRedistributable = true;
-	}
-	else if (IsDllVersionValid(TEXT(""), TEXT("msvcp140_2.dll")) &&
-		IsDllVersionValid(TEXT(""), TEXT("vcruntime140_1.dll")))
-	{
-		bHasValidRedistributable = true;
-	}
-
-	if (!bHasValidRedistributable)
-	{
-		const FString ArchitectureName = FWindowsPlatformMisc::GetHostArchitecture();
-		const FString VersionName = FString::Printf(TEXT("%lu.%lu.%lu.%lu"),
-			static_cast<unsigned long>(MinRedistVersion.Major),
-			static_cast<unsigned long>(MinRedistVersion.Minor),
-			static_cast<unsigned long>(MinRedistVersion.Bld),
-			static_cast<unsigned long>(MinRedistVersion.Rbld)
-		);
-
-		// Note: Message dialog is automatically disabled in unattended mode, see implementation.
-		FFormatNamedArguments Arguments;
-		Arguments.Add(TEXT("Arch"), FText::FromString(ArchitectureName));
-		Arguments.Add(TEXT("Ver"), FText::FromString(VersionName));
-		FMessageDialog::Open(EAppMsgType::Ok, FText::Format(LOCTEXT("FWindowsPlatformApplicationMiscRedist", "Invalid or outdated Visual C++ redistributable.\n\nPlease install version {Ver} (or above) from Engine\\Extras\\Redist\\en-us\\vc_redist.{Arch}.exe."), Arguments));
-
-		UE_LOG(LogWindows, Error, TEXT("Invalid or outdated Visual C++ redistributable. Please install version %s (or above) from Engine\\Extras\\Redist\\en-us\\vc_redist.%s.exe."), *VersionName, *ArchitectureName);
-	}
-#endif	//WITH_EDITOR
 }
 
 void FWindowsPlatformApplicationMisc::LoadStartupModules()
