Fix a race condition when resizing a window on Mac.

diff --git a/Engine/Source/Runtime/ApplicationCore/Private/Mac/MacWindow.cpp b/Engine/Source/Runtime/ApplicationCore/Private/Mac/MacWindow.cpp
--- a/Engine/Source/Runtime/ApplicationCore/Private/Mac/MacWindow.cpp
+++ b/Engine/Source/Runtime/ApplicationCore/Private/Mac/MacWindow.cpp
@@ -710,7 +710,10 @@ void FMacWindow::ApplySizeAndModeChanges(int32 X, int32 Y, int32 Width, int32 He
 
 	if(WindowHandle)
 	{
-		WindowHandle->bZoomed = WindowHandle.zoomed;
+		MainThreadCall(^{
+			SCOPED_AUTORELEASE_POOL;
+			WindowHandle->bZoomed = WindowHandle.zoomed;
+		});
 	}
 
 	if (FadeReservationToken != 0)
