Unshelved from pending changelist '48481443':

Disable WindowsApplication.UseWorkerThreadForRawInput by default for Engine

[RN] Disabled WindowsApplication.UseWorkerThreadForRawInput by default because of poor interaction with third-party overlay software.
#Synced-CL 48438952

diff --git a/Engine/Config/Windows/BaseWindowsEngine.ini b/Engine/Config/Windows/BaseWindowsEngine.ini
--- a/Engine/Config/Windows/BaseWindowsEngine.ini
+++ b/Engine/Config/Windows/BaseWindowsEngine.ini
@@ -68,4 +68,4 @@ MaxNumTlsSlots=2048
 MaxNumThreadsWithTlsSlots=512
 
 [ConsoleVariables]
-WindowsApplication.UseWorkerThreadForRawInput=True
+WindowsApplication.UseWorkerThreadForRawInput=False
diff --git a/Engine/Source/Runtime/ApplicationCore/Private/Windows/WindowsApplication.cpp b/Engine/Source/Runtime/ApplicationCore/Private/Windows/WindowsApplication.cpp
--- a/Engine/Source/Runtime/ApplicationCore/Private/Windows/WindowsApplication.cpp
+++ b/Engine/Source/Runtime/ApplicationCore/Private/Windows/WindowsApplication.cpp
@@ -904,7 +904,7 @@ namespace UE::WindowsApplication::RawInputThread
 		}
 	};
 
-	static bool bUseWorkerThreadForRawInput = true;
+	static bool bUseWorkerThreadForRawInput = false;
 	FAutoConsoleVariableRef	CVarUseWorkerThreadForRawInput(
 		TEXT("WindowsApplication.UseWorkerThreadForRawInput"),
 		bUseWorkerThreadForRawInput,
