Speculative fix for a top crash on mac that we are not able to reproduce.
Undoing my fix for UE-349286 since the crash no longer occur with the original repo steps and we think the fix is causing more race condition then it fixes even if on paper the code should be call on the cocoa thread.

diff --git a/Engine/Source/Runtime/ApplicationCore/Private/Mac/MacWindow.cpp b/Engine/Source/Runtime/ApplicationCore/Private/Mac/MacWindow.cpp
--- a/Engine/Source/Runtime/ApplicationCore/Private/Mac/MacWindow.cpp
+++ b/Engine/Source/Runtime/ApplicationCore/Private/Mac/MacWindow.cpp
@@ -710,10 +710,10 @@ void FMacWindow::ApplySizeAndModeChanges(int32 X, int32 Y, int32 Width, int32 He
 
 	if(WindowHandle)
 	{
-		MainThreadCall(^{
-			SCOPED_AUTORELEASE_POOL;
-			WindowHandle->bZoomed = WindowHandle.zoomed;
-		});
+		// Should be run on the main cocoa thread but apparently it cause more race condition then it fixes.
+		//MainThreadCall(^{
+		WindowHandle->bZoomed = WindowHandle.zoomed;
+		//});
 	}
 
 	if (FadeReservationToken != 0)
