Build fix for missing include of CoreDelegates.h.

#rnx
#FailedValidations Preflight

diff --git a/Engine/Source/Runtime/ApplicationCore/Public/Mac/MacApplication.h b/Engine/Source/Runtime/ApplicationCore/Public/Mac/MacApplication.h
--- a/Engine/Source/Runtime/ApplicationCore/Public/Mac/MacApplication.h
+++ b/Engine/Source/Runtime/ApplicationCore/Public/Mac/MacApplication.h
@@ -6,6 +6,7 @@
 #include "GenericPlatform/GenericApplication.h"
 #include "Mac/MacWindow.h"
 #include "Mac/MacTextInputMethodSystem.h"
+#include "Misc/CoreDelegates.h"
 #include "GenericPlatform/IInputInterface.h"
 
 class IModularFeature;
