Change how IOS pauses rendering when moved to background. Allows hang detector to be re-enabled on IOS

#rnx

diff --git a/Engine/Source/Runtime/Apple/MetalRHI/Private/MetalPipeline.cpp b/Engine/Source/Runtime/Apple/MetalRHI/Private/MetalPipeline.cpp
--- a/Engine/Source/Runtime/Apple/MetalRHI/Private/MetalPipeline.cpp
+++ b/Engine/Source/Runtime/Apple/MetalRHI/Private/MetalPipeline.cpp
@@ -1329,6 +1329,9 @@ FMetalPipelineStateCacheManager::FMetalPipelineStateCacheManager(FMetalDevice& D
 	FMetalShaderPipelineCache::Create(Device);
 	
 #if PLATFORM_IOS
+	FCoreDelegates::ApplicationWillDeactivateDelegate.AddStatic(&FShaderPipelineCache::PauseBatching);
+	FCoreDelegates::ApplicationHasReactivatedDelegate.AddStatic(&FShaderPipelineCache::ResumeBatching);
+	
 	OnShaderPipelineCachePreOpenDelegate = FShaderPipelineCache::GetCachePreOpenDelegate().AddRaw(this, &FMetalPipelineStateCacheManager::OnShaderPipelineCachePreOpen);
 	OnShaderPipelineCacheOpenedDelegate = FShaderPipelineCache::GetCacheOpenedDelegate().AddRaw(this, &FMetalPipelineStateCacheManager::OnShaderPipelineCacheOpened);
 	OnShaderPipelineCachePrecompilationCompleteDelegate = FShaderPipelineCache::GetPrecompilationCompleteDelegate().AddRaw(this, &FMetalPipelineStateCacheManager::OnShaderPipelineCachePrecompilationComplete);
diff --git a/Engine/Source/Runtime/ApplicationCore/Private/IOS/IOSAppDelegate.cpp b/Engine/Source/Runtime/ApplicationCore/Private/IOS/IOSAppDelegate.cpp
--- a/Engine/Source/Runtime/ApplicationCore/Private/IOS/IOSAppDelegate.cpp
+++ b/Engine/Source/Runtime/ApplicationCore/Private/IOS/IOSAppDelegate.cpp
@@ -41,6 +41,7 @@
 #if WITH_ENGINE
 #include "RHICommandList.h"
 #endif
+#include "UnrealClient.h"
 
 #if WITH_ACCESSIBILITY
 #include "IOS/Accessibility/IOSAccessibilityCache.h"
@@ -1628,7 +1629,6 @@ static FAutoConsoleVariableRef CVarGEnableThermalsReport(
 	return YES;
 }
 
-FCriticalSection RenderSuspend;
 - (void)applicationWillResignActive:(UIApplication *)application
 {
     /*
@@ -1681,6 +1681,12 @@ FCriticalSection RenderSuspend;
 #else
     [self ToggleAudioSession:false];
 #endif
+	
+	// Suspend all future rendering
+	FViewport::SetGameRenderingEnabled(false);
+#if !PLATFORM_TVOS
+	FAppEntry::Suspend(true);
+#endif
     
 #if WITH_ENGINE
 	if (FTaskGraphInterface::IsRunning() && self.bEngineInit)
@@ -1698,26 +1704,15 @@ FCriticalSection RenderSuspend;
 	}
 #endif
 	
-    RenderSuspend.TryLock();
-    if (FTaskGraphInterface::IsRunning())
-    {
-        if (self.bEngineInit)
-        {
-            FGraphEventRef ResignTask = FFunctionGraphTask::CreateAndDispatchWhenReady([]()
-            {
-                FScopeLock ScopeLock(&RenderSuspend);
-            }, TStatId(), NULL, ENamedThreads::GameThread);
-        }
-        else
-        {
-            FGraphEventRef ResignTask = FFunctionGraphTask::CreateAndDispatchWhenReady([]()
-            {
-                FScopeLock ScopeLock(&RenderSuspend);
-            }, TStatId(), NULL, ENamedThreads::ActualRenderingThread);
-        }
-    }
+	__block UIBackgroundTaskIdentifier BackgroundTaskId = UIBackgroundTaskInvalid;	
+	BackgroundTaskId = [[UIApplication sharedApplication] beginBackgroundTaskWithExpirationHandler:^{
+		UE_LOG(LogIOS, Warning, TEXT("Background task expired before completing FIOSCoreDelegates::OnWillResignActive broadcast"));
+		[[UIApplication sharedApplication] endBackgroundTask: BackgroundTaskId];
+	}];
 	
 	FIOSCoreDelegates::OnWillResignActive.Broadcast();
+	
+	[[UIApplication sharedApplication] endBackgroundTask: BackgroundTaskId];
 }
 
 - (void)applicationDidEnterBackground:(UIApplication *)application
@@ -1737,7 +1732,9 @@ FCriticalSection RenderSuspend;
 #endif
 
     self.bAudioActive = false;
-    FAppEntry::Suspend(true);
+#if PLATFORM_TVOS
+	FAppEntry::Suspend(true);
+#endif
 	FBackgroundURLSessionHandler::HandleDidEnterBackground();
 
 	FEmbeddedCommunication::KeepAwake(TEXT("Background"), false);
@@ -1806,7 +1803,10 @@ extern double GCStartTime;
 	/*
 	 Restart any tasks that were paused (or not yet started) while the application was inactive. If the application was previously in the background, optionally refresh the user interface.
 	 */
-	RenderSuspend.Unlock();
+	
+	// Resume Rendering
+	FViewport::SetGameRenderingEnabled(true);
+	
 	[self ToggleSuspend : false];
     [self ToggleAudioSession:true];
 	FAppEntry::RestartAudio();
diff --git a/Engine/Source/Runtime/Core/Private/IOS/IOSPlatformMisc.cpp b/Engine/Source/Runtime/Core/Private/IOS/IOSPlatformMisc.cpp
--- a/Engine/Source/Runtime/Core/Private/IOS/IOSPlatformMisc.cpp
+++ b/Engine/Source/Runtime/Core/Private/IOS/IOSPlatformMisc.cpp
@@ -335,6 +335,25 @@ void FIOSPlatformMisc::PlatformHandleSplashScreen(bool ShowSplashScreen)
     }
 }
 
+bool FIOSPlatformMisc::AllowThreadHeartBeat()
+{
+	static uint32 AllowThreadHeartBeatOnce = -1;
+	if (AllowThreadHeartBeatOnce == -1)
+	{
+		const FString* AllowThreadHeartBeatConfigVar = GetConfigRulesVariable(TEXT("EnableThreadHeartBeat"));
+		if (AllowThreadHeartBeatConfigVar)
+		{
+			AllowThreadHeartBeatOnce = (uint32)AllowThreadHeartBeatConfigVar->Equals("true", ESearchCase::IgnoreCase);
+		}
+		else
+		{
+			AllowThreadHeartBeatOnce = 0;
+		}
+	}
+
+	return AllowThreadHeartBeatOnce == 1;
+}
+
 const TCHAR* FIOSPlatformMisc::GamePersistentDownloadDir()
 {
     static FString GamePersistentDownloadDir;
diff --git a/Engine/Source/Runtime/Core/Public/IOS/IOSPlatformMisc.h b/Engine/Source/Runtime/Core/Public/IOS/IOSPlatformMisc.h
--- a/Engine/Source/Runtime/Core/Public/IOS/IOSPlatformMisc.h
+++ b/Engine/Source/Runtime/Core/Public/IOS/IOSPlatformMisc.h
@@ -32,10 +32,7 @@ struct CORE_API FIOSPlatformMisc : public FApplePlatformMisc
 		return IOS_MAX_PATH;
 	}
 
-	static bool AllowThreadHeartBeat()
-	{
-		return false;
-	}
+	static bool AllowThreadHeartBeat();
 
 	static EAppReturnType::Type MessageBoxExt( EAppMsgType::Type MsgType, const TCHAR* Text, const TCHAR* Caption );
 	static void SetMemoryWarningHandler(void (* Handler)(const FGenericMemoryWarningContext& Context));
diff --git a/Engine/Source/Runtime/Launch/Private/IOS/LaunchIOS.cpp b/Engine/Source/Runtime/Launch/Private/IOS/LaunchIOS.cpp
--- a/Engine/Source/Runtime/Launch/Private/IOS/LaunchIOS.cpp
+++ b/Engine/Source/Runtime/Launch/Private/IOS/LaunchIOS.cpp
@@ -152,6 +152,11 @@ void FAppEntry::Suspend(bool bIsInterrupt)
             IncrementAudioSuspendCounters();
 		}
 	}
+	
+	FGraphEventRef FlushRenderThreadTask = FFunctionGraphTask::CreateAndDispatchWhenReady([]()
+	{
+		FlushRenderingCommands();
+	}, TStatId(), NULL, ENamedThreads::GameThread);
 }
 
 void FAppEntry::Resume(bool bIsInterrupt)
