Fix compile error in WebTests

#rnx
#FailedValidations Preflight

diff --git a/Engine/Source/Runtime/ApplicationCore/Private/IOS/IOSAppDelegate.cpp b/Engine/Source/Runtime/ApplicationCore/Private/IOS/IOSAppDelegate.cpp
--- a/Engine/Source/Runtime/ApplicationCore/Private/IOS/IOSAppDelegate.cpp
+++ b/Engine/Source/Runtime/ApplicationCore/Private/IOS/IOSAppDelegate.cpp
@@ -40,8 +40,8 @@
 #include "HAL/IConsoleManager.h"
 #if WITH_ENGINE
 #include "RHICommandList.h"
-#endif
 #include "UnrealClient.h"
+#endif
 
 #if WITH_ACCESSIBILITY
 #include "IOS/Accessibility/IOSAccessibilityCache.h"
@@ -1682,8 +1682,11 @@ static FAutoConsoleVariableRef CVarGEnableThermalsReport(
     [self ToggleAudioSession:false];
 #endif
 	
+#if WITH_ENGINE
 	// Suspend all future rendering
 	FViewport::SetGameRenderingEnabled(false);
+#endif
+	
 #if !PLATFORM_TVOS
 	FAppEntry::Suspend(true);
 #endif
@@ -1804,8 +1807,10 @@ extern double GCStartTime;
 	 Restart any tasks that were paused (or not yet started) while the application was inactive. If the application was previously in the background, optionally refresh the user interface.
 	 */
 	
+#if WITH_ENGINE
 	// Resume Rendering
 	FViewport::SetGameRenderingEnabled(true);
+#endif
 	
 	[self ToggleSuspend : false];
     [self ToggleAudioSession:true];
