AllowVirtualKeyboard can be toggled at runtime, such as when it is detected that we are using a handheld device etc.

#rnx

diff --git a/Engine/Source/Runtime/ApplicationCore/Private/GenericPlatform/GenericPlatformApplicationMisc.cpp b/Engine/Source/Runtime/ApplicationCore/Private/GenericPlatform/GenericPlatformApplicationMisc.cpp
--- a/Engine/Source/Runtime/ApplicationCore/Private/GenericPlatform/GenericPlatformApplicationMisc.cpp
+++ b/Engine/Source/Runtime/ApplicationCore/Private/GenericPlatform/GenericPlatformApplicationMisc.cpp
@@ -37,7 +37,7 @@ FAutoConsoleVariableRef FGenericPlatformApplicationMisc::CVarAllowVirtualKeyboar
 	TEXT("AllowVirtualKeyboard"),
 	bAllowVirtualKeyboard,
 	TEXT("Allow the use of a virtual keyboard despite platform main screen being non-touch"),
-	ECVF_ReadOnly
+	ECVF_Default
 );
 
 void FGenericPlatformApplicationMisc::PreInit()
