Windows: Fix failed robomerge at CL 47044783.

#rnx

diff --git a/Engine/Source/Runtime/ApplicationCore/Internal/Windows/WindowsRedistributableValidation.h b/Engine/Source/Runtime/ApplicationCore/Internal/Windows/WindowsRedistributableValidation.h
--- a/Engine/Source/Runtime/ApplicationCore/Internal/Windows/WindowsRedistributableValidation.h
+++ b/Engine/Source/Runtime/ApplicationCore/Internal/Windows/WindowsRedistributableValidation.h
@@ -33,6 +33,8 @@ struct VersionInfo
 // Engine/Extras/Redist/en-us/vc_redist.x64.exe
 inline constexpr VersionInfo MinRedistVersion = { 14, 44, 35211, 0 };
 
+inline constexpr VersionInfo MinGameInputRedistVersion = { 3, 0, 26100, 6154 };
+
 /** Valid if the specified version is greater or equal than the minimum version. */
 inline bool IsVersionValid(const VersionInfo& Version, const VersionInfo& MinVersion)
 {
