[Linux] Fix inactive windows ignoring mouse input.

#rn Inactive windows no longer ignore first mouse click on Linux.

diff --git a/Engine/Source/Runtime/ApplicationCore/Private/Linux/LinuxPlatformApplicationMisc.cpp b/Engine/Source/Runtime/ApplicationCore/Private/Linux/LinuxPlatformApplicationMisc.cpp
--- a/Engine/Source/Runtime/ApplicationCore/Private/Linux/LinuxPlatformApplicationMisc.cpp
+++ b/Engine/Source/Runtime/ApplicationCore/Private/Linux/LinuxPlatformApplicationMisc.cpp
@@ -300,7 +300,7 @@ bool FLinuxPlatformApplicationMisc::InitSDL()
 		// In debug mode, let's log all the SDL warnings + errors, etc.
 #if UE_BUILD_DEBUG		
 		SDL_SetLogPriorities(SDL_LOG_PRIORITY_VERBOSE);
-#endif		
+#endif
 
 
 
@@ -312,6 +312,9 @@ bool FLinuxPlatformApplicationMisc::InitSDL()
 			SDL_SetHint("SDL_HINT_EGL_DEVICE", TCHAR_TO_UTF8(*EglDeviceHint));
 		}
 
+		// Initially-unfocused windows should not ignore clicks.
+		SDL_SetHint(SDL_HINT_MOUSE_FOCUS_CLICKTHROUGH, "1");
+
 		// The following hints are needed when FLinuxApplication::SetHighPrecisionMouseMode is called and Enable = true.
 		// SDL_SetRelativeMouseMode when enabled is warping the mouse in default mode but we don't want that. 
 		// Furthermore SDL hides the mouse which we prevent by setting SDL_HINT_MOUSE_RELATIVE_CURSOR_VISIBLE
