Ensure touch input is enabled when 'Allow Bluetooth controller' support is disabled on Android. This is due to the fact that the ToActivate -> ToValidate state transition for input devices only occurs when gamepad support is enabled. Skip that state for non gamepad devices.

diff --git a/Engine/Source/Runtime/ApplicationCore/Private/Android/AndroidInputInterface.cpp b/Engine/Source/Runtime/ApplicationCore/Private/Android/AndroidInputInterface.cpp
--- a/Engine/Source/Runtime/ApplicationCore/Private/Android/AndroidInputInterface.cpp
+++ b/Engine/Source/Runtime/ApplicationCore/Private/Android/AndroidInputInterface.cpp
@@ -1670,13 +1670,14 @@ void FAndroidInputInterface::AddPendingInputDevice(int32 DeviceId, InputDeviceTy
 		return; // a changed device will return here, nothing needs to be updated at the native side.
 	}
 
+	bool IsGameController = DeviceType == InputDeviceType::GameController;
 	// newly connected device is saved for now and will be validated in the next frame
 	FAndroidInputDeviceInfo& AddedDeviceInfo = InputDeviceInfoMap.Add(DeviceId);
-	AddedDeviceInfo.DeviceState = AndroidEnableInputDeviceListener ? MappingState::ToActivate : MappingState::ToValidate;
+	AddedDeviceInfo.DeviceState = (AndroidEnableInputDeviceListener && IsGameController) ? MappingState::ToActivate : MappingState::ToValidate;
 	AddedDeviceInfo.DeviceType = DeviceType;
 	AddedDeviceInfo.DeviceId = DeviceId;
 
-	if (DeviceType == InputDeviceType::GameController)
+	if (IsGameController)
 	{
 		int32 UnassignedId = -1;
 		for (int32 ControllerId = 0; ControllerId < GameControllerIdMapping.Num(); ++ControllerId)
