- Add IOS record permission check to JoinChannel flow

diff --git a/Engine/Source/Runtime/ApplicationCore/Private/IOS/IOSAppDelegate.cpp b/Engine/Source/Runtime/ApplicationCore/Private/IOS/IOSAppDelegate.cpp
--- a/Engine/Source/Runtime/ApplicationCore/Private/IOS/IOSAppDelegate.cpp
+++ b/Engine/Source/Runtime/ApplicationCore/Private/IOS/IOSAppDelegate.cpp
@@ -29,12 +29,14 @@
 #include "Misc/EmbeddedCommunication.h"
 #include "Containers/Ticker.h"
 #include "UObject/GarbageCollection.h"
+#include "Templates/Function.h"
 #if USE_MUTE_SWITCH_DETECTION
 #include "SharkfoodMuteSwitchDetector.h"
 #include "SharkfoodMuteSwitchDetector.m"
 #endif
 
 #include <AudioToolbox/AudioToolbox.h>
+#include <AVFAudio/AVAudioApplication.h>
 #include <AVFoundation/AVAudioSession.h>
 #import <AVFoundation/AVFoundation.h>
 #include "HAL/IConsoleManager.h"
@@ -828,6 +830,24 @@ static IOSAppDelegate* CachedDelegate = nil;
 #endif
 }
 
+-(void)RequestRecordPermission:(TFunction<void(bool)>)OnComplete
+{
+#if PLATFORM_TVOS || PLATFORM_VISIONOS
+	// TVOS does not have sound recording capabilities.
+	return;
+#elif (defined(__IPHONE_17_0) && __IPHONE_OS_VERSION_MAX_ALLOWED >= __IPHONE_17_0)
+	[AVAudioApplication requestRecordPermissionWithCompletionHandler:^(BOOL bGranted)
+	{
+		OnComplete(bGranted);
+	}];
+#else
+	[[AVAudioSession sharedInstance] requestRecordPermission:^(BOOL bGranted)
+	{
+		OnComplete(bGranted);
+	}];
+#endif
+}
+
 -(void)EnableHighQualityVoiceChat:(bool)bEnable
 {
 	[self SetFeature:EAudioFeature::VoiceChat Active:bEnable];
diff --git a/Engine/Source/Runtime/ApplicationCore/Public/IOS/IOSAppDelegate.h b/Engine/Source/Runtime/ApplicationCore/Public/IOS/IOSAppDelegate.h
--- a/Engine/Source/Runtime/ApplicationCore/Public/IOS/IOSAppDelegate.h
+++ b/Engine/Source/Runtime/ApplicationCore/Public/IOS/IOSAppDelegate.h
@@ -282,6 +282,7 @@ APPLICATIONCORE_API
 - (void)ToggleAudioSession:(bool)bActive;
 - (bool)IsBackgroundAudioPlaying;
 - (bool)HasRecordPermission;
+- (void)RequestRecordPermission:(TFunction<void(bool)>)OnComplete;
 - (void)EnableVoiceChat:(bool)bEnable;
 - (void)EnableHighQualityVoiceChat:(bool)bEnable;
 - (bool)IsVoiceChatEnabled;
