[Linux] Configure SDL to not simulate clicks from touches.

The default value of this setting results (on my desktop) in two synthetic mouse-down events preceding each touch event, which causes bad state which makes the touch input not work.  Unreal already treats touches as clicks (ProcessTouchStartedEvent calls ProcessMouseButtonDownEvent), so this default SDL behavior is redundant anyways.

#rn Fixed touch input not working in some Linux environments.

diff --git a/Engine/Source/Runtime/ApplicationCore/Private/Linux/LinuxPlatformApplicationMisc.cpp b/Engine/Source/Runtime/ApplicationCore/Private/Linux/LinuxPlatformApplicationMisc.cpp
--- a/Engine/Source/Runtime/ApplicationCore/Private/Linux/LinuxPlatformApplicationMisc.cpp
+++ b/Engine/Source/Runtime/ApplicationCore/Private/Linux/LinuxPlatformApplicationMisc.cpp
@@ -315,6 +315,10 @@ bool FLinuxPlatformApplicationMisc::InitSDL()
 		// Initially-unfocused windows should not ignore clicks.
 		SDL_SetHint(SDL_HINT_MOUSE_FOCUS_CLICKTHROUGH, "1");
 
+		// Unreal has support for touches, so this is unnecessary, and on some X11 implementations causes a mismatched number of mouse down/up events
+		// which confuses Slate, preventing touch-based clicks from registering.
+		SDL_SetHint(SDL_HINT_TOUCH_MOUSE_EVENTS, "0"); 
+
 		// The following hints are needed when FLinuxApplication::SetHighPrecisionMouseMode is called and Enable = true.
 		// SDL_SetRelativeMouseMode when enabled is warping the mouse in default mode but we don't want that. 
 		// Furthermore SDL hides the mouse which we prevent by setting SDL_HINT_MOUSE_RELATIVE_CURSOR_VISIBLE
