[Updated Resubmit of CL 47042812]

Windows: Validate Visual C++ redistributable version on editor launch, similar to BootstrapPackagedGame verification.

* Mostly applicable to users building from source, since installation from EGS will install redistributables automatically.
* Warn users on launch of the issue given potential crashes, for instance in third-party libraries (i.e. access violations with std::mutex).

diff --git a/Engine/Source/Runtime/ApplicationCore/Internal/Windows/WindowsRedistributableValidation.h b/Engine/Source/Runtime/ApplicationCore/Internal/Windows/WindowsRedistributableValidation.h
--- a/Engine/Source/Runtime/ApplicationCore/Internal/Windows/WindowsRedistributableValidation.h
+++ b/Engine/Source/Runtime/ApplicationCore/Internal/Windows/WindowsRedistributableValidation.h
@@ -23,6 +23,12 @@ struct VersionInfo
 		, Rbld((InVersion >> 0) & 0xffff)
 	{ }
 
+	/** Set if any version number is non-zero. */
+	inline bool IsSet() const
+	{
+		return Major || Minor || Bld || Rbld;
+	}
+
 	unsigned long Major = 0;
 	unsigned long Minor = 0;
 	unsigned long Bld = 0;
diff --git a/Engine/Source/Runtime/ApplicationCore/Private/Windows/WindowsPlatformApplicationMisc.cpp b/Engine/Source/Runtime/ApplicationCore/Private/Windows/WindowsPlatformApplicationMisc.cpp
--- a/Engine/Source/Runtime/ApplicationCore/Private/Windows/WindowsPlatformApplicationMisc.cpp
+++ b/Engine/Source/Runtime/ApplicationCore/Private/Windows/WindowsPlatformApplicationMisc.cpp
@@ -42,6 +42,89 @@ APPLICATIONCORE_API GetDpiForMonitorProc GetDpiForMonitor;
 void FWindowsPlatformApplicationMisc::PreInit()
 {
 	FApp::SetHasFocusFunction(&FWindowsPlatformApplicationMisc::IsThisApplicationForeground);
+
+#if WITH_EDITOR
+	// Note: We currently exclude Wine from visual c++ redistributable validation
+	if (!FWindowsPlatformMisc::IsWine())
+	{
+		auto IsDllVersionValid = [](const TCHAR* ExecDirectory, const TCHAR* Name, VersionInfo& OutVersionInfo) -> bool
+			{
+				const FString FilePath = FPaths::Combine(ExecDirectory, Name);
+				if (!FilePath.IsEmpty())
+				{
+					const uint64 FileVersion = FWindowsPlatformMisc::GetFileVersion(FilePath);
+					if (FileVersion != 0)
+					{
+						OutVersionInfo = VersionInfo{ FileVersion };
+
+						return IsVersionValid(OutVersionInfo, MinRedistVersion);
+					}
+				}
+
+				return false;
+			};
+
+		// Reproducing verifications done in BootstrapPackagedGame.cpp
+		bool bHasValidRedistributable = false;
+		VersionInfo CurrentVersionInfo = {};
+
+		const FString SystemRoot = FWindowsPlatformMisc::GetEnvironmentVariable(TEXT("SystemRoot"));
+		const FString System32Path = FPaths::Combine(SystemRoot, TEXT("system32"));
+
+		if (IsDllVersionValid(*System32Path, TEXT("msvcp140_2.dll"), CurrentVersionInfo) &&
+			IsDllVersionValid(*System32Path, TEXT("vcruntime140_1.dll"), CurrentVersionInfo))
+		{
+			bHasValidRedistributable = true;
+		}
+		else if (IsDllVersionValid(TEXT(""), TEXT("msvcp140_2.dll"), CurrentVersionInfo) &&
+			IsDllVersionValid(TEXT(""), TEXT("vcruntime140_1.dll"), CurrentVersionInfo))
+		{
+			bHasValidRedistributable = true;
+		}
+
+		if (!bHasValidRedistributable)
+		{
+			FFormatNamedArguments Arguments;
+
+			if (CurrentVersionInfo.IsSet())
+			{
+				FFormatNamedArguments SubArguments;
+				SubArguments.Add(TEXT("Current"), FText::AsCultureInvariant(
+					FString::Printf(TEXT("%lu.%lu.%lu.%lu"),
+						static_cast<unsigned long>(CurrentVersionInfo.Major),
+						static_cast<unsigned long>(CurrentVersionInfo.Minor),
+						static_cast<unsigned long>(CurrentVersionInfo.Bld),
+						static_cast<unsigned long>(CurrentVersionInfo.Rbld)
+					)
+				));
+
+				Arguments.Add(TEXT("Warning"), FText::Format(LOCTEXT("FWindowsPlatformApplicationMisc_Outdated", "Visual C++ redistributable version {Current} is outdated."), SubArguments));
+			}
+			else
+			{
+				Arguments.Add(TEXT("Warning"), LOCTEXT("FWindowsPlatformApplicationMisc_Invalid", "Invalid Visual C++ redistributable."));
+			}
+
+			Arguments.Add(TEXT("Version"), FText::AsCultureInvariant(
+				FString::Printf(TEXT("%lu.%lu.%lu.%lu"),
+					static_cast<unsigned long>(MinRedistVersion.Major),
+					static_cast<unsigned long>(MinRedistVersion.Minor),
+					static_cast<unsigned long>(MinRedistVersion.Bld),
+					static_cast<unsigned long>(MinRedistVersion.Rbld)
+				)
+			));
+			Arguments.Add(TEXT("Arch"), FText::AsCultureInvariant(FWindowsPlatformMisc::GetHostArchitecture()));
+
+			const FText MessageText = FText::Format(LOCTEXT("FWindowsPlatformApplicationMisc_RedistValidation", "{Warning}\n\nPlease install version {Version} (or above) from Engine\\Extras\\Redist\\en-us\\vc_redist.{Arch}.exe."), Arguments);
+			const FText MessageTitle(LOCTEXT("FWindowsPlatformApplicationMisc_Title", "Warning"));
+
+			// Note: Message dialog is automatically disabled in unattended mode, see implementation.
+			FMessageDialog::Open(EAppMsgType::Ok, MessageText, MessageTitle);
+
+			UE_LOG(LogWindows, Error, TEXT("%s"), *MessageText.ToString());
+		}
+	}
+#endif	//WITH_EDITOR
 }
 
 void FWindowsPlatformApplicationMisc::LoadStartupModules()
