Support UE_WINDOWS_FORCE_NONEXCLUSIVE_FULLSCREEN to disable exclusive fullscreen for Windows builds

diff --git a/Engine/Source/Runtime/ApplicationCore/Private/Windows/WindowsWindow.cpp b/Engine/Source/Runtime/ApplicationCore/Private/Windows/WindowsWindow.cpp
--- a/Engine/Source/Runtime/ApplicationCore/Private/Windows/WindowsWindow.cpp
+++ b/Engine/Source/Runtime/ApplicationCore/Private/Windows/WindowsWindow.cpp
@@ -1121,8 +1121,12 @@ bool FWindowsWindow::IsForegroundWindow() const
 
 bool FWindowsWindow::IsFullscreenSupported() const
 {
+#if defined(UE_WINDOWS_FORCE_NONEXCLUSIVE_FULLSCREEN) && UE_WINDOWS_FORCE_NONEXCLUSIVE_FULLSCREEN
+	return false;
+#else
 	// fullscreen not supported when using remote desktop
 	return !::GetSystemMetrics(SM_REMOTESESSION);
+#endif
 }
 
 void FWindowsWindow::SetText( const TCHAR* const Text )
