iOS BG download uses resume data in RecreateDownload to reduce amount of data dropped during retriggering stale downloads

diff --git a/Engine/Source/Runtime/ApplicationCore/Private/IOS/IOSBackgroundURLSessionHandler.cpp b/Engine/Source/Runtime/ApplicationCore/Private/IOS/IOSBackgroundURLSessionHandler.cpp
--- a/Engine/Source/Runtime/ApplicationCore/Private/IOS/IOSBackgroundURLSessionHandler.cpp
+++ b/Engine/Source/Runtime/ApplicationCore/Private/IOS/IOSBackgroundURLSessionHandler.cpp
@@ -12,6 +12,22 @@
 
 #include "IOS/IOSAppDelegate.h"
 
+static bool GBackgroundHttpUseExistingNonResumableTasks = true;
+static FAutoConsoleVariableRef CVarBackgroundHttpUseExistingNonResumableTasks(
+	TEXT("ios.BackgroundHttp.bUseExistingNonResumableTasks"),
+	GBackgroundHttpUseExistingNonResumableTasks,
+	TEXT("Make getTasksWithCompletionHandler use existing non resumable tasks"),
+	ECVF_Default
+);
+
+static bool GBackgroundHttpUseCancelByProducingResumeData = true;
+static FAutoConsoleVariableRef CVarBackgroundHttpUseCancelByProducingResumeData(
+	TEXT("ios.BackgroundHttp.bUseCancelByProducingResumeData"),
+	GBackgroundHttpUseCancelByProducingResumeData,
+	TEXT("Use cancelByProducingResumeData in RecreateDownload to try reuse existing downloaded data"),
+	ECVF_Default
+);
+
 #endif
 
 #include <future>
@@ -398,6 +414,8 @@ enum class EBackgroundNSURLCDNInfoResponse : uint32 // Beware these constants de
 	int32 _RetryResumeDataLimit;
 	int32 _HeadRequestTimeout;
 	bool _bCDNReorderByPingTime;
+	bool _bUseExistingNonResumableTasks;
+	bool _bUseCancelByProducingResumeData;
 	double _CheckForForegroundStaleDownloadsWithInterval;
 	double _ForegroundStaleDownloadTimeout;
 	std::atomic<bool> _bAnyTaskDidCompleteWithError;
@@ -457,6 +475,8 @@ static constexpr NSInteger TypicalSizeForDownloadResponseHeader = 400; // bytes
 	_RetryResumeDataLimit = 3;
 	_HeadRequestTimeout = 400;
 	_bCDNReorderByPingTime = false;
+	_bUseExistingNonResumableTasks = GBackgroundHttpUseExistingNonResumableTasks; // pass getTasksWithCompletionHandler tasks even if they are not in resumable state
+	_bUseCancelByProducingResumeData = GBackgroundHttpUseCancelByProducingResumeData; // use [task cancelByProducingResumeData...] to restart tasks if true, use [task cancel] if false
 	_CheckForForegroundStaleDownloadsWithInterval = 1.0; // how often to check for stale downloads, <=0.0 to disable
 	_ForegroundStaleDownloadTimeout = 30.0; // If download hasn't received any bytes for this duration, cancel and retry if possible
 
@@ -470,6 +490,8 @@ static constexpr NSInteger TypicalSizeForDownloadResponseHeader = 400; // bytes
 	GConfig->GetInt(TEXT("BackgroundHttp.iOSSettings"), TEXT("RetryResumeDataLimit"), _RetryResumeDataLimit, GEngineIni);
 	GConfig->GetInt(TEXT("BackgroundHttp.iOSSettings"), TEXT("CDNReorderingTimeout"), _HeadRequestTimeout, GEngineIni);
 	GConfig->GetBool(TEXT("BackgroundHttp.iOSSettings"), TEXT("bCDNReorderByPingTime"), _bCDNReorderByPingTime, GEngineIni);
+	GConfig->GetBool(TEXT("BackgroundHttp.iOSSettings"), TEXT("bUseExistingNonResumableTasks"), _bUseExistingNonResumableTasks, GEngineIni);
+	GConfig->GetBool(TEXT("BackgroundHttp.iOSSettings"), TEXT("bUseCancelByProducingResumeData"), _bUseCancelByProducingResumeData, GEngineIni);
 	GConfig->GetDouble(TEXT("BackgroundHttp.iOSSettings"), TEXT("CheckForForegroundStaleDownloadsWithInterval"), _CheckForForegroundStaleDownloadsWithInterval, GEngineIni);
 	GConfig->GetDouble(TEXT("BackgroundHttp.iOSSettings"), TEXT("ForegroundStaleDownloadTimeout"), _ForegroundStaleDownloadTimeout, GEngineIni);
 #endif
@@ -482,7 +504,9 @@ static constexpr NSInteger TypicalSizeForDownloadResponseHeader = 400; // bytes
 	UE_DNLD_LOG(@"TimeoutIntervalForResource=%f", TimeoutIntervalForResource);
 	UE_DNLD_LOG(@"RetryResumeDataLimit=%i", _RetryResumeDataLimit);
 	UE_DNLD_LOG(@"HeadRequestTimeout=%i", _HeadRequestTimeout);
-	UE_DNLD_LOG(@"CDNReorderByPingTime=%u", _bCDNReorderByPingTime ? 1 : 0);
+	UE_DNLD_LOG(@"bCDNReorderByPingTime=%u", _bCDNReorderByPingTime ? 1 : 0);
+	UE_DNLD_LOG(@"bUseExistingNonResumableTasks=%u", _bUseExistingNonResumableTasks ? 1 : 0);
+	UE_DNLD_LOG(@"bUseCancelByProducingResumeData=%u", _bUseCancelByProducingResumeData ? 1 : 0);
 	UE_DNLD_LOG(@"CheckForForegroundStaleDownloadsWithInterval=%f", _CheckForForegroundStaleDownloadsWithInterval);
 	UE_DNLD_LOG(@"ForegroundStaleDownloadTimeout=%f", _ForegroundStaleDownloadTimeout);
 
@@ -543,9 +567,9 @@ static constexpr NSInteger TypicalSizeForDownloadResponseHeader = 400; // bytes
 					(ExistingTask.state == NSURLSessionTaskStateRunning) ||
 					(ExistingTask.state == NSURLSessionTaskStateSuspended);
 
-				if (!bCanRestartTask)
+				if (_bUseExistingNonResumableTasks == false && !bCanRestartTask)
 				{
-					UE_DNLD_LOG(@"Skipping tracking for existing download task with taskIdentifier %lu because it's not in resumable state", ExistingTask.taskIdentifier);
+					UE_DNLD_LOG(@"Skipping tracking for existing download task with taskIdentifier %lu because it's not in resumable state (%lu)", ExistingTask.taskIdentifier, ExistingTask.state);
 					continue;
 				}
 
@@ -557,7 +581,7 @@ static constexpr NSInteger TypicalSizeForDownloadResponseHeader = 400; // bytes
 					[self CancelDownload:DownloadId];
 				}
 
-				UE_DNLD_LOG(@"Existing task '%@'", ExistingTask.taskDescription);
+				UE_DNLD_LOG(@"Existing task '%@' with taskIdentifier %lu with DownloadId %lu and state %lu", ExistingTask.taskDescription, ExistingTask.taskIdentifier, DownloadId, ExistingTask.state);
 			}
 		}
 
@@ -1170,6 +1194,11 @@ static constexpr NSInteger TypicalSizeForDownloadResponseHeader = 400; // bytes
 		{
 			return CompletedBytes.unsignedLongLongValue;
 		}
+		else // if didWriteData was not invoked yet
+		{
+			// TODO maybe we can use this for everything?
+			return Task.countOfBytesReceived;
+		}
 	}
 
 	return 0;
@@ -1183,27 +1212,83 @@ static constexpr NSInteger TypicalSizeForDownloadResponseHeader = 400; // bytes
 	const float OldTaskPriority = OldTask.priority;
 	const NSURLSessionTaskState OldTaskState = OldTask.state;
 
+	FBackgroundNSURLSessionDownloadTaskData* OldTaskData = [FBackgroundNSURLSessionDownloadTaskData TaskDataFromSerializedString:OldTask.taskDescription];
 	FBackgroundNSURLSessionDownloadTaskData* NewTaskData = [FBackgroundNSURLSessionDownloadTaskData TaskDataFromSerializedString:OldTask.taskDescription];
-	if (NewTaskData == nil)
+	if (OldTaskData == nil || NewTaskData == nil)
 	{
 		return;
 	}
-
-	// Cancel old task
-	OldTask = nil;
-	[self CancelDownload:DownloadId];
-
+	
 	if (ResetRetryCount)
 	{
 		[NewTaskData ResetRetryCount:_RetryResumeDataLimit];
 	}
 
-	// Start a new task
-	NSURLSessionDownloadTask* NewTask = [self CreateDownloadForURL:[NewTaskData GetFirstURL] WithPriority:OldTaskPriority WithTaskData:NewTaskData];
-	[self ReplaceTrackedTaskWith:NewTask ForDownloadId:DownloadId];
-	if (OldTaskState != NSURLSessionTaskStateSuspended)
+	// Remove task data from this download task, otherwise didCompleteWithError might retry the request.
+	[OldTaskData Cancel];
+	OldTask.taskDescription = [OldTaskData ToSerializedString];
+
+	// Must retain otherwise EnsureTaskIsNotTracked will free the object.
+	[OldTask retain];
+
+	// We're done with this task.
+	[self EnsureTaskIsNotTracked:OldTask];
+	
+	if (_bUseCancelByProducingResumeData)
 	{
-		[NewTask resume];
+		// Will invoke didCompleteWithError if task is incomplete.
+		[OldTask cancelByProducingResumeData:^(NSData * _Nullable ResumeData)
+		{
+			if (ResumeData != nil)
+			{
+				UE_DNLD_LOG(@"RecreateDownload OldTask cancelByProducingResumeData for DownloadId %lu returned resume data", DownloadId);
+				
+				NSURLSessionDownloadTask* NewTask = [self CreateDownloadForResumeData:ResumeData WithPriority:OldTaskPriority WithTaskData:NewTaskData];
+				[self ReplaceTrackedTaskWith:NewTask ForDownloadId:DownloadId];
+				
+				if (OldTaskState != NSURLSessionTaskStateSuspended)
+				{
+					[NewTask resume];
+				}
+			}
+			else
+			{
+				UE_DNLD_LOG(@"RecreateDownload OldTask cancelByProducingResumeData for DownloadId %lu did not return resume data", DownloadId);
+				
+				// start without resume data, will drop any already downloaded data
+				NSURLSessionDownloadTask* NewTask = [self CreateDownloadForURL:[NewTaskData GetFirstURL] WithPriority:OldTaskPriority WithTaskData:NewTaskData];
+				[self ReplaceTrackedTaskWith:NewTask ForDownloadId:DownloadId];
+				
+				if (OldTaskState != NSURLSessionTaskStateSuspended)
+				{
+					[NewTask resume];
+				}
+			}
+		}];
+
+		// Finally release the object.
+		[OldTask release];
+
+		OldTask = nil;
+	}
+	else
+	{
+		// Will invoke didCompleteWithError if task is incomplete.
+		[OldTask cancel];
+
+		// Finally release the object.
+		[OldTask release];
+
+		OldTask = nil;
+
+		// start without resume data, will drop any already downloaded data
+		NSURLSessionDownloadTask* NewTask = [self CreateDownloadForURL:[NewTaskData GetFirstURL] WithPriority:OldTaskPriority WithTaskData:NewTaskData];
+		[self ReplaceTrackedTaskWith:NewTask ForDownloadId:DownloadId];
+				
+		if (OldTaskState != NSURLSessionTaskStateSuspended)
+		{
+			[NewTask resume];
+		}
 	}
 }
 
@@ -1440,6 +1525,9 @@ static constexpr NSInteger TypicalSizeForDownloadResponseHeader = 400; // bytes
 	@synchronized(_AllDownloads)
 	{
 		[_AllDownloads setObject:NewTask forKey:[NSNumber numberWithUnsignedInteger:DownloadId]];
+
+		// in case if all downloads were drained, we need to restart a timer as a new task was inserted
+		[self StartCheckingForStaleDownloads];
 	}
 }
 
