Support downloading the same URL more than once from multiple request (IDs)

#rnx

diff --git a/Engine/Plugins/Online/Android/AndroidFetchBackgroundDownload/Source/AndroidFetchBackgroundDownload/Private/AndroidPlatformBackgroundHttpRequest.cpp b/Engine/Plugins/Online/Android/AndroidFetchBackgroundDownload/Source/AndroidFetchBackgroundDownload/Private/AndroidPlatformBackgroundHttpRequest.cpp
--- a/Engine/Plugins/Online/Android/AndroidFetchBackgroundDownload/Source/AndroidFetchBackgroundDownload/Private/AndroidPlatformBackgroundHttpRequest.cpp
+++ b/Engine/Plugins/Online/Android/AndroidFetchBackgroundDownload/Source/AndroidFetchBackgroundDownload/Private/AndroidPlatformBackgroundHttpRequest.cpp
@@ -109,7 +109,7 @@ FString FAndroidPlatformBackgroundHttpRequest::GetDestinationLocation() const
 	if (URLList.Num() > 0)
 	{
 		const FString FirstURL = URLList[0];
-		return FBackgroundHttpModule::Get().GetBackgroundHttpManager()->GetTempFileLocationForURL(FirstURL);
+		return FBackgroundHttpModule::Get().GetBackgroundHttpManager()->GetTempFileLocationForURL(FirstURL, GetRequestID());
 	}
 
 	return TEXT("");
diff --git a/Engine/Source/Runtime/ApplicationCore/Private/IOS/IOSBackgroundURLSessionHandler.cpp b/Engine/Source/Runtime/ApplicationCore/Private/IOS/IOSBackgroundURLSessionHandler.cpp
--- a/Engine/Source/Runtime/ApplicationCore/Private/IOS/IOSBackgroundURLSessionHandler.cpp
+++ b/Engine/Source/Runtime/ApplicationCore/Private/IOS/IOSBackgroundURLSessionHandler.cpp
@@ -71,9 +71,10 @@ static void LogIOSBackgroundDownloadMessage(NSString* Message)
 @property (nonatomic, retain) NSMutableArray<__kindof NSNumber*>* RetryCountPerURL;
 @property (nonatomic) uint64 ExpectedResultSize;
 @property (nonatomic) uint64 ExpectedResultSizeFromHeadRequest;
+@property (nonatomic, retain) NSString* RequestID;
 
 // assumes all URL's have same content path but different domain
-+ (FBackgroundNSURLSessionDownloadTaskData* _Nonnull)TaskDataWithURLs:(NSArray<__kindof NSURL*>*)URLs WithRetryCount:(NSInteger)RetryCount WithExpectedResultSize:(uint64)ExpectedResultSize;
++ (FBackgroundNSURLSessionDownloadTaskData* _Nonnull)TaskDataWithURLs:(NSArray<__kindof NSURL*>*)URLs WithRetryCount:(NSInteger)RetryCount WithExpectedResultSize:(uint64)ExpectedResultSize WithRequestID:(NSString*)RequestID;
 + (FBackgroundNSURLSessionDownloadTaskData* _Nullable)TaskDataFromSerializedString:(NSString*)SerializedData;
 
 - (NSString*)ToSerializedString;
@@ -95,8 +96,9 @@ NSString* const SerializationKeyPath = @"p";
 NSString* const SerializationKeyRetryCountPerURL = @"r";
 NSString* const SerializationKeyExpectedResultSize = @"s";
 NSString* const SerializationKeyExpectedResultSizeFromHeaderRequest = @"h";
+NSString* const SerializationKeyRequestID = @"i";
 
-+ (FBackgroundNSURLSessionDownloadTaskData* _Nonnull)TaskDataWithURLs:(NSArray<__kindof NSURL*>*)URLs WithRetryCount:(NSInteger)RetryCount WithExpectedResultSize:(uint64)ExpectedResultSize
++ (FBackgroundNSURLSessionDownloadTaskData* _Nonnull)TaskDataWithURLs:(NSArray<__kindof NSURL*>*)URLs WithRetryCount:(NSInteger)RetryCount WithExpectedResultSize:(uint64)ExpectedResultSize WithRequestID:(NSString*)RequestID
 {
 	FBackgroundNSURLSessionDownloadTaskData* Data = [[FBackgroundNSURLSessionDownloadTaskData alloc] init];
 
@@ -110,6 +112,7 @@ NSString* const SerializationKeyExpectedResultSizeFromHeaderRequest = @"h";
 	NSAssert(Data.URLs.count > 0, @"URLs should be non empty");
 	NSAssert(Data.URLs.count == Data.RetryCountPerURL.count, @"URLs and RetryCountPerURL arrays should have same size");
 	Data.ExpectedResultSize = ExpectedResultSize;
+	Data.RequestID = RequestID;
 
 	return [Data autorelease];
 }
@@ -138,10 +141,11 @@ NSString* const SerializationKeyExpectedResultSizeFromHeaderRequest = @"h";
 	NSArray<__kindof NSNumber*>* RetryCountPerURL = [Dict valueForKey:SerializationKeyRetryCountPerURL];
 	NSNumber* ExpectedResultSize = [Dict valueForKey:SerializationKeyExpectedResultSize]; // can be nil
 	NSNumber* ExpectedResultSizeFromHeadRequest = [Dict valueForKey:SerializationKeyExpectedResultSizeFromHeaderRequest];
+	NSString* RequestID = [Dict valueForKey:SerializationKeyRequestID];
 
-	if (Error != nil || Version == nil || Version.intValue != 1 || CDNs == nil || Path == nil || RetryCountPerURL == nil)
+	if (Error != nil || Version == nil || Version.intValue != 1 || CDNs == nil || Path == nil || RetryCountPerURL == nil || RequestID == nil)
 	{
-		UE_DNLD_LOG(@"Failed to deserialize task state '%@' due to '%@', %u, %u, %u, %u, %u, %u, %u",
+		UE_DNLD_LOG(@"Failed to deserialize task state '%@' due to '%@', %u, %u, %u, %u, %u, %u, %u, %u",
 			SerializedData,
 			Error != nil ? Error.localizedDescription : @"",
 			Version == nil ? 1 : 0,
@@ -150,7 +154,8 @@ NSString* const SerializationKeyExpectedResultSizeFromHeaderRequest = @"h";
 			Path == nil ? 1 : 0,
 			RetryCountPerURL == nil ? 1 : 0,
 			ExpectedResultSize == nil ? 1 : 0,
-			ExpectedResultSizeFromHeadRequest == nil ? 1 : 0
+			ExpectedResultSizeFromHeadRequest == nil ? 1 : 0,
+			RequestID == nil ? 1 : 0
 		);
 		[Data release];
 		return nil;
@@ -166,6 +171,7 @@ NSString* const SerializationKeyExpectedResultSizeFromHeaderRequest = @"h";
 	NSAssert(Data.URLs.count == Data.RetryCountPerURL.count, @"URLs and RetryCountPerURL arrays should have same size");
 	Data.ExpectedResultSize = ExpectedResultSize != nil ? ExpectedResultSize.unsignedLongLongValue : 0;
 	Data.ExpectedResultSizeFromHeadRequest = ExpectedResultSizeFromHeadRequest != nil ? ExpectedResultSizeFromHeadRequest.unsignedLongLongValue : 0;
+	Data.RequestID = RequestID;
 
 	return [Data autorelease];
 }
@@ -174,6 +180,7 @@ NSString* const SerializationKeyExpectedResultSizeFromHeaderRequest = @"h";
 {
 	[_URLs release];
 	[_RetryCountPerURL release];
+	[_RequestID release];
 	[super dealloc];
 }
 
@@ -212,6 +219,7 @@ NSString* const SerializationKeyExpectedResultSizeFromHeaderRequest = @"h";
 	{
 		[Dict setValue:[NSNumber numberWithUnsignedLongLong:self.ExpectedResultSizeFromHeadRequest] forKey:SerializationKeyExpectedResultSizeFromHeaderRequest];
 	}
+	[Dict setValue:self.RequestID forKey:SerializationKeyRequestID];
 
 	NSError* Error = nil;
 	NSData* Data = [NSJSONSerialization dataWithJSONObject:Dict options:NSJSONWritingSortedKeys error:&Error];
@@ -347,13 +355,13 @@ enum class EBackgroundNSURLCDNInfoResponse : uint32 // Beware these constants de
 - (void)SetFileHashHelper:(BackgroundHttpFileHashHelperRef)HelperRef;
 - (BackgroundHttpFileHashHelperRef)GetFileHashHelper;
 - (void)SaveFileHashHelperState;
-- (NSString*)GetTempPathForURL:(NSURL* _Nonnull)URL;
+- (NSString*)GetTempPathForURL:(NSURL* _Nonnull)URL WithRequestID:(NSString*)RequestID;
 - (NSMutableArray<__kindof NSURL*>*)ReorderCDNsByReachability:(NSMutableArray<__kindof NSURL*>*)URLs;
 - (NSArray<__kindof FBackgroundNSURLCDNInfo*>*)GetCDNInfo;
 - (NSURLSessionDownloadTask*)CreateDownloadForURL:(NSURL* _Nonnull)URL WithPriority:(float)Priority WithTaskData:(FBackgroundNSURLSessionDownloadTaskData* _Nonnull)TaskData;
 - (NSURLSessionDownloadTask*)CreateDownloadForResumeData:(NSData* _Nonnull)ResumeData WithPriority:(float)Priority WithTaskData:(FBackgroundNSURLSessionDownloadTaskData* _Nonnull)TaskData;
 
-- (NSUInteger)CreateOrFindDownloadForURLs:(NSArray<__kindof NSString*>*)URLStrings WithPriority:(float)Priority WithExpectedResultSize:(uint64)ExpectedResultSize;
+- (NSUInteger)CreateOrFindDownloadForURLs:(NSArray<__kindof NSString*>*)URLStrings WithPriority:(float)Priority WithExpectedResultSize:(uint64)ExpectedResultSize WithRequestID:(NSString*)RequestID;
 - (void)PauseDownload:(NSUInteger)DownloadId;
 - (void)ResumeDownload:(NSUInteger)DownloadId;
 - (void)CancelDownload:(NSUInteger)DownloadId;
@@ -669,14 +677,14 @@ static constexpr NSInteger TypicalSizeForDownloadResponseHeader = 400; // bytes
 	}
 }
 
-- (NSString*)GetTempPathForURL:(NSURL* _Nonnull)URL
+- (NSString*)GetTempPathForURL:(NSURL* _Nonnull)URL WithRequestID:(NSString*)RequestID
 {
 	@synchronized (self)
 	{
 		BackgroundHttpFileHashHelperRef HelperRef = [self GetFileHashHelper];
 
 		const FString TaskURL(URL.absoluteString);
-		const FString& TempFileName = HelperRef->FindOrAddTempFilenameMappingForURL(TaskURL);
+		const FString& TempFileName = HelperRef->FindOrAddTempFilenameMappingForURL(TaskURL, RequestID);
 		const FString& DestinationPath = HelperRef->GetFullPathOfTempFilename(TempFileName);
 
 		IPlatformFile& PlatformFile = FPlatformFileManager::Get().GetPlatformFile();
@@ -905,7 +913,7 @@ static constexpr NSInteger TypicalSizeForDownloadResponseHeader = 400; // bytes
 	return Task;
 }
 
-- (NSUInteger)CreateOrFindDownloadForURLs:(NSArray<__kindof NSString*>*)URLStrings WithPriority:(float)Priority WithExpectedResultSize:(uint64)ExpectedResultSize
+- (NSUInteger)CreateOrFindDownloadForURLs:(NSArray<__kindof NSString*>*)URLStrings WithPriority:(float)Priority WithExpectedResultSize:(uint64)ExpectedResultSize WithRequestID:(NSString*)RequestID
 {
 	if (_AllDownloadsFuture.wait_for(std::chrono::seconds(0)) != std::future_status::ready)
 	{
@@ -937,7 +945,7 @@ static constexpr NSInteger TypicalSizeForDownloadResponseHeader = 400; // bytes
 	URLs = [self ReorderCDNsByReachability:URLs];
 
 	// Serialize current settings
-	FBackgroundNSURLSessionDownloadTaskData* TaskData = [FBackgroundNSURLSessionDownloadTaskData TaskDataWithURLs:URLs WithRetryCount:_RetryResumeDataLimit WithExpectedResultSize:ExpectedResultSize];
+	FBackgroundNSURLSessionDownloadTaskData* TaskData = [FBackgroundNSURLSessionDownloadTaskData TaskDataWithURLs:URLs WithRetryCount:_RetryResumeDataLimit WithExpectedResultSize:ExpectedResultSize WithRequestID:RequestID];
 
 	// Check for existing download task, could be from previous app session.
 	UE_DNLD_LOG(@"Trying to find existing download for asset path '%@'", AssetPath);
@@ -1735,7 +1743,10 @@ static constexpr NSInteger TypicalSizeForDownloadResponseHeader = 400; // bytes
 	// Should not be needed, but ensure this just in case
 	[self EnsureTaskIsTracked:Task];
 
-	NSString* DestinationPath = [self GetTempPathForURL:Task.originalRequest.URL];
+	FBackgroundNSURLSessionDownloadTaskData* TaskData = [FBackgroundNSURLSessionDownloadTaskData TaskDataFromSerializedString:Task.taskDescription];
+	NSString* RequestID = TaskData != nil ? TaskData.RequestID : @"";
+
+	NSString* DestinationPath = [self GetTempPathForURL:Task.originalRequest.URL WithRequestID:RequestID];
 
 	[self SaveFileHashHelperState];
 
@@ -1747,7 +1758,6 @@ static constexpr NSInteger TypicalSizeForDownloadResponseHeader = 400; // bytes
 
 	// Check file size before attempting to move
 	bool bResultSizeIsCorrect = false;
-	FBackgroundNSURLSessionDownloadTaskData* TaskData = [FBackgroundNSURLSessionDownloadTaskData TaskDataFromSerializedString:Task.taskDescription];
 	if (TaskData != nil && TaskData.ExpectedResultSize > 0)
 	{
 		NSString* LocationPath = Location.path;
@@ -1837,7 +1847,7 @@ void FBackgroundURLSessionHandler::AllowCellular(bool bAllow)
 	}
 }
 
-uint64 FBackgroundURLSessionHandler::CreateOrFindDownload(const TArray<FString>& URLs, const float Priority, BackgroundHttpFileHashHelperRef HelperRef, const uint64 ExpectedResultSize)
+uint64 FBackgroundURLSessionHandler::CreateOrFindDownload(const TArray<FString>& URLs, const float Priority, BackgroundHttpFileHashHelperRef HelperRef, const uint64 ExpectedResultSize, const FString& RequestID)
 {
 	@autoreleasepool
 	{
@@ -1848,7 +1858,7 @@ uint64 FBackgroundURLSessionHandler::CreateOrFindDownload(const TArray<FString>&
 		}
 
 		[[FBackgroundNSURLSession Shared] SetFileHashHelper:HelperRef];
-		return [[FBackgroundNSURLSession Shared] CreateOrFindDownloadForURLs:URLArray WithPriority:Priority WithExpectedResultSize:ExpectedResultSize];
+		return [[FBackgroundNSURLSession Shared] CreateOrFindDownloadForURLs:URLArray WithPriority:Priority WithExpectedResultSize:ExpectedResultSize WithRequestID:RequestID.GetNSString()];
 	}
 }
 
diff --git a/Engine/Source/Runtime/ApplicationCore/Public/IOS/IOSBackgroundURLSessionHandler.h b/Engine/Source/Runtime/ApplicationCore/Public/IOS/IOSBackgroundURLSessionHandler.h
--- a/Engine/Source/Runtime/ApplicationCore/Public/IOS/IOSBackgroundURLSessionHandler.h
+++ b/Engine/Source/Runtime/ApplicationCore/Public/IOS/IOSBackgroundURLSessionHandler.h
@@ -91,7 +91,9 @@ public:
 	// Priority is a value between 0.0 to 1.0, see NSURLSessionTaskPriorityDefault.
 	// HelperRef is optional shared reference to BackgroundHttpFileHashHelperRef.
 	// In case of HandleEventsForBackgroundURLSession this subsystem will create it's own reference.
-	static uint64 CreateOrFindDownload(const TArray<FString>& URLs, const float Priority, BackgroundHttpFileHashHelperRef HelperRef, const uint64 ExpectedResultSize = 0);
+	static uint64 CreateOrFindDownload(const TArray<FString>& URLs, const float Priority, BackgroundHttpFileHashHelperRef HelperRef, const uint64 ExpectedResultSize, const FString& RequestID);
+	UE_DEPRECATED(5.8, "FBackgroundURLSessionHandler::CreateOrFindDownload without RequestID is deprecated.")
+	static uint64 CreateOrFindDownload(const TArray<FString>& URLs, const float Priority, BackgroundHttpFileHashHelperRef HelperRef, const uint64 ExpectedResultSize = 0) { return CreateOrFindDownload(URLs, Priority, HelperRef, ExpectedResultSize, {}); }
 
 	static void PauseDownload(const uint64 DownloadId);
 
diff --git a/Engine/Source/Runtime/Online/BackgroundHTTP/Private/BackgroundHttpManagerImpl.cpp b/Engine/Source/Runtime/Online/BackgroundHTTP/Private/BackgroundHttpManagerImpl.cpp
--- a/Engine/Source/Runtime/Online/BackgroundHTTP/Private/BackgroundHttpManagerImpl.cpp
+++ b/Engine/Source/Runtime/Online/BackgroundHTTP/Private/BackgroundHttpManagerImpl.cpp
@@ -343,7 +343,7 @@ void FBackgroundHttpManagerImpl::CleanUpDataAfterCompletingRequest(const FBackgr
 	const TArray<FString>& URLList = Request->GetURLList();
 	for (const FString& URL : URLList)
 	{
-		OurFileHashHelper->RemoveURLMapping(URL);
+		OurFileHashHelper->RemoveURLMapping(URL, Request->GetRequestID());
 	}
 }
 
@@ -390,7 +390,7 @@ bool FBackgroundHttpManagerImpl::CheckForExistingCompletedDownload(const FBackgr
 	for (const FString& URL : URLList)
 	{
 		FString FoundTempFilename = {};
-		if (GetFileHashHelper()->FindTempFilenameMappingForURL(URL, FoundTempFilename))
+		if (GetFileHashHelper()->FindTempFilenameMappingForURL(URL, Request->GetRequestID(), FoundTempFilename))
 		{
 			const FString& FileDestination = FBackgroundHttpFileHashHelper::GetFullPathOfTempFilename(FoundTempFilename);
 			if (PlatformFile.FileExists(*FileDestination))
@@ -479,11 +479,11 @@ void FBackgroundHttpManagerImpl::ActivatePendingRequests()
 	}
 }
 
-FString FBackgroundHttpManagerImpl::GetTempFileLocationForURL(const FString& URL)
+FString FBackgroundHttpManagerImpl::GetTempFileLocationForURL(const FString& URL, const FString& RequestID)
 {
 	if (ensureAlwaysMsgf(IsInGameThread(), TEXT("Should only call GetTempFileLocationForURL from the GameThread as this is not thread-safe otherwise!")))
 	{
-		const FString& TempLocation = GetFileHashHelper()->FindOrAddTempFilenameMappingForURL(URL);
+		const FString& TempLocation = GetFileHashHelper()->FindOrAddTempFilenameMappingForURL(URL, RequestID);
 		return FBackgroundHttpFileHashHelper::GetFullPathOfTempFilename(TempLocation);
 	}
 	
diff --git a/Engine/Source/Runtime/Online/BackgroundHTTP/Private/GenericPlatform/GenericPlatformBackgroundHttpResponse.cpp b/Engine/Source/Runtime/Online/BackgroundHTTP/Private/GenericPlatform/GenericPlatformBackgroundHttpResponse.cpp
--- a/Engine/Source/Runtime/Online/BackgroundHTTP/Private/GenericPlatform/GenericPlatformBackgroundHttpResponse.cpp
+++ b/Engine/Source/Runtime/Online/BackgroundHTTP/Private/GenericPlatform/GenericPlatformBackgroundHttpResponse.cpp
@@ -34,7 +34,7 @@ FGenericPlatformBackgroundHttpResponse::FGenericPlatformBackgroundHttpResponse(F
 			if (ContentBuffer.Num() > 0)
 			{
 				const FString RequestURL = HttpRequestIn->GetURL();
-				const FString& FileDestination = FBackgroundHttpModule::Get().GetBackgroundHttpManager()->GetTempFileLocationForURL(RequestURL);
+				const FString& FileDestination = FBackgroundHttpModule::Get().GetBackgroundHttpManager()->GetTempFileLocationForURL(RequestURL, {});
 
 				FFileHelper::SaveArrayToFile(ContentBuffer, *FileDestination);
 				TempContentFilePath = FileDestination;
diff --git a/Engine/Source/Runtime/Online/BackgroundHTTP/Private/IOS/ApplePlatformBackgroundHttpManager.cpp b/Engine/Source/Runtime/Online/BackgroundHTTP/Private/IOS/ApplePlatformBackgroundHttpManager.cpp
--- a/Engine/Source/Runtime/Online/BackgroundHTTP/Private/IOS/ApplePlatformBackgroundHttpManager.cpp
+++ b/Engine/Source/Runtime/Online/BackgroundHTTP/Private/IOS/ApplePlatformBackgroundHttpManager.cpp
@@ -46,7 +46,7 @@ void FApplePlatformBackgroundHttpManager::AddRequest(const FBackgroundHttpReques
 		return;
 	}
 
-	const uint64 DownloadId = FBackgroundURLSessionHandler::CreateOrFindDownload(Request->GetURLList(), Request->GetNSURLSessionPriority(), GetFileHashHelper(), Request->GetExpectedResultSize());
+	const uint64 DownloadId = FBackgroundURLSessionHandler::CreateOrFindDownload(Request->GetURLList(), Request->GetNSURLSessionPriority(), GetFileHashHelper(), Request->GetExpectedResultSize(), Request->GetRequestID());
 	Request->SetInternalDownloadId(DownloadId);
 
 	// Should never happen in practice
diff --git a/Engine/Source/Runtime/Online/BackgroundHTTP/Public/BackgroundHttpManagerImpl.h b/Engine/Source/Runtime/Online/BackgroundHTTP/Public/BackgroundHttpManagerImpl.h
--- a/Engine/Source/Runtime/Online/BackgroundHTTP/Public/BackgroundHttpManagerImpl.h
+++ b/Engine/Source/Runtime/Online/BackgroundHTTP/Public/BackgroundHttpManagerImpl.h
@@ -41,7 +41,7 @@ public:
 	BACKGROUNDHTTP_API virtual void SetMaxActiveDownloads(int MaxActiveDownloads) override;
 
 	
-	BACKGROUNDHTTP_API virtual FString GetTempFileLocationForURL(const FString& URL) override;
+	BACKGROUNDHTTP_API virtual FString GetTempFileLocationForURL(const FString& URL, const FString& RequestID) override;
 	
 	BACKGROUNDHTTP_API virtual void CleanUpDataAfterCompletingRequest(const FBackgroundHttpRequestPtr Request) override;
 	virtual void SetCellularPreference(int32 Value) override {}
diff --git a/Engine/Source/Runtime/Online/BackgroundHTTP/Public/Interfaces/IBackgroundHttpManager.h b/Engine/Source/Runtime/Online/BackgroundHTTP/Public/Interfaces/IBackgroundHttpManager.h
--- a/Engine/Source/Runtime/Online/BackgroundHTTP/Public/Interfaces/IBackgroundHttpManager.h
+++ b/Engine/Source/Runtime/Online/BackgroundHTTP/Public/Interfaces/IBackgroundHttpManager.h
@@ -65,10 +65,12 @@ public:
 	virtual void SetMaxActiveDownloads(int MaxActiveDownloads) = 0;
 
 	/**
-	 * Function that returns an FString fullpath where we would expect the given URL's temp file to be located
+	 * Function that returns an FString fullpath where we would expect the given URL's temp part of the given RequestID to be located
 	 *
 	 */
-	virtual FString GetTempFileLocationForURL(const FString& URL) = 0;
+	virtual FString GetTempFileLocationForURL(const FString& URL, const FString& RequestID) = 0;
+	UE_DEPRECATED(5.8, "Please use GetTempFileLocationForURL with added RequestID argument instead")
+	FString GetTempFileLocationForURL(const FString& URL) { return GetTempFileLocationForURL(URL, {}); }
 	
 	/**
 	 * Function that cleans up any persistent data after we have completed a reqeust.
diff --git a/Engine/Source/Runtime/Online/BackgroundHTTPFileHash/Private/BackgroundHttpFileHashHelper.cpp b/Engine/Source/Runtime/Online/BackgroundHTTPFileHash/Private/BackgroundHttpFileHashHelper.cpp
--- a/Engine/Source/Runtime/Online/BackgroundHTTPFileHash/Private/BackgroundHttpFileHashHelper.cpp
+++ b/Engine/Source/Runtime/Online/BackgroundHTTPFileHash/Private/BackgroundHttpFileHashHelper.cpp
@@ -173,11 +173,11 @@ bool FBackgroundHttpFileHashHelper::InitFromStringNoLock(const FString& StringIn
 	return bInitSuccess;
 }
 
-bool FBackgroundHttpFileHashHelper::FindTempFilenameMappingForURL(const FString& URL, FString& OutTempFilename) const
+bool FBackgroundHttpFileHashHelper::FindTempFilenameMappingForURL(const FString& URL, const FString& RequestID, FString& OutTempFilename) const
 {
 	FReadScopeLock _(URLFileMappingsLock);
 
-	const FURLTempFileMapping* FoundMapping = FindMappingForURLNoLock(URL);
+	const FURLTempFileMapping* FoundMapping = FindMappingForURLNoLock(URL, RequestID);
 	if (nullptr != FoundMapping)
 	{
 		OutTempFilename = FoundMapping->TempFilename;
@@ -191,7 +191,7 @@ const FString* FBackgroundHttpFileHashHelper::FindTempFilenameMappingForURL(cons
 {
 	FReadScopeLock _(URLFileMappingsLock);
 
-	const FURLTempFileMapping* FoundMapping = FindMappingForURLNoLock(URL);
+	const FURLTempFileMapping* FoundMapping = FindMappingForURLNoLock(URL, {});
 	if (nullptr != FoundMapping)
 	{
 		return &FoundMapping->TempFilename;
@@ -200,18 +200,18 @@ const FString* FBackgroundHttpFileHashHelper::FindTempFilenameMappingForURL(cons
 	return nullptr;
 }
 
-const FString& FBackgroundHttpFileHashHelper::FindOrAddTempFilenameMappingForURL(const FString& URL)
+const FString& FBackgroundHttpFileHashHelper::FindOrAddTempFilenameMappingForURL(const FString& URL, const FString& RequestID)
 {
 	FWriteScopeLock _(URLFileMappingsLock);
 
-	const FURLTempFileMapping* FoundMapping = FindMappingForURLNoLock(URL);
+	const FURLTempFileMapping* FoundMapping = FindMappingForURLNoLock(URL, RequestID);
 	if (nullptr != FoundMapping)
 	{
 		return (FoundMapping->TempFilename);
 	}
 
 	// Didn't find a valid old one, so lets add it
-	FString Hash = FindValidFilenameHashForURLNoLock(URL);
+	FString Hash = FindValidFilenameHashForURLNoLock(URL, RequestID);
 	FURLTempFileMapping& NewMapping = URLFileMappings.FindOrAdd(Hash);
 	NewMapping.URL = URL;
 	NewMapping.TempFilename = Hash;
@@ -250,9 +250,9 @@ const FString* FBackgroundHttpFileHashHelper::FindMappedURLForTempFilename(const
 	return (nullptr != FoundMapping) ? &(FoundMapping->URL) : nullptr;
 }
 
-const FURLTempFileMapping* FBackgroundHttpFileHashHelper::FindMappingForURLNoLock(const FString& URL) const
+const FURLTempFileMapping* FBackgroundHttpFileHashHelper::FindMappingForURLNoLock(const FString& URL, const FString& RequestID) const
 {
-	FString Hash = FindValidFilenameHashForURLNoLock(URL);
+	FString Hash = FindValidFilenameHashForURLNoLock(URL, RequestID);
 	const FURLTempFileMapping* FoundMapping = URLFileMappings.Find(Hash);
 
 	// We found an already existing entry for the given Hash, so return that temp file mapping
@@ -268,7 +268,7 @@ const FURLTempFileMapping* FBackgroundHttpFileHashHelper::FindMappingForURLNoLoc
 	return nullptr;
 }
 
-FString FBackgroundHttpFileHashHelper::FindValidFilenameHashForURLNoLock(const FString& URL) const
+FString FBackgroundHttpFileHashHelper::FindValidFilenameHashForURLNoLock(const FString& URL, const FString& RequestID) const
 {
 	FString HashToReturn;
 
@@ -281,7 +281,7 @@ FString FBackgroundHttpFileHashHelper::FindValidFilenameHashForURLNoLock(const F
 	// in case there is a serious logic error here
 	while (CollisionCount < MaxCollisionsToTry)
 	{
-		const FString Hash = GenerateHashedFilenameForURLNoLock(URL, CollisionCount);
+		const FString Hash = GenerateHashedFilenameForURLNoLock(URL, RequestID, CollisionCount);
 		const FURLTempFileMapping* FoundMapping = URLFileMappings.Find(Hash);
 		const bool bFoundValidMapping = (nullptr != FoundMapping);
 
@@ -291,7 +291,7 @@ FString FBackgroundHttpFileHashHelper::FindValidFilenameHashForURLNoLock(const F
 			// Existing entry was ours so use it
 			if (FoundMapping->URL.Equals(URL, ESearchCase::IgnoreCase))
 			{
-				UE_LOG(LogBackgroundHttpFileHash, VeryVerbose, TEXT("Found valid existing URL Temp File Mapping for %s at %s"), *URL, *Hash);
+				UE_LOG(LogBackgroundHttpFileHash, VeryVerbose, TEXT("Found valid existing URL Temp File Mapping for %s and %s at %s"), *URL, *RequestID, *Hash);
 				
 				HashToReturn = Hash;
 				break;
@@ -299,13 +299,13 @@ FString FBackgroundHttpFileHashHelper::FindValidFilenameHashForURLNoLock(const F
 			// The existing entry was not ours, we have a collision and need to move on
 			else
 			{
-				UE_LOG(LogBackgroundHttpFileHash, Log, TEXT("Collision for URL Mapping for %s with %s with hash %s"), *URL, *(FoundMapping->URL), *Hash);
+				UE_LOG(LogBackgroundHttpFileHash, Log, TEXT("Collision for URL Mapping for %s and %s with %s with hash %s"), *URL, *RequestID, *(FoundMapping->URL), *Hash);
 			}
 		}
 		// Data not in our mapping, so its definitely safe to use for this URL
 		else
 		{
-			UE_LOG(LogBackgroundHttpFileHash, Verbose, TEXT("Found valid unused URL Temp File Mapping for %s at %s"), *URL, *Hash);
+			UE_LOG(LogBackgroundHttpFileHash, Verbose, TEXT("Found valid unused URL Temp File Mapping for %s and %s at %s"), *URL, *RequestID, *Hash);
 			HashToReturn = Hash;
 
 			// No need to keep searching as we found a valid result
@@ -318,8 +318,8 @@ FString FBackgroundHttpFileHashHelper::FindValidFilenameHashForURLNoLock(const F
 	// If we didn't manage to find a valid hash above, fallback to our sanitized URL
 	if (HashToReturn.IsEmpty())
 	{
-		const FString TempFilename = FPaths::MakeValidFileName(URL);
-		UE_LOG(LogBackgroundHttpFileHash, Error, TEXT("Unable to find valid temp filename after max collisions for %s! Falling back to using sanitized URL %s. This shouldn't ever be required!"), *URL, *TempFilename);
+		const FString TempFilename = FPaths::MakeValidFileName(URL + RequestID);
+		UE_LOG(LogBackgroundHttpFileHash, Error, TEXT("Unable to find valid temp filename after max collisions for %s and %s! Falling back to using sanitized URL %s. This shouldn't ever be required!"), *URL, *RequestID, *TempFilename);
 		
 		HashToReturn = TempFilename;
 	}
@@ -327,12 +327,12 @@ FString FBackgroundHttpFileHashHelper::FindValidFilenameHashForURLNoLock(const F
 	return MoveTemp(HashToReturn);
 }
 
-FString FBackgroundHttpFileHashHelper::GenerateHashedFilenameForURLNoLock(const FString& URL, uint32 CollisionCount)
+FString FBackgroundHttpFileHashHelper::GenerateHashedFilenameForURLNoLock(const FString& URL, const FString& RequestID, uint32 CollisionCount)
 {
-	const TArray<TCHAR, FString::AllocatorType>& UnderlyingCharData = URL.GetCharArray();
-	const uint32 HashedURL = CityHash32(reinterpret_cast<const char*>(UnderlyingCharData.GetData()), (UnderlyingCharData.Num() * sizeof(TCHAR)));
-
-	return FPaths::MakeValidFileName(FString::Printf(TEXT("%d%d"), HashedURL, CollisionCount));
+	const FString CompleteString = FString::Printf(TEXT("%s|%s"), *URL, *RequestID);
+	const TArray<TCHAR, FString::AllocatorType>& UnderlyingCharData = CompleteString.GetCharArray();
+	const uint32 Hash = CityHash32(reinterpret_cast<const char*>(UnderlyingCharData.GetData()), (UnderlyingCharData.Num() * sizeof(TCHAR)));
+	return FPaths::MakeValidFileName(FString::Printf(TEXT("%d%d"), Hash, CollisionCount));
 }
 
 void FBackgroundHttpFileHashHelper::DeleteURLMappingsWithoutTempFiles()
@@ -358,11 +358,11 @@ void FBackgroundHttpFileHashHelper::DeleteURLMappingsWithoutTempFiles()
 	bIsDirty = true;
 }
 
-void FBackgroundHttpFileHashHelper::RemoveURLMapping(const FString& URL)
+void FBackgroundHttpFileHashHelper::RemoveURLMapping(const FString& URL, const FString& RequestID)
 {
 	FWriteScopeLock _(URLFileMappingsLock);
 
-	const FString Hash = FindValidFilenameHashForURLNoLock(URL);
+	const FString Hash = FindValidFilenameHashForURLNoLock(URL, RequestID);
 	URLFileMappings.Remove(Hash);
 	
 	bIsDirty = true;
diff --git a/Engine/Source/Runtime/Online/BackgroundHTTPFileHash/Public/BackgroundHttpFileHashHelper.h b/Engine/Source/Runtime/Online/BackgroundHTTPFileHash/Public/BackgroundHttpFileHashHelper.h
--- a/Engine/Source/Runtime/Online/BackgroundHTTPFileHash/Public/BackgroundHttpFileHashHelper.h
+++ b/Engine/Source/Runtime/Online/BackgroundHTTPFileHash/Public/BackgroundHttpFileHashHelper.h
@@ -65,17 +65,18 @@ public:
 	// Returns if the parse was successful or not
 	BACKGROUNDHTTPFILEHASH_API bool InitFromString(const FString& StringIn);
 
-	// Looks for a temp filename mapping for the given URL. If one isn't found returns nullptr and does NOT generate one.
-	BACKGROUNDHTTPFILEHASH_API bool FindTempFilenameMappingForURL(const FString& URL, FString& OutTempFilename) const;
-
-	UE_DEPRECATED(5.7, "Please use FindTempFilenameMappingForURL with reference out argument instead")
+	// Looks for a temp filename mapping for the given URL. If one isn't found returns false and does NOT generate one.
+	BACKGROUNDHTTPFILEHASH_API bool FindTempFilenameMappingForURL(const FString& URL, const FString& RequestID, FString& OutTempFilename) const;
+	UE_DEPRECATED(5.8, "Please use FindTempFilenameMappingForURL with RequestID and reference out arguments instead")
 	BACKGROUNDHTTPFILEHASH_API const FString* FindTempFilenameMappingForURL(const FString& URL) const;
 
 	// Looks for a temp filename mapping for the given URL. If one isn't found a new one is generated and returned.
-	BACKGROUNDHTTPFILEHASH_API const FString& FindOrAddTempFilenameMappingForURL(const FString& URL);
+	BACKGROUNDHTTPFILEHASH_API const FString& FindOrAddTempFilenameMappingForURL(const FString& URL, const FString& RequestID);
 
 	// Removes URL mapping
-	BACKGROUNDHTTPFILEHASH_API void RemoveURLMapping(const FString& URL);
+	BACKGROUNDHTTPFILEHASH_API void RemoveURLMapping(const FString& URL, const FString& RequestID);
+	UE_DEPRECATED(5.8, "Please use RemoveURLMapping with RequestID argument instead")
+	BACKGROUNDHTTPFILEHASH_API void RemoveURLMapping(const FString& URL) { RemoveURLMapping(URL, {}); }
 
 	// Looks for a URL mapped to the supplied TempFilename. Nullptr returned if one isn't found
 	BACKGROUNDHTTPFILEHASH_API bool FindMappedURLForTempFilename(const FString& TempFilename, FString& OutMappedURL) const;
@@ -92,7 +93,7 @@ public:
 	// Helper function that returns the file extension used by our BackgroundHTTP temp files
 	static BACKGROUNDHTTPFILEHASH_API const FString& GetTempFileExtension();
 
-	// Helper function that returns a full path representation of hte given TempFilename
+	// Helper function that returns a full path representation of the given TempFilename
 	static BACKGROUNDHTTPFILEHASH_API FString GetFullPathOfTempFilename(const FString& TempFilename);
 
 private:
@@ -106,14 +107,14 @@ private:
 	static BACKGROUNDHTTPFILEHASH_API const FString& GetURLMappingFilePathNoLock();
 
 	// Looks for a temp filename mapping for the given URL. If one isn't found returns nullptr an does NOT generate one.
-	BACKGROUNDHTTPFILEHASH_API const FURLTempFileMapping* FindMappingForURLNoLock(const FString& URL) const;
+	BACKGROUNDHTTPFILEHASH_API const FURLTempFileMapping* FindMappingForURLNoLock(const FString& URL, const FString& RequestID) const;
 
-	// Finds the hash that should be used to refer to this URL. Either finds the hash used in the map to reference the
-	// URLFileMappings entry for this URL or returns the hash we should use (the next valid hash entry) without inserting it into the map.
-	BACKGROUNDHTTPFILEHASH_API FString FindValidFilenameHashForURLNoLock(const FString& URL) const;
+	// Finds the hash that should be used to refer to this URL+RequestID. Either finds the hash used in the map to reference the
+	// URLFileMappings entry for this URL+RequestID or returns the hash we should use (the next valid hash entry) without inserting it into the map.
+	BACKGROUNDHTTPFILEHASH_API FString FindValidFilenameHashForURLNoLock(const FString& URL, const FString& RequestID) const;
 
 	// Creates a hashed Filename with a modifier for the given collision count.
-	static BACKGROUNDHTTPFILEHASH_API FString GenerateHashedFilenameForURLNoLock(const FString& URL, uint32 CollisionCount);
+	static BACKGROUNDHTTPFILEHASH_API FString GenerateHashedFilenameForURLNoLock(const FString& URL, const FString& RequestID, uint32 CollisionCount);
 
 	// Used to make sure that we have called LoadData before any operations that depend on it
 	std::atomic<bool> bHasLoadedURLData;
