Iteration 1: Refactored the game controller button mapping with generalized data accessing and button state mapping. The hardcoded data moved into a new class FAndroidGameControllerMapper. It will finally come from data file generated by the GameControllerMapping.Automation tool.
Added fallback strategy to query mappings by the device name in case the VID/PID are not recorded in the data.

diff --git a/Engine/Source/Runtime/ApplicationCore/Private/Android/AndroidGameControllerMapper.cpp b/Engine/Source/Runtime/ApplicationCore/Private/Android/AndroidGameControllerMapper.cpp
new file mode 100644
--- /dev/null
+++ b/Engine/Source/Runtime/ApplicationCore/Private/Android/AndroidGameControllerMapper.cpp
@@ -0,0 +1,555 @@
+// Copyright Epic Games, Inc. All Rights Reserved.
+
+#include "Android/AndroidGameControllerMapper.h"
+#include "Android/AndroidPlatformMisc.h"
+#include "CoreGlobals.h"
+#include "Misc/ConfigCacheIni.h"
+#include <android/input.h>
+
+extern int32 GAndroidOldXBoxWirelessFirmware;
+
+namespace GameControllerMapping
+{
+	TSharedPtr<FAndroidGameControllerMapper> FAndroidGameControllerMapper::Instance = nullptr;
+
+	TWeakPtr<FAndroidGameControllerMapper> FAndroidGameControllerMapper::GetInstance()
+	{
+		if (Instance == nullptr)
+		{
+			Instance = MakeShareable(new FAndroidGameControllerMapper());
+			bool bBlockAndroidKeysOnControllers = false;
+			GConfig->GetBool(TEXT("/Script/AndroidRuntimeSettings.AndroidRuntimeSettings"), TEXT("bBlockAndroidKeysOnControllers"), bBlockAndroidKeysOnControllers, GEngineIni);
+			if (!Instance->InitByDataFile())
+			{
+				FPlatformMisc::LowLevelOutputDebugStringf(TEXT("FAndroidGameControllerMapper - fallback to hardcoded mapping data."));
+				Instance->Init();
+			}
+		}
+		return Instance.ToWeakPtr();
+	}
+
+	bool FAndroidGameControllerMapper::Init()
+	{
+		ButtonMappings.Reset();
+		DeviceConfigurationMap.Reset();
+		DeviceNames.Reset();
+
+		// init ButtonMappings ...
+
+		// Normal - index 0
+		KeyCodeToKeyNameMap_SRef CodeToNameMap = MakeShared<KeyCodeToKeyNameMap>();
+		AppendKeyCodeKeyNameMapping(CodeToNameMap, AKEYCODE_BUTTON_A, 0);
+		AppendKeyCodeKeyNameMapping(CodeToNameMap, AKEYCODE_DPAD_CENTER, 0);
+		AppendKeyCodeKeyNameMapping(CodeToNameMap, AKEYCODE_BUTTON_B, 1);
+		AppendKeyCodeKeyNameMapping(CodeToNameMap, AKEYCODE_BUTTON_X, 2);
+		AppendKeyCodeKeyNameMapping(CodeToNameMap, AKEYCODE_BUTTON_Y, 3);
+		AppendKeyCodeKeyNameMapping(CodeToNameMap, AKEYCODE_BUTTON_L1, 4);
+		AppendKeyCodeKeyNameMapping(CodeToNameMap, AKEYCODE_BUTTON_R1, 5);
+		AppendKeyCodeKeyNameMapping(CodeToNameMap, AKEYCODE_BUTTON_START, 6);
+		AppendKeyCodeKeyNameMapping(CodeToNameMap, AKEYCODE_MENU, 6);
+		AppendKeyCodeKeyNameMapping(CodeToNameMap, AKEYCODE_BUTTON_START, 17);
+		AppendKeyCodeKeyNameMapping(CodeToNameMap, AKEYCODE_MENU, 17);
+		AppendKeyCodeKeyNameMapping(CodeToNameMap, AKEYCODE_BUTTON_SELECT, 7);
+		AppendKeyCodeKeyNameMapping(CodeToNameMap, AKEYCODE_BACK, 7);
+		AppendKeyCodeKeyNameMapping(CodeToNameMap, AKEYCODE_BUTTON_SELECT, 16);
+		AppendKeyCodeKeyNameMapping(CodeToNameMap, AKEYCODE_BACK, 16);
+		AppendKeyCodeKeyNameMapping(CodeToNameMap, AKEYCODE_BUTTON_THUMBL, 8);
+		AppendKeyCodeKeyNameMapping(CodeToNameMap, AKEYCODE_BUTTON_THUMBR, 9);
+		AppendKeyCodeKeyNameMapping(CodeToNameMap, AKEYCODE_BUTTON_L2, 10);
+		AppendKeyCodeKeyNameMapping(CodeToNameMap, AKEYCODE_BUTTON_R2, 11);
+		AppendKeyCodeKeyNameMapping(CodeToNameMap, AKEYCODE_DPAD_UP, 12);
+		AppendKeyCodeKeyNameMapping(CodeToNameMap, AKEYCODE_DPAD_DOWN, 13);
+		AppendKeyCodeKeyNameMapping(CodeToNameMap, AKEYCODE_DPAD_LEFT, 14);
+		AppendKeyCodeKeyNameMapping(CodeToNameMap, AKEYCODE_DPAD_RIGHT, 15);
+		ButtonMappings.Add(MoveTemp(CodeToNameMap));
+
+		// Normal and MapL1R1ToTriggers - index 1
+		CodeToNameMap = MakeShared<KeyCodeToKeyNameMap>();
+		AppendKeyCodeKeyNameMapping(CodeToNameMap, AKEYCODE_BUTTON_A, 0);
+		AppendKeyCodeKeyNameMapping(CodeToNameMap, AKEYCODE_DPAD_CENTER, 0);
+		AppendKeyCodeKeyNameMapping(CodeToNameMap, AKEYCODE_BUTTON_B, 1);
+		AppendKeyCodeKeyNameMapping(CodeToNameMap, AKEYCODE_BUTTON_X, 2);
+		AppendKeyCodeKeyNameMapping(CodeToNameMap, AKEYCODE_BUTTON_Y, 3);
+		AppendKeyCodeKeyNameMapping(CodeToNameMap, AKEYCODE_BUTTON_L1, 4);
+		AppendKeyCodeKeyNameMapping(CodeToNameMap, AKEYCODE_BUTTON_L1, 10);
+		AppendKeyCodeKeyNameMapping(CodeToNameMap, AKEYCODE_BUTTON_R1, 5);
+		AppendKeyCodeKeyNameMapping(CodeToNameMap, AKEYCODE_BUTTON_R1, 11);
+		AppendKeyCodeKeyNameMapping(CodeToNameMap, AKEYCODE_BUTTON_START, 6);
+		AppendKeyCodeKeyNameMapping(CodeToNameMap, AKEYCODE_MENU, 6);
+		AppendKeyCodeKeyNameMapping(CodeToNameMap, AKEYCODE_BUTTON_START, 17);
+		AppendKeyCodeKeyNameMapping(CodeToNameMap, AKEYCODE_MENU, 17);
+		AppendKeyCodeKeyNameMapping(CodeToNameMap, AKEYCODE_BUTTON_SELECT, 7);
+		AppendKeyCodeKeyNameMapping(CodeToNameMap, AKEYCODE_BACK, 7);
+		AppendKeyCodeKeyNameMapping(CodeToNameMap, AKEYCODE_BUTTON_SELECT, 16);
+		AppendKeyCodeKeyNameMapping(CodeToNameMap, AKEYCODE_BACK, 16);
+		AppendKeyCodeKeyNameMapping(CodeToNameMap, AKEYCODE_BUTTON_THUMBL, 8);
+		AppendKeyCodeKeyNameMapping(CodeToNameMap, AKEYCODE_BUTTON_THUMBR, 9);
+		AppendKeyCodeKeyNameMapping(CodeToNameMap, AKEYCODE_BUTTON_L2, 10);
+		AppendKeyCodeKeyNameMapping(CodeToNameMap, AKEYCODE_BUTTON_R2, 11);
+		AppendKeyCodeKeyNameMapping(CodeToNameMap, AKEYCODE_DPAD_UP, 12);
+		AppendKeyCodeKeyNameMapping(CodeToNameMap, AKEYCODE_DPAD_DOWN, 13);
+		AppendKeyCodeKeyNameMapping(CodeToNameMap, AKEYCODE_DPAD_LEFT, 14);
+		AppendKeyCodeKeyNameMapping(CodeToNameMap, AKEYCODE_DPAD_RIGHT, 15);
+		ButtonMappings.Add(MoveTemp(CodeToNameMap));
+
+		// XBox - index 2
+		CodeToNameMap = MakeShared<KeyCodeToKeyNameMap>();
+		AppendKeyCodeKeyNameMapping(CodeToNameMap, AKEYCODE_BUTTON_A, 0);
+		AppendKeyCodeKeyNameMapping(CodeToNameMap, AKEYCODE_BUTTON_B, 1);
+		AppendKeyCodeKeyNameMapping(CodeToNameMap, AKEYCODE_BUTTON_C, 2);
+		AppendKeyCodeKeyNameMapping(CodeToNameMap, AKEYCODE_BUTTON_X, 3);
+		AppendKeyCodeKeyNameMapping(CodeToNameMap, AKEYCODE_BUTTON_Y, 4);
+		AppendKeyCodeKeyNameMapping(CodeToNameMap, AKEYCODE_BUTTON_Z, 5);
+		AppendKeyCodeKeyNameMapping(CodeToNameMap, AKEYCODE_BUTTON_R1, 6);
+		AppendKeyCodeKeyNameMapping(CodeToNameMap, AKEYCODE_BUTTON_R1, 17);
+		AppendKeyCodeKeyNameMapping(CodeToNameMap, AKEYCODE_BUTTON_L1, 7);
+		AppendKeyCodeKeyNameMapping(CodeToNameMap, AKEYCODE_BUTTON_L1, 16);
+		AppendKeyCodeKeyNameMapping(CodeToNameMap, AKEYCODE_BUTTON_L2, 8);
+		AppendKeyCodeKeyNameMapping(CodeToNameMap, AKEYCODE_BUTTON_R2, 9);
+		ButtonMappings.Add(MoveTemp(CodeToNameMap));
+
+		// PS4 - index 3
+		CodeToNameMap = MakeShared<KeyCodeToKeyNameMap>();
+		AppendKeyCodeKeyNameMapping(CodeToNameMap, AKEYCODE_BUTTON_B, 0);
+		AppendKeyCodeKeyNameMapping(CodeToNameMap, AKEYCODE_BUTTON_C, 1);
+		AppendKeyCodeKeyNameMapping(CodeToNameMap, AKEYCODE_BUTTON_A, 2);
+		AppendKeyCodeKeyNameMapping(CodeToNameMap, AKEYCODE_BUTTON_X, 3);
+		AppendKeyCodeKeyNameMapping(CodeToNameMap, AKEYCODE_BUTTON_Y, 4);
+		AppendKeyCodeKeyNameMapping(CodeToNameMap, AKEYCODE_BUTTON_Z, 5);
+		AppendKeyCodeKeyNameMapping(CodeToNameMap, AKEYCODE_BUTTON_L2, 6);
+		AppendKeyCodeKeyNameMapping(CodeToNameMap, AKEYCODE_BUTTON_L2, 17);
+		AppendKeyCodeKeyNameMapping(CodeToNameMap, AKEYCODE_MENU, 7);
+		AppendKeyCodeKeyNameMapping(CodeToNameMap, AKEYCODE_MENU, 16);
+		AppendKeyCodeKeyNameMapping(CodeToNameMap, AKEYCODE_BUTTON_SELECT, 8);
+		AppendKeyCodeKeyNameMapping(CodeToNameMap, AKEYCODE_BUTTON_START, 9);
+		AppendKeyCodeKeyNameMapping(CodeToNameMap, AKEYCODE_BUTTON_L1, 10);
+		AppendKeyCodeKeyNameMapping(CodeToNameMap, AKEYCODE_BUTTON_R1, 11);
+		ButtonMappings.Add(MoveTemp(CodeToNameMap));
+
+		// PS5 - index 4
+		CodeToNameMap = MakeShared<KeyCodeToKeyNameMap>();
+		AppendKeyCodeKeyNameMapping(CodeToNameMap, AKEYCODE_BUTTON_B, 0);
+		AppendKeyCodeKeyNameMapping(CodeToNameMap, AKEYCODE_BUTTON_C, 1);
+		AppendKeyCodeKeyNameMapping(CodeToNameMap, AKEYCODE_BUTTON_A, 2);
+		AppendKeyCodeKeyNameMapping(CodeToNameMap, AKEYCODE_BUTTON_X, 3);
+		AppendKeyCodeKeyNameMapping(CodeToNameMap, AKEYCODE_BUTTON_Y, 4);
+		AppendKeyCodeKeyNameMapping(CodeToNameMap, AKEYCODE_BUTTON_Z, 5);
+		AppendKeyCodeKeyNameMapping(CodeToNameMap, AKEYCODE_BUTTON_R2, 6);
+		AppendKeyCodeKeyNameMapping(CodeToNameMap, AKEYCODE_BUTTON_R2, 17);
+		AppendKeyCodeKeyNameMapping(CodeToNameMap, 3002, 7); // Special Left
+		AppendKeyCodeKeyNameMapping(CodeToNameMap, AKEYCODE_BUTTON_THUMBL, 7);
+		AppendKeyCodeKeyNameMapping(CodeToNameMap, AKEYCODE_BUTTON_THUMBL, 16);
+		AppendKeyCodeKeyNameMapping(CodeToNameMap, AKEYCODE_BUTTON_SELECT, 8);
+		AppendKeyCodeKeyNameMapping(CodeToNameMap, AKEYCODE_BUTTON_START, 9);
+		AppendKeyCodeKeyNameMapping(CodeToNameMap, AKEYCODE_BUTTON_L1, 10);
+		AppendKeyCodeKeyNameMapping(CodeToNameMap, AKEYCODE_BUTTON_R1, 11);
+		ButtonMappings.Add(MoveTemp(CodeToNameMap));
+
+		// PS5 on Android newer than API 30 - index 5
+		CodeToNameMap = MakeShared<KeyCodeToKeyNameMap>();
+		AppendKeyCodeKeyNameMapping(CodeToNameMap, AKEYCODE_BUTTON_A, 0);
+		AppendKeyCodeKeyNameMapping(CodeToNameMap, AKEYCODE_BUTTON_B, 1);
+		AppendKeyCodeKeyNameMapping(CodeToNameMap, AKEYCODE_BUTTON_X, 2);
+		AppendKeyCodeKeyNameMapping(CodeToNameMap, AKEYCODE_BUTTON_Y, 3);
+		AppendKeyCodeKeyNameMapping(CodeToNameMap, AKEYCODE_BUTTON_L1, 4);
+		AppendKeyCodeKeyNameMapping(CodeToNameMap, AKEYCODE_BUTTON_R1, 5);
+		AppendKeyCodeKeyNameMapping(CodeToNameMap, AKEYCODE_BUTTON_THUMBL, 8);
+		AppendKeyCodeKeyNameMapping(CodeToNameMap, AKEYCODE_BUTTON_THUMBR, 9);
+		AppendKeyCodeKeyNameMapping(CodeToNameMap, AKEYCODE_BUTTON_L2, 10);
+		AppendKeyCodeKeyNameMapping(CodeToNameMap, AKEYCODE_BUTTON_R2, 11);
+		AppendKeyCodeKeyNameMapping(CodeToNameMap, 3002, 7); // Special Left
+		AppendKeyCodeKeyNameMapping(CodeToNameMap, 3002, 16); // Touchpad
+		AppendKeyCodeKeyNameMapping(CodeToNameMap, AKEYCODE_BUTTON_START, 6);
+		AppendKeyCodeKeyNameMapping(CodeToNameMap, AKEYCODE_BUTTON_START, 17);
+		ButtonMappings.Add(MoveTemp(CodeToNameMap));
+
+		// init DeviceConfigurationMap ...
+
+		// Default
+		DeviceConfig_SRef DeviceCfg = MakeShared<DeviceConfig>();
+		DeviceCfg->ButtonMapping = ButtonMappings[0];
+		DeviceCfg->DeviceNameIndex = DeviceNames.Add(TEXT("Default"));
+		RegisterDeviceConfig(-1, -1, DeviceCfg);
+
+		// Amazon Fire Game Controller
+		DeviceCfg = MakeShared<DeviceConfig>();
+		DeviceCfg->ButtonMapping = ButtonMappings[0];
+		DeviceCfg->DeviceNameIndex = DeviceNames.Add(TEXT("Amazon Fire Game Controller"));
+		RegisterDeviceConfig(0x1949, 0x0402, DeviceCfg);
+		RegisterDeviceConfig(0x1949, 0x0406, DeviceCfg);
+		
+		// Amazon Luna Game Controller
+		DeviceCfg = MakeShared<DeviceConfig>();
+		DeviceCfg->ButtonMapping = ButtonMappings[0];
+		DeviceCfg->DeviceNameIndex = DeviceNames.Add(TEXT("Amazon Luna Game Controller"));
+		RegisterDeviceConfig(0x1949, 0x0419, DeviceCfg);
+
+		// Amazon other controllers
+		DeviceCfg = MakeShared<DeviceConfig>();
+		DeviceCfg->SupportsHat = false;
+		DeviceCfg->ButtonMapping = ButtonMappings[0];
+		DeviceCfg->DeviceNameIndex = DeviceNames.Add(TEXT("Amazon"));
+		RegisterDeviceConfig(0x1949, -1, DeviceCfg);
+
+		// NVIDIA Corporation NVIDIA Controller
+		DeviceCfg = MakeShared<DeviceConfig>();
+		DeviceCfg->ButtonMapping = ButtonMappings[0];
+		DeviceCfg->DeviceNameIndex = DeviceNames.Add(TEXT("NVIDIA Corporation NVIDIA Controller"));
+		RegisterDeviceConfig(0x0955, 0x7214, DeviceCfg);
+		RegisterDeviceConfig(0x0955, 0x7210, DeviceCfg);
+		RegisterDeviceConfig(0x0955, 0x7203, DeviceCfg);
+
+		// Samsung Game Pad EI-GP20
+		DeviceCfg = MakeShared<DeviceConfig>();
+		DeviceCfg->SupportsHat = true;
+		DeviceCfg->MapL1R1ToTriggers = true;
+		DeviceCfg->RightStickZRZ = false;
+		DeviceCfg->RightStickRXRY = true;
+		DeviceCfg->ButtonMapping = ButtonMappings[1];
+		DeviceCfg->DeviceNameIndex = DeviceNames.Add(TEXT("Samsung Game Pad EI-GP20"));
+		RegisterDeviceConfig(0x04e8, 0xa000, DeviceCfg);
+
+		// Mad Catz C.T.R.L.R
+		DeviceCfg = MakeShared<DeviceConfig>();
+		DeviceCfg->ButtonMapping = ButtonMappings[0];
+		DeviceCfg->DeviceNameIndex = DeviceNames.Add(TEXT("Mad Catz C.T.R.L.R"));
+		RegisterDeviceConfig(0x0738, 0x5263, DeviceCfg);
+		RegisterDeviceConfig(0x0738, 0x5266, DeviceCfg);
+
+		// Generic X-Box pad - Android API >= 31
+		DeviceCfg = MakeShared<DeviceConfig>();
+		DeviceCfg->ControllerClass = ControllerClassType::XBoxWired;
+		DeviceCfg->SupportsHat = true;
+		DeviceCfg->TriggersUseThresholdForClick = true;
+		DeviceCfg->ButtonMapping = ButtonMappings[0];
+		DeviceCfg->DeviceNameIndex = DeviceNames.Add(TEXT("Generic X-Box pad"));
+		// The VID here is just a placeholder, the query will finally fall back to one with the device name
+		RegisterDeviceConfig(0x045e, 0x0b12, DeviceCfg);
+
+		// Generic X-Box pad - Android API < 31
+		DeviceCfg = MakeShared<DeviceConfig>();
+		DeviceCfg->ControllerClass = ControllerClassType::XBoxWired_AndroidAPIBefore31;
+		DeviceCfg->SupportsHat = true;
+		DeviceCfg->TriggersUseThresholdForClick = true;
+		DeviceCfg->RightStickZRZ = false;
+		DeviceCfg->RightStickRXRY = true;
+		DeviceCfg->MapZRZToTriggers = true;
+		DeviceCfg->LTAnalogRangeMinimum = -1.0f;
+		DeviceCfg->RTAnalogRangeMinimum = -1.0f;
+		DeviceCfg->ButtonMapping = ButtonMappings[0];
+		DeviceCfg->DeviceNameIndex = DeviceNames.Add(TEXT("Generic X-Box pad"));
+		RegisterDeviceConfig(0x045e, 0x0b12, DeviceCfg); // Add to the special device config map
+
+		// Xbox Wired Controller
+		DeviceCfg = MakeShared<DeviceConfig>();
+		DeviceCfg->ControllerClass = ControllerClassType::XBoxWired;
+		DeviceCfg->SupportsHat = true;
+		DeviceCfg->TriggersUseThresholdForClick = true;
+		DeviceCfg->ButtonMapping = ButtonMappings[0];
+		DeviceCfg->DeviceNameIndex = DeviceNames.Add(TEXT("Xbox Wired Controller"));
+		RegisterDeviceConfig(0x045e, 0x0202, DeviceCfg);
+		RegisterDeviceConfig(0x045e, 0x0285, DeviceCfg);
+		RegisterDeviceConfig(0x045e, 0x0287, DeviceCfg);
+		RegisterDeviceConfig(0x045e, 0x0288, DeviceCfg);
+		RegisterDeviceConfig(0x045e, 0x0289, DeviceCfg);
+		RegisterDeviceConfig(0x045e, 0x028e, DeviceCfg);
+		RegisterDeviceConfig(0x045e, 0x02d1, DeviceCfg);
+		RegisterDeviceConfig(0x045e, 0x02ea, DeviceCfg);
+
+		// Xbox Wireless Controller
+		DeviceCfg = MakeShared<DeviceConfig>();
+		DeviceCfg->ControllerClass = ControllerClassType::XBoxWireless;
+		DeviceCfg->SupportsHat = true;
+		DeviceCfg->TriggersUseThresholdForClick = true;
+		DeviceCfg->ButtonMapping = ButtonMappings[0];
+		DeviceCfg->DeviceNameIndex = DeviceNames.Add(TEXT("Xbox Wireless Controller"));
+		RegisterDeviceConfig(0x045e, 0x02e0, DeviceCfg);
+		RegisterDeviceConfig(0x045e, 0x02fd, DeviceCfg);
+		RegisterDeviceConfig(0x045e, 0x0b13, DeviceCfg);
+
+		// Xbox Wireless Controller - Older firmware before 3.1.1221.0
+		DeviceCfg = MakeShared<DeviceConfig>();
+		DeviceCfg->ControllerClass = ControllerClassType::XBoxWireless_OlderFirmware;
+		DeviceCfg->SupportsHat = true;
+		DeviceCfg->TriggersUseThresholdForClick = true;
+		DeviceCfg->MapZRZToTriggers = true;
+		DeviceCfg->RightStickZRZ = false;
+		DeviceCfg->RightStickRXRY = true;
+		DeviceCfg->ButtonMapping = ButtonMappings[2];
+		DeviceCfg->DeviceNameIndex = DeviceNames.Add(TEXT("Xbox Wireless Controller"));
+		RegisterDeviceConfig(0x045e, 0x02e0, DeviceCfg); // Add to the special device config map
+		RegisterDeviceConfig(0x045e, 0x02fd, DeviceCfg); // Add to the special device config map
+		RegisterDeviceConfig(0x045e, 0x0b13, DeviceCfg); // Add to the special device config map
+
+		// Xbox Elite Wireless Controller
+		DeviceCfg = MakeShared<DeviceConfig>();
+		DeviceCfg->ControllerClass = ControllerClassType::XBoxWireless;
+		DeviceCfg->SupportsHat = true;
+		DeviceCfg->TriggersUseThresholdForClick = true;
+		DeviceCfg->ButtonMapping = ButtonMappings[0];
+		DeviceCfg->DeviceNameIndex = DeviceNames.Add(TEXT("Xbox Elite Wireless Controller"));
+		RegisterDeviceConfig(0x045e, 0x0b00, DeviceCfg);
+		RegisterDeviceConfig(0x045e, 0x0b05, DeviceCfg);
+
+		// Xbox Elite Wireless Controller - Older firmware before 3.1.1221.0
+		DeviceCfg = MakeShared<DeviceConfig>();
+		DeviceCfg->ControllerClass = ControllerClassType::XBoxWireless_OlderFirmware;
+		DeviceCfg->SupportsHat = true;
+		DeviceCfg->TriggersUseThresholdForClick = true;
+		DeviceCfg->MapZRZToTriggers = true;
+		DeviceCfg->RightStickZRZ = false;
+		DeviceCfg->RightStickRXRY = true;
+		DeviceCfg->ButtonMapping = ButtonMappings[2];
+		DeviceCfg->DeviceNameIndex = DeviceNames.Add(TEXT("Xbox Elite Wireless Controller"));
+		RegisterDeviceConfig(0x045e, 0x0b00, DeviceCfg); // Add to the special device config map
+		RegisterDeviceConfig(0x045e, 0x0b05, DeviceCfg);
+
+		// SteelSeries Stratus XL
+		DeviceCfg = MakeShared<DeviceConfig>();
+		DeviceCfg->SupportsHat = true;
+		DeviceCfg->TriggersUseThresholdForClick = true;
+		DeviceCfg->LTAnalogRangeMinimum = 0.5f;
+		DeviceCfg->ButtonMapping = ButtonMappings[0];
+		DeviceCfg->DeviceNameIndex = DeviceNames.Add(TEXT("SteelSeries Stratus XL"));
+		RegisterDeviceConfig(0x0111, 0x1417, DeviceCfg);
+		RegisterDeviceConfig(0x0111, 0x1419, DeviceCfg);
+
+		// PS4 Wireless Controller (v2) - Android API >= 10
+		DeviceCfg = MakeShared<DeviceConfig>();
+		DeviceCfg->ControllerClass = ControllerClassType::PlaystationWireless;
+		DeviceCfg->SupportsHat = true;
+		DeviceCfg->RightStickZRZ = true;
+		DeviceCfg->ButtonMapping = ButtonMappings[0];
+		DeviceCfg->DeviceNameIndex = DeviceNames.Add(TEXT("PS4 Wireless Controller (v2)"));
+		RegisterDeviceConfig(0x054c, 0x09cc, DeviceCfg);
+
+		// PS4 Wireless Controller (v2) - Android API < 10 and CPU Vendor is not "Sony"
+		DeviceCfg = MakeShared<DeviceConfig>();
+		DeviceCfg->ControllerClass = ControllerClassType::PS4Wireless_AndroidAPIBefore10_CPUNotSony;
+		DeviceCfg->SupportsHat = true;
+		DeviceCfg->RightStickZRZ = true;
+		DeviceCfg->ButtonMapping = ButtonMappings[3];
+		DeviceCfg->DeviceNameIndex = DeviceNames.Add(TEXT("PS4 Wireless Controller (v2)"));
+		RegisterDeviceConfig(0x054c, 0x09cc, DeviceCfg); // Add to the special device config map
+
+		// PS4 Wireless Controller
+		DeviceCfg = MakeShared<DeviceConfig>();
+		DeviceCfg->ControllerClass = ControllerClassType::PlaystationWireless;
+		DeviceCfg->SupportsHat = true;
+		DeviceCfg->RightStickZRZ = true;
+		DeviceCfg->ButtonMapping = ButtonMappings[0];
+		DeviceCfg->DeviceNameIndex = DeviceNames.Add(TEXT("PS4 Wireless Controller"));
+		RegisterDeviceConfig(0x054c, 0x05c4, DeviceCfg);
+
+		// PS5 Wireless Controller - Android API >= 31
+		DeviceCfg = MakeShared<DeviceConfig>();
+		DeviceCfg->ControllerClass = ControllerClassType::PlaystationWireless;
+		DeviceCfg->SupportsHat = true;
+		DeviceCfg->RightStickZRZ = true;
+		DeviceCfg->MapRXRYToTriggers = false;
+		DeviceCfg->LTAnalogRangeMinimum = -1.0f;
+		DeviceCfg->RTAnalogRangeMinimum = -1.0f;
+		DeviceCfg->ButtonMapping = ButtonMappings[5];
+		DeviceCfg->DeviceNameIndex = DeviceNames.Add(TEXT("PS5 Wireless Controller"));
+		RegisterDeviceConfig(0x054c, 0x0ce6, DeviceCfg);
+
+		// PS5 Wireless Controller - Android API < 31
+		DeviceCfg = MakeShared<DeviceConfig>();
+		DeviceCfg->ControllerClass = ControllerClassType::PS5Wireless_AndroidAPIBefore31;
+		DeviceCfg->SupportsHat = true;
+		DeviceCfg->RightStickZRZ = true;
+		DeviceCfg->MapRXRYToTriggers = true;
+		DeviceCfg->ButtonMapping = ButtonMappings[4];
+		DeviceCfg->DeviceNameIndex = DeviceNames.Add(TEXT("PS5 Wireless Controller"));
+		RegisterDeviceConfig(0x054c, 0x0ce6, DeviceCfg); // Add to the special device config map
+
+		// glap QXPGP001
+		DeviceCfg = MakeShared<DeviceConfig>();
+		DeviceCfg->ButtonMapping = ButtonMappings[0];
+		DeviceCfg->DeviceNameIndex = DeviceNames.Add(TEXT("glap QXPGP001"));
+		RegisterDeviceConfig(0x05ac, 0x056a, DeviceCfg);
+
+		// STMicroelectronics Lenovo GamePad
+		DeviceCfg = MakeShared<DeviceConfig>();
+		DeviceCfg->ButtonMapping = ButtonMappings[0];
+		DeviceCfg->DeviceNameIndex = DeviceNames.Add(TEXT("STMicroelectronics Lenovo GamePad"));
+		RegisterDeviceConfig(0x0483, 0x5750, DeviceCfg);
+		
+		// Razer Kishi V2 Pro XBox360 - Android API >= 31
+		DeviceCfg = MakeShared<DeviceConfig>();
+		DeviceCfg->ControllerClass = ControllerClassType::XBoxWired;
+		DeviceCfg->SupportsHat = true;
+		DeviceCfg->TriggersUseThresholdForClick = true;
+		DeviceCfg->ButtonMapping = ButtonMappings[0];
+		DeviceCfg->DeviceNameIndex = DeviceNames.Add(TEXT("Razer Kishi V2 Pro XBox360"));
+		RegisterDeviceConfig(0x1532, 0x0037, DeviceCfg);
+		RegisterDeviceConfig(0x1532, 0x0718, DeviceCfg);
+
+		// Razer Kishi V2 Pro XBox360 - Android API < 31
+		DeviceCfg = MakeShared<DeviceConfig>();
+		DeviceCfg->ControllerClass = ControllerClassType::XBoxWired_AndroidAPIBefore31;
+		DeviceCfg->SupportsHat = true;
+		DeviceCfg->TriggersUseThresholdForClick = true;
+		DeviceCfg->RightStickZRZ = false;
+		DeviceCfg->RightStickRXRY = true;
+		DeviceCfg->MapZRZToTriggers = true;
+		DeviceCfg->LTAnalogRangeMinimum = -1.0f;
+		DeviceCfg->RTAnalogRangeMinimum = -1.0f;
+		DeviceCfg->ButtonMapping = ButtonMappings[0];
+		DeviceCfg->DeviceNameIndex = DeviceNames.Add(TEXT("Razer Kishi V2 Pro XBox360"));
+		RegisterDeviceConfig(0x1532, 0x0037, DeviceCfg);// Add to the special device config map
+		RegisterDeviceConfig(0x1532, 0x0718, DeviceCfg);// Add to the special device config map
+
+		// Razer Kishi V2
+		DeviceCfg = MakeShared<DeviceConfig>();
+		DeviceCfg->ControllerClass = ControllerClassType::XBoxWired;
+		DeviceCfg->SupportsHat = true;
+		DeviceCfg->TriggersUseThresholdForClick = true;
+		DeviceCfg->ButtonMapping = ButtonMappings[0];
+		DeviceCfg->DeviceNameIndex = DeviceNames.Add(TEXT("Razer Kishi V2"));
+		RegisterDeviceConfig(0x1532, 0x0712, DeviceCfg);
+
+		// Razer
+		DeviceCfg = MakeShared<DeviceConfig>();
+		DeviceCfg->ButtonMapping = ButtonMappings[0];
+		DeviceCfg->DeviceNameIndex = DeviceNames.Add(TEXT("Razer"));
+		RegisterDeviceConfig(0x1532, -1, DeviceCfg);
+
+		// Luna
+		DeviceCfg = MakeShared<DeviceConfig>();
+		DeviceCfg->TriggersUseThresholdForClick = true;
+		DeviceCfg->ButtonMapping = ButtonMappings[0];
+		DeviceCfg->DeviceNameIndex = DeviceNames.Add(TEXT("Luna"));
+		RegisterDeviceConfig(0x0171, 0x0419, DeviceCfg);
+
+		FPlatformMisc::LowLevelOutputDebugStringf(TEXT("FAndroidGameControllerMapper::Init - Layouts: %d, DeviceConfig: %d, SpecialConfig: %d, DeviceNames: %d"), 
+			ButtonMappings.Num(), DeviceConfigurationMap.Num(), DeviceConfigurationMap_Special.Num(), DeviceNames.Num());
+
+		return true;
+	}
+
+	bool FAndroidGameControllerMapper::InitByDataFile()
+	{
+		return false;
+	}
+
+	const DeviceConfig_WPtr FAndroidGameControllerMapper::GetDeviceConfig(int32 VendorId, int32 ProductId, const FName DeviceName)
+	{
+		FPlatformMisc::LowLevelOutputDebugStringf(TEXT("FAndroidGameControllerMapper::GetDeviceConfig - Vendor: 0x%08x, Product: 0x%08x, Name: %s"), VendorId, ProductId, *(DeviceName.ToString()));
+
+#if USE_ANDROID_JNI
+		if (VendorId == 0x054c && ProductId == 0x09cc && FAndroidMisc::GetCPUVendor() != TEXT("Sony") && FAndroidMisc::GetAndroidBuildVersion() < 10)
+		{
+			const DeviceConfig_SRef* Config = DeviceConfigurationMap_Special.Find(PackDeviceKey(VendorId, ProductId));
+			if (Config != nullptr)
+			{
+				check((*Config)->ControllerClass == ControllerClassType::PS4Wireless_AndroidAPIBefore10_CPUNotSony);
+				return Config->ToWeakPtr();
+			}
+		}
+		else if (VendorId == 0x054c && ProductId == 0x0ce6 && FAndroidMisc::GetAndroidBuildVersion() < 31)
+		{
+			const DeviceConfig_SRef* Config = DeviceConfigurationMap_Special.Find(PackDeviceKey(VendorId, ProductId));
+			if (Config != nullptr)
+			{
+				check((*Config)->ControllerClass == ControllerClassType::PS5Wireless_AndroidAPIBefore31);
+				return Config->ToWeakPtr();
+			}
+		}
+		else
+#endif
+		{
+			const DeviceConfig_SRef* Config = DeviceConfigurationMap.Find(PackDeviceKey(VendorId, ProductId));
+			if (Config != nullptr)
+			{
+				return Config->ToWeakPtr();
+			}
+		}
+
+		// fall back to a query by the device name
+		DeviceConfig_SPtr ConfigFoundByName = GetDeviceConfig(DeviceName);
+		if (ConfigFoundByName != nullptr)
+		{
+			FPlatformMisc::LowLevelOutputDebugStringf(TEXT("Warning: FAndroidGameControllerMapper::GetDeviceConfig - Returns the mapping found by name."));
+			return ConfigFoundByName.ToWeakPtr();
+		}
+
+		FPlatformMisc::LowLevelOutputDebugStringf(TEXT("Warning: FAndroidGameControllerMapper::GetDeviceConfig - Returns default mapping."));
+		
+		return DeviceConfigurationMap[PackDeviceKey(-1, -1)].ToWeakPtr();
+	}
+
+	const DeviceConfig_SPtr FAndroidGameControllerMapper::GetDeviceConfig(const FName InDeviceName)
+	{
+		FString DeviceName = InDeviceName.ToString();
+
+		FPlatformMisc::LowLevelOutputDebugStringf(TEXT("FAndroidGameControllerMapper::GetDeviceConfig - Fall back to query by name, Name: %s"), *DeviceName);
+
+		int32 NameIndexFound = -1;
+		for (int i = 0; i < DeviceNames.Num(); ++i)
+		{
+			FString NameInStock = DeviceNames[i].ToString();
+			if (DeviceName.StartsWith(*NameInStock))
+			{
+				NameIndexFound = i;
+				break;
+			}
+		}
+
+		if (NameIndexFound < 0)
+		{
+			return nullptr;
+		}
+
+		TMap<uint64, DeviceConfig_SRef>& ConfigMap_Ref = DeviceConfigurationMap;
+		if (GAndroidOldXBoxWirelessFirmware
+#if USE_ANDROID_JNI
+			|| (FAndroidMisc::GetCPUVendor() != TEXT("Sony") && FAndroidMisc::GetAndroidBuildVersion() < 10) ||
+			(FAndroidMisc::GetAndroidBuildVersion() < 31)
+#endif
+			)
+		{
+			ConfigMap_Ref = DeviceConfigurationMap_Special;
+		}
+
+		uint64 ConfigKey = 0;
+		for (const TPair<uint64, DeviceConfig_SRef>& Config_KV : ConfigMap_Ref)
+		{
+			if (Config_KV.Value.Get().DeviceNameIndex == NameIndexFound)
+			{
+				ConfigKey = Config_KV.Key;
+				break;
+			}
+		}
+
+		return ConfigKey == 0 ? nullptr : ConfigMap_Ref[ConfigKey].ToSharedPtr();
+	}
+
+	void FAndroidGameControllerMapper::RegisterDeviceConfig(int32 VendorId, int32 ProductId, DeviceConfig_SRef& DeviceCfg)
+	{
+		int64 DeviceKey = PackDeviceKey(VendorId, ProductId);
+
+		TMap<uint64, DeviceConfig_SRef>& ConfigMap_Ref = 
+			DeviceCfg->ControllerClass < ControllerClassType::XBoxWired_AndroidAPIBefore31 ? DeviceConfigurationMap : DeviceConfigurationMap_Special;
+
+		if (ConfigMap_Ref.Contains(DeviceKey))
+		{
+			FPlatformMisc::LowLevelOutputDebugStringf(TEXT("Warning: FAndroidGameControllerMapper::RegisterDeviceConfig - Existed device config will be replaced. Vendor: 0x%08x, Product: 0x%08x"), VendorId, ProductId);
+		}
+
+		ConfigMap_Ref.Add(DeviceKey, DeviceCfg);
+	}
+
+	void FAndroidGameControllerMapper::AppendKeyCodeKeyNameMapping(KeyCodeToKeyNameMap_SRef& TargetContainer, int32 KeyCode, int32 KeyNameIndex)
+	{
+		TArray<int32>* TargetIndexArray = TargetContainer->Find(KeyCode);
+		if (TargetIndexArray == nullptr)
+		{
+			(*TargetContainer).Add(KeyCode, { KeyNameIndex });
+		}
+		else
+		{
+			TargetIndexArray->Add(KeyNameIndex);
+		}
+	}
+}
diff --git a/Engine/Source/Runtime/ApplicationCore/Private/Android/AndroidInputInterface.cpp b/Engine/Source/Runtime/ApplicationCore/Private/Android/AndroidInputInterface.cpp
--- a/Engine/Source/Runtime/ApplicationCore/Private/Android/AndroidInputInterface.cpp
+++ b/Engine/Source/Runtime/ApplicationCore/Private/Android/AndroidInputInterface.cpp
@@ -94,6 +94,7 @@ static FAutoConsoleVariableRef CVarAndroidEnableInputDeviceListener(
 	TEXT("Determines how to detect gamepad connection/disconnection. true for using InputDeviceListener. false for using the gamepad input events (Default: true)"),
 	ECVF_Default);
 
+using namespace GameControllerMapping;
 
 bool AndroidThunkCpp_InitGameControllerManager()
 {
@@ -120,149 +121,12 @@ void JNICALL UE::Jni::FGameControllerManager::nativeOnInputDeviceStateEvent(JNIE
 	FAndroidInputInterface::HandleInputDeviceStateEvent(device_id, StateEvent, DeviceType);
 }
 
-void FAndroidGamepadDeviceMapping::Init(const FName InDeviceName)
+void FAndroidGamepadDeviceMapping::Init(const int32 InVendorId, const int32 InProductId, const FName DeviceName)
 {
-	FString DeviceName = InDeviceName.ToString();
+	TSharedPtr<FAndroidGameControllerMapper> ControllerMapper = FAndroidGameControllerMapper::GetInstance().Pin();
+	check(ControllerMapper.IsValid());
 
-	// Use device name to decide on mapping scheme
-	if (DeviceName.StartsWith(TEXT("Amazon")))
-	{
-		if (DeviceName.StartsWith(TEXT("Amazon Fire Game Controller")))
-		{
-			SupportsHat = true;
-		}
-		else if (DeviceName.StartsWith(TEXT("Amazon Fire TV Remote")))
-		{
-			SupportsHat = false;
-		}
-		else
-		{
-			SupportsHat = false;
-		}
-	}
-	else if (DeviceName.StartsWith(TEXT("NVIDIA Corporation NVIDIA Controller")))
-	{
-		SupportsHat = true;
-	}
-	else if (DeviceName.StartsWith(TEXT("Samsung Game Pad EI-GP20")))
-	{
-		SupportsHat = true;
-		MapL1R1ToTriggers = true;
-		RightStickZRZ = false;
-		RightStickRXRY = true;
-	}
-	else if (DeviceName.StartsWith(TEXT("Mad Catz C.T.R.L.R")))
-	{
-		SupportsHat = true;
-	}
-	else if (DeviceName.StartsWith(TEXT("Generic X-Box pad")))
-	{
-		ControllerClass = ControllerClassType::XBoxWired;
-		SupportsHat = true;
-		TriggersUseThresholdForClick = true;
-
-		// different mapping before Android 12
-		if (FAndroidMisc::GetAndroidBuildVersion() < 31)
-		{
-			RightStickZRZ = false;
-			RightStickRXRY = true;
-			MapZRZToTriggers = true;
-			LTAnalogRangeMinimum = -1.0f;
-			RTAnalogRangeMinimum = -1.0f;
-		}
-	}
-	else if (DeviceName.StartsWith(TEXT("Xbox Wired Controller")))
-	{
-		ControllerClass = ControllerClassType::XBoxWired;
-		SupportsHat = true;
-		TriggersUseThresholdForClick = true;
-	}
-	else if (DeviceName.StartsWith(TEXT("Xbox Wireless Controller"))
-		|| DeviceName.StartsWith(TEXT("Xbox Elite Wireless Controller")))
-	{
-		ControllerClass = ControllerClassType::XBoxWireless;
-		SupportsHat = true;
-		TriggersUseThresholdForClick = true;
-
-		if (GAndroidOldXBoxWirelessFirmware == 1)
-		{
-			// Apply mappings for older firmware before 3.1.1221.0
-			ButtonRemapping = ButtonRemapType::XBox;
-			MapL1R1ToTriggers = false;
-			MapZRZToTriggers = true;
-			RightStickZRZ = false;
-			RightStickRXRY = true;
-		}
-	}
-	else if (DeviceName.StartsWith(TEXT("SteelSeries Stratus XL")))
-	{
-		SupportsHat = true;
-		TriggersUseThresholdForClick = true;
-
-		// For some reason the left trigger is at 0.5 when at rest so we have to adjust for that.
-		LTAnalogRangeMinimum = 0.5f;
-	}
-	else if (DeviceName.StartsWith(TEXT("PS4 Wireless Controller")))
-	{
-		ControllerClass = ControllerClassType::PlaystationWireless;
-		if (DeviceName.EndsWith(TEXT(" (v2)")) && FAndroidMisc::GetCPUVendor() != TEXT("Sony")
-			&& FAndroidMisc::GetAndroidBuildVersion() < 10)
-		{
-			// Only needed for non-Sony devices with v2 firmware
-			ButtonRemapping = ButtonRemapType::PS4;
-		}
-		SupportsHat = true;
-		RightStickZRZ = true;
-	}
-	else if (DeviceName.StartsWith(TEXT("PS5 Wireless Controller")))
-	{
-		//FAndroidMisc::GetAndroidBuildVersion() actually returns the API Level instead of the Android Version
-		bool bUseNewPS5Mapping = FAndroidMisc::GetAndroidBuildVersion() > 30;
-		ButtonRemapping = bUseNewPS5Mapping ? ButtonRemapType::PS5New : ButtonRemapType::PS5;
-		ControllerClass = ControllerClassType::PlaystationWireless;
-		SupportsHat = true;
-		RightStickZRZ = true;
-		MapRXRYToTriggers = !bUseNewPS5Mapping;
-		LTAnalogRangeMinimum = bUseNewPS5Mapping ? 0.0f : -1.0f;
-		RTAnalogRangeMinimum = bUseNewPS5Mapping ? 0.0f : -1.0f;
-	}
-	else if (DeviceName.StartsWith(TEXT("glap QXPGP001")))
-	{
-		SupportsHat = true;
-	}
-	else if (DeviceName.StartsWith(TEXT("STMicroelectronics Lenovo GamePad")))
-	{
-		SupportsHat = true;
-	}
-	else if (DeviceName.StartsWith(TEXT("Razer")))
-	{
-		SupportsHat = true;
-		if (DeviceName.StartsWith(TEXT("Razer Kishi V2 Pro XBox360")))
-		{
-			ControllerClass = ControllerClassType::XBoxWired;
-			SupportsHat = true;
-			TriggersUseThresholdForClick = true;
-
-			// different mapping before Android 12
-			if (FAndroidMisc::GetAndroidBuildVersion() < 31)
-			{
-				RightStickZRZ = false;
-				RightStickRXRY = true;
-				MapZRZToTriggers = true;
-				LTAnalogRangeMinimum = -1.0f;
-				RTAnalogRangeMinimum = -1.0f;
-			}
-		}
-		else if (DeviceName.StartsWith(TEXT("Razer Kishi V2")))
-		{
-			ControllerClass = ControllerClassType::XBoxWired;
-			TriggersUseThresholdForClick = true;
-		}
-	}
-	else if (DeviceName.StartsWith(TEXT("Luna")))
-	{
-		TriggersUseThresholdForClick = true;
-	}
+	DeviceConfig = ControllerMapper->GetDeviceConfig(InVendorId, InProductId, DeviceName);
 
 	ResetRuntimeData();
 }
@@ -1269,7 +1133,7 @@ void FAndroidInputInterface::SendControllerEvents()
 						FInputDeviceId InputDeviceId = InternalDeviceIdMappings.FindDeviceId(CurrentDevice.Descriptor);
 						if (InputDeviceId == INPUTDEVICEID_NONE)
 						{
-							FAndroidGamepadDeviceMapping DeviceData(CurrentDevice.Name);
+							FAndroidGamepadDeviceMapping DeviceData(CurrentDevice.VendorId, CurrentDevice.ProductId, CurrentDevice.Name);
 							InputDeviceId = InternalDeviceIdMappings.GetOrCreateDeviceId(CurrentDevice.Descriptor);
 							GameControllerDataMap.Add(InputDeviceId, DeviceData);
 						}
@@ -1389,7 +1253,7 @@ void FAndroidInputInterface::SendControllerEvents()
 			FInputDeviceScope InputScope(nullptr, DeviceInfo.Name, InputDeviceId.GetId(), HardwareDeviceIdentifier_DefaultGamepad);
 
 			FAndroidGamepadDeviceMapping* DeviceData = GameControllerDataMap.Find(InputDeviceId);
-			if (DeviceData == nullptr)
+			if (DeviceData == nullptr || !DeviceData->DeviceConfig.IsValid())
 			{
 				// The protection here is for some very rare cases, with only 1 PS5 controller connected, the OS will report
 				// two game controllers and one of them is a false game controller without a name and descriptor. It will generate
@@ -1434,7 +1298,9 @@ void FAndroidInputInterface::SendControllerEvents()
 				NewControllerState.ButtonStates[MAX_NUM_PHYSICAL_CONTROLLER_BUTTONS + 7] = NewControllerState.RYAnalog <= -RepeatDeadzone;
 			}
 			
-			const bool UseTriggerThresholdForClick = DeviceData->TriggersUseThresholdForClick;
+			DeviceConfig_SPtr DeviceConfigPtr = DeviceData->DeviceConfig.Pin();
+			check(DeviceConfigPtr.IsValid());
+			const bool UseTriggerThresholdForClick = DeviceConfigPtr->TriggersUseThresholdForClick;
 			if (NewControllerState.LTAnalog != OldControllerState.LTAnalog)
 			{
 				//LOGD("    Sending updated LeftTriggerAnalog value of %f", NewControllerState.LTAnalog);
@@ -1443,7 +1309,7 @@ void FAndroidInputInterface::SendControllerEvents()
 				if (UseTriggerThresholdForClick)
 				{
 					// Handle the trigger theshold "virtual" button state
-					//check(ButtonMapping[10] == FGamepadKeyNames::LeftTriggerThreshold);
+					//check(IndexToKeyNameMap[10] == FGamepadKeyNames::LeftTriggerThreshold);
 					NewControllerState.ButtonStates[10] = NewControllerState.LTAnalog >= ANDROID_GAMEPAD_TRIGGER_THRESHOLD;
 				}
 			}
@@ -1455,7 +1321,7 @@ void FAndroidInputInterface::SendControllerEvents()
 				if (UseTriggerThresholdForClick)
 				{
 					// Handle the trigger theshold "virtual" button state
-					//check(ButtonMapping[11] == FGamepadKeyNames::RightTriggerThreshold);
+					//check(IndexToKeyNameMap[11] == FGamepadKeyNames::RightTriggerThreshold);
 					NewControllerState.ButtonStates[11] = NewControllerState.RTAnalog >= ANDROID_GAMEPAD_TRIGGER_THRESHOLD;
 				}
 			}
@@ -1499,7 +1365,7 @@ void FAndroidInputInterface::SendControllerEvents()
 			{
 				if (DeviceInfo.FeedbackMotorCount > 0)
 				{
-					UpdateControllerVibeMotors(DeviceInfo.DeviceId, DeviceData->ControllerClass, DeviceData->ControllerVibeState);
+					UpdateControllerVibeMotors(DeviceInfo.DeviceId, DeviceConfigPtr->ControllerClass, DeviceData->ControllerVibeState);
 				}
 			}
 
@@ -1616,7 +1482,10 @@ ControllerClassType FAndroidInputInterface::GetControllerClass(int32 ControllerI
 	FAndroidGamepadDeviceMapping* OutDeviceData;
 	if (GetInputDeviceByControllerId(ControllerId, nullptr, &OutDeviceData))
 	{
-		return OutDeviceData->ControllerClass;
+		if (OutDeviceData->DeviceConfig.IsValid())
+		{
+			return OutDeviceData->DeviceConfig.Pin()->ControllerClass;
+		}
 	}
 	
 	return Result;
@@ -1760,8 +1629,11 @@ void FAndroidInputInterface::DumpInputDevices()
 	{
 		FInputDeviceId InputDeviceId = KVP.Key;
 		const FAndroidGamepadDeviceMapping& Data = KVP.Value;
-		UE_LOG(LogAndroid, Log, TEXT("DeviceData, InputDeviceId = %02d, ControllerClass = %d, ButtonRemapping = %d"),
-			InputDeviceId.GetId(), Data.ControllerClass, Data.ButtonRemapping);
+		if (Data.DeviceConfig.IsValid())
+		{
+			UE_LOG(LogAndroid, Log, TEXT("DeviceData, InputDeviceId = %02d, ControllerClass = %d"),
+				InputDeviceId.GetId(), Data.DeviceConfig.Pin()->ControllerClass);
+		}
 	}
 
 	UE_LOG(LogAndroid, Log, TEXT("---"));
@@ -1817,7 +1689,7 @@ void FAndroidInputInterface::JoystickAxisEvent(int32 DeviceId, int32 AxisId, flo
 		UE_LOG(LogAndroid, Log, TEXT("Gamepad state changed to Valid, DeviceId = %d"), DeviceId);
 	}
 
-	if (DeviceData == nullptr)
+	if (DeviceData == nullptr || !DeviceData->DeviceConfig.IsValid())
 	{
 		return; // bypass inputs from the device not identified as a game controller
 	}
@@ -1833,6 +1705,9 @@ void FAndroidInputInterface::JoystickAxisEvent(int32 DeviceId, int32 AxisId, flo
 		return Value;
 	};
 
+	DeviceConfig_SPtr DeviceConfigPtr = DeviceData->DeviceConfig.Pin();
+	check(DeviceConfigPtr.IsValid());
+
 	// Deal with left stick and triggers (generic)
 	switch (AxisId)
 	{
@@ -1840,24 +1715,24 @@ void FAndroidInputInterface::JoystickAxisEvent(int32 DeviceId, int32 AxisId, flo
 		case AMOTION_EVENT_AXIS_Y:			DeviceData->NewControllerData.LYAnalog = -AxisValue; return;
 		case AMOTION_EVENT_AXIS_LTRIGGER:
 			{
-				if (!(DeviceData->MapZRZToTriggers || DeviceData->MapRXRYToTriggers))
+				if (!(DeviceConfigPtr->MapZRZToTriggers || DeviceConfigPtr->MapRXRYToTriggers))
 				{
-					DeviceData->NewControllerData.LTAnalog = RemapTriggerFunction(DeviceData->LTAnalogRangeMinimum, AxisValue);
+					DeviceData->NewControllerData.LTAnalog = RemapTriggerFunction(DeviceConfigPtr->LTAnalogRangeMinimum, AxisValue);
 					return;
 				}
 			}
 		case AMOTION_EVENT_AXIS_RTRIGGER:
 			{
-				if (!(DeviceData->MapZRZToTriggers || DeviceData->MapRXRYToTriggers))
+				if (!(DeviceConfigPtr->MapZRZToTriggers || DeviceConfigPtr->MapRXRYToTriggers))
 				{
-					DeviceData->NewControllerData.RTAnalog = RemapTriggerFunction(DeviceData->RTAnalogRangeMinimum, AxisValue);
+					DeviceData->NewControllerData.RTAnalog = RemapTriggerFunction(DeviceConfigPtr->RTAnalogRangeMinimum, AxisValue);
 					return;
 				}
 			}
 	}
 
 	// Deal with right stick Z/RZ events
-	if (DeviceData->RightStickZRZ)
+	if (DeviceConfigPtr->RightStickZRZ)
 	{
 		switch (AxisId)
 		{
@@ -1867,7 +1742,7 @@ void FAndroidInputInterface::JoystickAxisEvent(int32 DeviceId, int32 AxisId, flo
 	}
 
 	// Deal with right stick RX/RY events
-	if (DeviceData->RightStickRXRY)
+	if (DeviceConfigPtr->RightStickRXRY)
 	{
 		switch (AxisId)
 		{
@@ -1877,26 +1752,26 @@ void FAndroidInputInterface::JoystickAxisEvent(int32 DeviceId, int32 AxisId, flo
 	}
 
 	// Deal with Z/RZ mapping to triggers
-	if (DeviceData->MapZRZToTriggers)
+	if (DeviceConfigPtr->MapZRZToTriggers)
 	{
 		switch (AxisId)
 		{
-			case AMOTION_EVENT_AXIS_Z:		DeviceData->NewControllerData.LTAnalog = RemapTriggerFunction(DeviceData->LTAnalogRangeMinimum, AxisValue); return;
-			case AMOTION_EVENT_AXIS_RZ:		DeviceData->NewControllerData.RTAnalog = RemapTriggerFunction(DeviceData->RTAnalogRangeMinimum, AxisValue); return;
+			case AMOTION_EVENT_AXIS_Z:		DeviceData->NewControllerData.LTAnalog = RemapTriggerFunction(DeviceConfigPtr->LTAnalogRangeMinimum, AxisValue); return;
+			case AMOTION_EVENT_AXIS_RZ:		DeviceData->NewControllerData.RTAnalog = RemapTriggerFunction(DeviceConfigPtr->RTAnalogRangeMinimum, AxisValue); return;
 		}
 	}
 
-	if (DeviceData->MapRXRYToTriggers)
+	if (DeviceConfigPtr->MapRXRYToTriggers)
 	{
 		switch (AxisId)
 		{
-			case AMOTION_EVENT_AXIS_RX:		DeviceData->NewControllerData.LTAnalog = RemapTriggerFunction(DeviceData->LTAnalogRangeMinimum, AxisValue); return;
-			case AMOTION_EVENT_AXIS_RY:		DeviceData->NewControllerData.RTAnalog = RemapTriggerFunction(DeviceData->RTAnalogRangeMinimum, AxisValue); return;
+			case AMOTION_EVENT_AXIS_RX:		DeviceData->NewControllerData.LTAnalog = RemapTriggerFunction(DeviceConfigPtr->LTAnalogRangeMinimum, AxisValue); return;
+			case AMOTION_EVENT_AXIS_RY:		DeviceData->NewControllerData.RTAnalog = RemapTriggerFunction(DeviceConfigPtr->RTAnalogRangeMinimum, AxisValue); return;
 		}
 	}
 
 	// Deal with hat (convert to DPAD buttons)
-	if (DeviceData->SupportsHat)
+	if (DeviceConfigPtr->SupportsHat)
 	{
 		// Apply a small dead zone to hats
 		const float DeadZone = 0.2f;
@@ -1963,172 +1838,33 @@ void FAndroidInputInterface::JoystickButtonEvent(int32 deviceId, int32 buttonId,
 		DumpInputDevices();
 	}
 
-	if (DeviceData == nullptr)
+	if (DeviceData == nullptr || !DeviceData->DeviceConfig.IsValid())
 	{
-		return; // bypass inputs from the device not identified as a game controller
-	}
-
-	//FPlatformMisc::LowLevelOutputDebugStringf(TEXT("JoystickButtonEvent[%d]: %d"), (int)DeviceData->ButtonRemapping, buttonId);
-
-	if (DeviceData->ControllerClass == ControllerClassType::PlaystationWireless)
-	{
-		if (buttonId == 3002)
-		{
-			DeviceData->NewControllerData.ButtonStates[7] = buttonDown;  // Touchpad = Special Left
-			return;
-		}
+		return; // bypass inputs from a device not identified as a game controller or without a valid device configuration data
 	}
 
 	// Deal with button remapping
-	switch (DeviceData->ButtonRemapping)
-	{
-		case ButtonRemapType::Normal:
-			switch (buttonId)
-			{
-				case AKEYCODE_BUTTON_A:
-				case AKEYCODE_DPAD_CENTER:   DeviceData->NewControllerData.ButtonStates[0] = buttonDown; break;
-				case AKEYCODE_BUTTON_B:      DeviceData->NewControllerData.ButtonStates[1] = buttonDown; break;
-				case AKEYCODE_BUTTON_X:      DeviceData->NewControllerData.ButtonStates[2] = buttonDown; break;
-				case AKEYCODE_BUTTON_Y:      DeviceData->NewControllerData.ButtonStates[3] = buttonDown; break;
-				case AKEYCODE_BUTTON_L1:     DeviceData->NewControllerData.ButtonStates[4] = buttonDown;
-											 if (DeviceData->MapL1R1ToTriggers)
-											 {
-												 DeviceData->NewControllerData.ButtonStates[10] = buttonDown;
-											 }
-											 break;
-				case AKEYCODE_BUTTON_R1:     DeviceData->NewControllerData.ButtonStates[5] = buttonDown;
-											 if (DeviceData->MapL1R1ToTriggers)
-											 {
-												 DeviceData->NewControllerData.ButtonStates[11] = buttonDown;
-											 }
-											 break;
-				case AKEYCODE_BUTTON_START:
-				case AKEYCODE_MENU:          DeviceData->NewControllerData.ButtonStates[6] = buttonDown;
-											 if (!bBlockAndroidKeysOnControllers)
-											 {
-												 DeviceData->NewControllerData.ButtonStates[17] = buttonDown;
-											 }
-											 break;
-				case AKEYCODE_BUTTON_SELECT:
-				case AKEYCODE_BACK:          DeviceData->NewControllerData.ButtonStates[7] = buttonDown;
-											 if (!bBlockAndroidKeysOnControllers)
-											 {
-												 DeviceData->NewControllerData.ButtonStates[16] = buttonDown;
-											 }
-											 break;
-				case AKEYCODE_BUTTON_THUMBL: DeviceData->NewControllerData.ButtonStates[8] = buttonDown; break;
-				case AKEYCODE_BUTTON_THUMBR: DeviceData->NewControllerData.ButtonStates[9] = buttonDown; break;
-				case AKEYCODE_BUTTON_L2:     DeviceData->NewControllerData.ButtonStates[10] = buttonDown; break;
-				case AKEYCODE_BUTTON_R2:     DeviceData->NewControllerData.ButtonStates[11] = buttonDown; break;
-				case AKEYCODE_DPAD_UP:       DeviceData->NewControllerData.ButtonStates[12] = buttonDown; break;
-				case AKEYCODE_DPAD_DOWN:     DeviceData->NewControllerData.ButtonStates[13] = buttonDown; break;
-				case AKEYCODE_DPAD_LEFT:     DeviceData->NewControllerData.ButtonStates[14] = buttonDown; break;
-				case AKEYCODE_DPAD_RIGHT:    DeviceData->NewControllerData.ButtonStates[15] = buttonDown; break;
-			}
-			break;
-
-		case ButtonRemapType::XBox:
-			switch (buttonId)
+	DeviceConfig_SPtr DeviceConfigPtr = DeviceData->DeviceConfig.Pin();
+	check(DeviceConfigPtr.IsValid());
+	KeyCodeToKeyNameMap_SPtr ButtonMappingPtr = DeviceConfigPtr->ButtonMapping.Pin();
+	check(ButtonMappingPtr.IsValid());
+	TArray<int32>* KeyNameIndices = ButtonMappingPtr->Find(buttonId);
+	if (KeyNameIndices != nullptr)
+	{
+		for (int32 KeyNameIndex : *KeyNameIndices)
+		{
+			if ((KeyNameIndex == 16 || KeyNameIndex == 17) && bBlockAndroidKeysOnControllers)
 			{
-				case AKEYCODE_BUTTON_A:      DeviceData->NewControllerData.ButtonStates[0] = buttonDown; break; // A
-				case AKEYCODE_BUTTON_B:      DeviceData->NewControllerData.ButtonStates[1] = buttonDown; break; // B
-				case AKEYCODE_BUTTON_C:      DeviceData->NewControllerData.ButtonStates[2] = buttonDown; break; // X
-				case AKEYCODE_BUTTON_X:      DeviceData->NewControllerData.ButtonStates[3] = buttonDown; break; // Y
-				case AKEYCODE_BUTTON_Y:      DeviceData->NewControllerData.ButtonStates[4] = buttonDown; break; // L1
-				case AKEYCODE_BUTTON_Z:      DeviceData->NewControllerData.ButtonStates[5] = buttonDown; break; // R1
-				case AKEYCODE_BUTTON_R1:     DeviceData->NewControllerData.ButtonStates[6] = buttonDown;
-											 if (!bBlockAndroidKeysOnControllers)
-											 {
-												 DeviceData->NewControllerData.ButtonStates[17] = buttonDown; // Menu
-											 }
-											 break;
-				case AKEYCODE_BUTTON_L1:     DeviceData->NewControllerData.ButtonStates[7] = buttonDown;
-											 if (!bBlockAndroidKeysOnControllers)
-											 {
-												 DeviceData->NewControllerData.ButtonStates[16] = buttonDown; // View
-											 }
-											 break;
-				case AKEYCODE_BUTTON_L2:     DeviceData->NewControllerData.ButtonStates[8] = buttonDown; break; // ThumbL
-				case AKEYCODE_BUTTON_R2:     DeviceData->NewControllerData.ButtonStates[9] = buttonDown; break; // ThumbR
+				continue;
 			}
-			break;
 
-		case ButtonRemapType::PS4:
-			switch (buttonId)
-			{
-				case AKEYCODE_BUTTON_B:      DeviceData->NewControllerData.ButtonStates[0] = buttonDown; break; // Cross
-				case AKEYCODE_BUTTON_C:      DeviceData->NewControllerData.ButtonStates[1] = buttonDown; break; // Circle
-				case AKEYCODE_BUTTON_A:      DeviceData->NewControllerData.ButtonStates[2] = buttonDown; break; // Square
-				case AKEYCODE_BUTTON_X:      DeviceData->NewControllerData.ButtonStates[3] = buttonDown; break; // Triangle
-				case AKEYCODE_BUTTON_Y:      DeviceData->NewControllerData.ButtonStates[4] = buttonDown; break; // L1
-				case AKEYCODE_BUTTON_Z:      DeviceData->NewControllerData.ButtonStates[5] = buttonDown; break; // R1
-				case AKEYCODE_BUTTON_L2:     DeviceData->NewControllerData.ButtonStates[6] = buttonDown;
-											 if (!bBlockAndroidKeysOnControllers)
-											 {
-											 	DeviceData->NewControllerData.ButtonStates[17] = buttonDown; // Options
-											 }
-											 break;
-				case AKEYCODE_MENU:          DeviceData->NewControllerData.ButtonStates[7] = buttonDown;
-											 if (!bBlockAndroidKeysOnControllers)
-											 {
-												 DeviceData->NewControllerData.ButtonStates[16] = buttonDown; // Touchpad
-											 }
-											 break;
-				case AKEYCODE_BUTTON_SELECT: DeviceData->NewControllerData.ButtonStates[8] = buttonDown; break; // ThumbL
-				case AKEYCODE_BUTTON_START:  DeviceData->NewControllerData.ButtonStates[9] = buttonDown; break; // ThumbR
-				case AKEYCODE_BUTTON_L1:     DeviceData->NewControllerData.ButtonStates[10] = buttonDown; break; // L2
-				case AKEYCODE_BUTTON_R1:     DeviceData->NewControllerData.ButtonStates[11] = buttonDown; break; // R2
-			}
-			break;
-		case ButtonRemapType::PS5:
-			switch (buttonId)
-			{
-				case AKEYCODE_BUTTON_B:      DeviceData->NewControllerData.ButtonStates[0] = buttonDown; break; // Cross
-				case AKEYCODE_BUTTON_C:      DeviceData->NewControllerData.ButtonStates[1] = buttonDown; break; // Circle
-				case AKEYCODE_BUTTON_A:      DeviceData->NewControllerData.ButtonStates[2] = buttonDown; break; // Square
-				case AKEYCODE_BUTTON_X:      DeviceData->NewControllerData.ButtonStates[3] = buttonDown; break; // Triangle
-				case AKEYCODE_BUTTON_Y:      DeviceData->NewControllerData.ButtonStates[4] = buttonDown; break; // L1
-				case AKEYCODE_BUTTON_Z:      DeviceData->NewControllerData.ButtonStates[5] = buttonDown; break; // R1
-				case AKEYCODE_BUTTON_R2:     DeviceData->NewControllerData.ButtonStates[6] = buttonDown;
-					if (!bBlockAndroidKeysOnControllers)
-					{
-						DeviceData->NewControllerData.ButtonStates[17] = buttonDown; // Options
-					}
-					break;
-				case AKEYCODE_BUTTON_THUMBL:          DeviceData->NewControllerData.ButtonStates[7] = buttonDown;
-					if (!bBlockAndroidKeysOnControllers)
-					{
-						DeviceData->NewControllerData.ButtonStates[16] = buttonDown; // Touchpad
-					}
-					break;
-				case AKEYCODE_BUTTON_SELECT: DeviceData->NewControllerData.ButtonStates[8] = buttonDown; break; // ThumbL
-				case AKEYCODE_BUTTON_START:  DeviceData->NewControllerData.ButtonStates[9] = buttonDown; break; // ThumbR
-				case AKEYCODE_BUTTON_L1:     DeviceData->NewControllerData.ButtonStates[10]= buttonDown; break; // L2
-				case AKEYCODE_BUTTON_R1:     DeviceData->NewControllerData.ButtonStates[11] = buttonDown; break; // R2
-			}
-			break;
-		case ButtonRemapType::PS5New:
-			switch (buttonId)
-			{
-				case AKEYCODE_BUTTON_A:		DeviceData->NewControllerData.ButtonStates[0] = buttonDown; break; // Cross
-				case AKEYCODE_BUTTON_B:		DeviceData->NewControllerData.ButtonStates[1] = buttonDown; break; // Circle
-				case AKEYCODE_BUTTON_X:		DeviceData->NewControllerData.ButtonStates[2] = buttonDown; break; // Triangle
-				case AKEYCODE_BUTTON_Y:		DeviceData->NewControllerData.ButtonStates[3] = buttonDown; break; // Square
-				case AKEYCODE_BUTTON_L1:	DeviceData->NewControllerData.ButtonStates[4] = buttonDown; break; // L1
-				case AKEYCODE_BUTTON_R1:	DeviceData->NewControllerData.ButtonStates[5] = buttonDown; break; // R1
-				case AKEYCODE_BUTTON_THUMBL:DeviceData->NewControllerData.ButtonStates[8] = buttonDown; break; // L3
-				case AKEYCODE_BUTTON_THUMBR:DeviceData->NewControllerData.ButtonStates[9] = buttonDown; break; // R3
-				case AKEYCODE_BUTTON_L2:	DeviceData->NewControllerData.ButtonStates[10] = buttonDown; break; // L2
-				case AKEYCODE_BUTTON_R2:	DeviceData->NewControllerData.ButtonStates[11] = buttonDown; break; // R2
-				case 3002:		DeviceData->NewControllerData.ButtonStates[16] = buttonDown; break; // Touchpad
-				case AKEYCODE_BUTTON_START:	DeviceData->NewControllerData.ButtonStates[6] = buttonDown; // Options
-					if (!bBlockAndroidKeysOnControllers)
-					{
-						DeviceData->NewControllerData.ButtonStates[17] = buttonDown; // Options
-					}
-					break;
-			}
-			break;
+			FPlatformMisc::LowLevelOutputDebugStringf(TEXT("JoystickButtonEvent[%d]: %d"), (int)KeyNameIndex, buttonId);
+			DeviceData->NewControllerData.ButtonStates[KeyNameIndex] = buttonDown;
+		}
+	}
+	else
+	{
+		FPlatformMisc::LowLevelOutputDebugStringf(TEXT("Found KeyCode without mapping: %d"), buttonId);
 	}
 }
 
@@ -2143,9 +1879,14 @@ int32 FAndroidInputInterface::GetAlternateKeyEventForMouse(int32 deviceId, int32
 		FAndroidGamepadDeviceMapping* data;
 		if (GetInputDeviceByDeviceId(deviceId, &Info, &data) &&
 			Info->DeviceState == MappingState::Valid &&
-			data->ControllerClass == ControllerClassType::PlaystationWireless)
+			data != nullptr && data->DeviceConfig.IsValid())
 		{
-			return 3002;
+			DeviceConfig_SPtr DeviceConfigPtr = data->DeviceConfig.Pin();
+			check(DeviceConfigPtr.IsValid());
+			if (DeviceConfigPtr->ControllerClass == ControllerClassType::PlaystationWireless)
+			{
+				return 3002;
+			}
 		}
 	}
 	return 0;
diff --git a/Engine/Source/Runtime/ApplicationCore/Public/Android/AndroidGameControllerMapper.h b/Engine/Source/Runtime/ApplicationCore/Public/Android/AndroidGameControllerMapper.h
new file mode 100644
--- /dev/null
+++ b/Engine/Source/Runtime/ApplicationCore/Public/Android/AndroidGameControllerMapper.h
@@ -0,0 +1,117 @@
+// Copyright Epic Games, Inc. All Rights Reserved.
+#pragma once
+
+#include "CoreTypes.h"
+#include "Templates/Function.h"
+#include "UObject/NameTypes.h"
+#include "Logging/LogMacros.h"
+#include <Templates/SharedPointer.h>
+#include "GenericPlatform/GenericApplicationMessageHandler.h"
+#include "Android/AndroidPlatform.h"
+
+namespace GameControllerMapping
+{
+	enum ControllerClassType : int
+	{
+		Generic,
+		XBoxWired,
+		XBoxWireless,
+		PlaystationWireless,
+		
+		XBoxWired_AndroidAPIBefore31,
+		XBoxWireless_OlderFirmware,
+		PS4Wireless_AndroidAPIBefore10_CPUNotSony,
+		PS5Wireless_AndroidAPIBefore31,
+	};
+
+	enum ButtonRemapType : int
+	{
+		Normal,
+		XBox,
+		PS4,
+		PS5,
+		PS5New
+	};
+
+	typedef TMap<int32/*KeyCode*/, TArray<int32/*KeyNameIndex*/>> KeyCodeToKeyNameMap; // ButtonStates[] and ButtonMapping[] use same indices
+	typedef TWeakPtr<KeyCodeToKeyNameMap> KeyCodeToKeyNameMap_WPtr;
+	typedef TSharedPtr<KeyCodeToKeyNameMap> KeyCodeToKeyNameMap_SPtr;
+	typedef TSharedRef<KeyCodeToKeyNameMap> KeyCodeToKeyNameMap_SRef;
+
+	struct DeviceConfig
+	{
+		// Type of controller
+		ControllerClassType ControllerClass = ControllerClassType::Generic;
+
+		// Sets the analog range of the trigger minimum (normally 0).  Final value is mapped as (input - Minimum) / (1 - Minimum) to [0,1] output.
+		float LTAnalogRangeMinimum = .0f;
+		float RTAnalogRangeMinimum = .0f;
+
+		// Device supports hat as dpad
+		bool SupportsHat = true;
+
+		// Device uses threshold to send button pressed events.
+		bool TriggersUseThresholdForClick = false;
+
+		// Map L1 and R1 to LTRIGGER and RTRIGGER
+		bool MapL1R1ToTriggers = false;
+
+		// Map Z and RZ to LTAnalog and RTAnalog
+		bool MapZRZToTriggers = false;
+
+		// Right stick on Z/RZ
+		bool RightStickZRZ = true;
+
+		// Right stick on RX/RY
+		bool RightStickRXRY = false;
+
+		// Map RX and RY to LTAnalog and RTAnalog
+		bool MapRXRYToTriggers = false;
+
+		int DeviceNameIndex = -1;
+
+		TWeakPtr<KeyCodeToKeyNameMap> ButtonMapping;
+	};
+
+	typedef TWeakPtr<DeviceConfig> DeviceConfig_WPtr;
+	typedef TSharedPtr<DeviceConfig> DeviceConfig_SPtr;
+	typedef TSharedRef<DeviceConfig> DeviceConfig_SRef;
+
+	class FAndroidGameControllerMapper
+	{
+	public:
+		~FAndroidGameControllerMapper() {}
+
+		static TWeakPtr<FAndroidGameControllerMapper> GetInstance();
+
+		const DeviceConfig_WPtr GetDeviceConfig(int32 VendorId, int32 ProductId, const FName DeviceName);
+
+		void RegisterDeviceConfig(int32 VendorId, int32 ProductId, DeviceConfig_SRef& DeviceCfg);
+
+	private:
+		FAndroidGameControllerMapper() {}
+
+		bool Init();
+
+		bool InitByDataFile();
+
+		const DeviceConfig_SPtr GetDeviceConfig(const FName DeviceName);
+
+		void AppendKeyCodeKeyNameMapping(KeyCodeToKeyNameMap_SRef& TargetContainer, int32 KeyCode, int32 KeyNameIndex);
+
+		uint64 PackDeviceKey(int32 VendorId, int32 ProductId) const
+		{
+			return ((uint64)VendorId) << 32 | (uint32)ProductId;
+		}
+
+		TArray<KeyCodeToKeyNameMap_SRef> ButtonMappings;
+		
+		TMap<uint64, DeviceConfig_SRef> DeviceConfigurationMap;
+
+		TMap<uint64, DeviceConfig_SRef> DeviceConfigurationMap_Special;
+
+		TArray<const FName> DeviceNames;
+
+		static TSharedPtr<FAndroidGameControllerMapper> Instance;
+	};
+}
diff --git a/Engine/Source/Runtime/ApplicationCore/Public/Android/AndroidInputInterface.h b/Engine/Source/Runtime/ApplicationCore/Public/Android/AndroidInputInterface.h
--- a/Engine/Source/Runtime/ApplicationCore/Public/Android/AndroidInputInterface.h
+++ b/Engine/Source/Runtime/ApplicationCore/Public/Android/AndroidInputInterface.h
@@ -12,6 +12,8 @@
 #include "Math/Vector2D.h"
 #include "Math/Color.h"
 
+#include "Android/AndroidGameControllerMapper.h"
+
 #if __ANDROID_API__ < 13
 
 // Joystick functions and constants only available API level 13 and above
@@ -103,23 +105,6 @@ enum MappingState
 	Valid
 };
 
-enum ControllerClassType
-{
-	Generic,
-	XBoxWired,
-	XBoxWireless,
-	PlaystationWireless
-};
-
-enum ButtonRemapType
-{
-	Normal,
-	XBox,
-	PS4,
-	PS5,
-	PS5New
-};
-
 struct FAndroidInputDeviceInfo 
 {
 	FAndroidInputDeviceInfo()
@@ -184,17 +169,12 @@ struct FAndroidControllerVibeState
 
 struct FAndroidGamepadDeviceMapping
 {
-	FAndroidGamepadDeviceMapping()
-	{
-		ResetRuntimeData();
-	}
-
-	FAndroidGamepadDeviceMapping(const FName DeviceName)
+	FAndroidGamepadDeviceMapping(const int32 InVendorId, const int32 InProductId, const FName InDeviceName)
 	{
-		Init(DeviceName);
+		Init(InVendorId, InProductId, InDeviceName);
 	}
 
-	void Init(const FName DeviceName);
+	void Init(const int32 InVendorId, const int32 InProductId, const FName InDeviceName);
 
 	void ResetRuntimeData()
 	{
@@ -203,36 +183,7 @@ struct FAndroidGamepadDeviceMapping
 		FMemory::Memset(ControllerVibeState, 0);
 	}
 
-	// Type of controller
-	ControllerClassType ControllerClass = ControllerClassType::Generic;
-
-	// Type of button remapping to use
-	ButtonRemapType ButtonRemapping = ButtonRemapType::Normal;
-
-	// Sets the analog range of the trigger minimum (normally 0).  Final value is mapped as (input - Minimum) / (1 - Minimum) to [0,1] output.
-	float LTAnalogRangeMinimum = .0f;
-	float RTAnalogRangeMinimum = .0f;
-
-	// Device supports hat as dpad
-	bool SupportsHat = true;
-
-	// Device uses threshold to send button pressed events.
-	bool TriggersUseThresholdForClick = false;
-
-	// Map L1 and R1 to LTRIGGER and RTRIGGER
-	bool MapL1R1ToTriggers = false;
-
-	// Map Z and RZ to LTAnalog and RTAnalog
-	bool MapZRZToTriggers = false;
-
-	// Right stick on Z/RZ
-	bool RightStickZRZ = true;
-
-	// Right stick on RX/RY
-	bool RightStickRXRY = false;
-
-	// Map RX and RY to LTAnalog and RTAnalog
-	bool MapRXRYToTriggers = false;
+	GameControllerMapping::DeviceConfig_WPtr DeviceConfig;
 
 	FAndroidControllerData OldControllerData;
 	FAndroidControllerData NewControllerData;
@@ -264,6 +215,7 @@ struct FDeferredAndroidMessage
 	};
 };
 
+typedef GameControllerMapping::ControllerClassType ControllerClassType;
 /**
  * Interface class for Android input devices                 
  */
