- Verify Direct Install files are cleaned up during startup on IOS

#rnx

diff --git a/Engine/Source/Runtime/ApplicationCore/ApplicationCore.Build.cs b/Engine/Source/Runtime/ApplicationCore/ApplicationCore.Build.cs
--- a/Engine/Source/Runtime/ApplicationCore/ApplicationCore.Build.cs
+++ b/Engine/Source/Runtime/ApplicationCore/ApplicationCore.Build.cs
@@ -83,6 +83,7 @@ public class ApplicationCore : ModuleRules
 			//Need to add this as BackgroundHTTP files can end up doing work directly from our AppDelegate in iOS and thus we need acccess to correct file locations to save these very early.
 			PrivateDependencyModuleNames.Add("BackgroundHTTPFileHash");
 			PrivateDependencyModuleNames.Add("CoreUObject");
+			PrivateDependencyModuleNames.Add("BuildSettings");
 			PublicIncludePathModuleNames.Add("BackgroundHTTP");
 		}
 		else if (Target.Platform == UnrealTargetPlatform.Android)
diff --git a/Engine/Source/Runtime/ApplicationCore/Private/IOS/IOSBackgroundURLSessionHandler.cpp b/Engine/Source/Runtime/ApplicationCore/Private/IOS/IOSBackgroundURLSessionHandler.cpp
--- a/Engine/Source/Runtime/ApplicationCore/Private/IOS/IOSBackgroundURLSessionHandler.cpp
+++ b/Engine/Source/Runtime/ApplicationCore/Private/IOS/IOSBackgroundURLSessionHandler.cpp
@@ -13,6 +13,8 @@
 
 #include "IOS/IOSAppDelegate.h"
 
+#include "BuildSettings.h"
+
 static bool GBackgroundHttpUseExistingNonResumableTasks = true;
 static FAutoConsoleVariableRef CVarBackgroundHttpUseExistingNonResumableTasks(
 	TEXT("ios.BackgroundHttp.bUseExistingNonResumableTasks"),
@@ -692,10 +694,13 @@ static constexpr NSInteger TypicalSizeForDownloadResponseHeader = 400; // bytes
 		const FString TaskURL(URL.absoluteString);
 		const FString& TempFileName = HelperRef->FindOrAddTempFilenameMappingForURL(TaskURL, RequestID);
 		const FString& DestinationPath = HelperRef->GetFullPathOfTempFilename(TempFileName);
-
+		
 		IPlatformFile& PlatformFile = FPlatformFileManager::Get().GetPlatformFile();
 		const FString ConvertedPath = PlatformFile.ConvertToAbsolutePathForExternalAppForWrite(*DestinationPath);
 
+		const FString DestinationDir = FPaths::GetPath(ConvertedPath);
+		PlatformFile.CreateDirectory(*DestinationDir);
+
 		return ConvertedPath.GetNSString();
 	}
 }
@@ -2043,3 +2048,28 @@ void FBackgroundURLSessionHandler::GetDownloadDebugText(const uint64 DownloadId,
 	}
 }
 #endif
+
+void FBackgroundURLSessionHandler::CleanupOldTempDownloads()
+{
+	IPlatformFile& PlatformFile = FPlatformFileManager::Get().GetPlatformFile();
+	const FString TempPath = FBackgroundHttpFileHashHelper::GetTemporaryRootPath();
+	UE_LOG(LogIOSBackgroundDownload, VeryVerbose, TEXT("CleanupOldTempDownloads: cleaning up %s"), *TempPath);
+	if (PlatformFile.DirectoryExists(*TempPath)) 
+	{
+		const FString BranchName = BuildSettings::GetBranchName();
+		UE_LOG(LogIOSBackgroundDownload, VeryVerbose, TEXT("CleanupOldTempDownloads: keeping directory: %s"), *BranchName);
+		PlatformFile.IterateDirectory(*TempPath, [&BranchName, &PlatformFile](const TCHAR* FilenameOrDirectory, bool bIsDirectory) -> bool {
+			UE_LOG(LogIOSBackgroundDownload, VeryVerbose, TEXT("CleanupOldTempDownloads: considering: %s | %d"), FilenameOrDirectory, (int)bIsDirectory);
+			if (bIsDirectory)
+			{
+				const FString DirectoryName = FPaths::GetCleanFilename(FilenameOrDirectory);
+				if (DirectoryName != TEXT("URLMap") && DirectoryName != BranchName)
+				{
+					UE_LOG(LogIOSBackgroundDownload, VeryVerbose, TEXT("CleanupOldTempDownloads: deleting: %s"), FilenameOrDirectory);
+					PlatformFile.DeleteDirectoryRecursively(FilenameOrDirectory);
+				}
+			}
+			return true;
+		});
+	}
+}
diff --git a/Engine/Source/Runtime/ApplicationCore/Public/IOS/IOSBackgroundURLSessionHandler.h b/Engine/Source/Runtime/ApplicationCore/Public/IOS/IOSBackgroundURLSessionHandler.h
--- a/Engine/Source/Runtime/ApplicationCore/Public/IOS/IOSBackgroundURLSessionHandler.h
+++ b/Engine/Source/Runtime/ApplicationCore/Public/IOS/IOSBackgroundURLSessionHandler.h
@@ -133,4 +133,7 @@ public:
 #if !UE_BUILD_SHIPPING
 	static void GetDownloadDebugText(const uint64 DownloadId, TArray<FString>& Output);
 #endif
+	
+	// Cleans up temp downloads placed in any subdirectory other than the current build one's
+	static void CleanupOldTempDownloads();
 };
diff --git a/Engine/Source/Runtime/Online/BackgroundHTTP/Private/IOS/ApplePlatformBackgroundHttpManager.cpp b/Engine/Source/Runtime/Online/BackgroundHTTP/Private/IOS/ApplePlatformBackgroundHttpManager.cpp
--- a/Engine/Source/Runtime/Online/BackgroundHTTP/Private/IOS/ApplePlatformBackgroundHttpManager.cpp
+++ b/Engine/Source/Runtime/Online/BackgroundHTTP/Private/IOS/ApplePlatformBackgroundHttpManager.cpp
@@ -25,6 +25,8 @@ FApplePlatformBackgroundHttpManager::FApplePlatformBackgroundHttpManager()
 {
 	OnDownloadCompletedHandle = FBackgroundURLSessionHandler::OnDownloadCompleted.AddRaw(this, &FApplePlatformBackgroundHttpManager::OnDownloadCompleted);
 	OnDownloadMetricsExtendedHandle = FBackgroundURLSessionHandler::OnDownloadMetricsExtended.AddRaw(this, &FApplePlatformBackgroundHttpManager::OnDownloadMetricsExtended);
+	
+	FBackgroundURLSessionHandler::CleanupOldTempDownloads();
 }
 
 FApplePlatformBackgroundHttpManager::~FApplePlatformBackgroundHttpManager()
diff --git a/Engine/Source/Runtime/Online/BackgroundHTTPFileHash/Private/BackgroundHttpFileHashHelper.cpp b/Engine/Source/Runtime/Online/BackgroundHTTPFileHash/Private/BackgroundHttpFileHashHelper.cpp
--- a/Engine/Source/Runtime/Online/BackgroundHTTPFileHash/Private/BackgroundHttpFileHashHelper.cpp
+++ b/Engine/Source/Runtime/Online/BackgroundHTTPFileHash/Private/BackgroundHttpFileHashHelper.cpp
@@ -97,7 +97,7 @@ FString FBackgroundHttpFileHashHelper::GetFullPathOfTempFilename(const FString&
 {
 	// don't take lock here, as other methods taking lock call it
 	FString FullFileName = FPaths::SetExtension(FPaths::GetCleanFilename(TempFilename), GetTempFileExtension());
-#if PLATFORM_ANDROID
+#if PLATFORM_ANDROID || PLATFORM_IOS
 	FullFileName = FPaths::Combine(BuildSettings::GetBranchName(), FullFileName);
 #endif
 	return FPaths::Combine(GetTemporaryRootPath(), FullFileName);
