Added missing check to make duplicate URLs from different requests work on IOS

#rnx
#Synced-CL Unknown

diff --git a/Engine/Source/Runtime/ApplicationCore/Private/IOS/IOSBackgroundURLSessionHandler.cpp b/Engine/Source/Runtime/ApplicationCore/Private/IOS/IOSBackgroundURLSessionHandler.cpp
--- a/Engine/Source/Runtime/ApplicationCore/Private/IOS/IOSBackgroundURLSessionHandler.cpp
+++ b/Engine/Source/Runtime/ApplicationCore/Private/IOS/IOSBackgroundURLSessionHandler.cpp
@@ -971,15 +971,19 @@ static constexpr NSInteger TypicalSizeForDownloadResponseHeader = 400; // bytes
 				if ((IterTask.originalRequest != nil && [IterTask.originalRequest.URL.path isEqualToString:AssetPath])
 					|| (IterTask.currentRequest != nil && [IterTask.currentRequest.URL.path isEqualToString:AssetPath]))
 				{
-					ExistingDownloadId = IterKey.unsignedIntegerValue;
-					ExistingTask = [IterTask retain]; // Retain in case if task gets killed in another thread.
-					*IterStop = YES;
+					FBackgroundNSURLSessionDownloadTaskData* IterTaskData = [FBackgroundNSURLSessionDownloadTaskData TaskDataFromSerializedString:IterTask.taskDescription];
+					if (TaskData == nil || [IterTaskData.RequestID isEqualToString:RequestID])
+					{
+						ExistingDownloadId = IterKey.unsignedIntegerValue;
+						ExistingTask = [IterTask retain]; // Retain in case if task gets killed in another thread.
+						*IterStop = YES;
+					}
 				}
 			}];
 
 			if (ExistingDownloadId != InvalidDownloadId && ExistingTask != nil)
-			{
-				UE_DNLD_LOG(@"Found existing download task for path '%@' with DownloadId %lu", AssetPath, ExistingDownloadId);
+			{				
+				UE_DNLD_LOG(@"Found existing download task for path '%@' requestID '%@' with DownloadId %lu", AssetPath, RequestID, ExistingDownloadId);
 				
 				// Update existing task state to new one, to reset retry counters, cdn links, etc.
 				[ExistingTask setTaskDescription:[TaskData ToSerializedString]];
