Manually scaling initial Splash Window on Linux

#Synced-CL 48328183

diff --git a/Engine/Source/Runtime/ApplicationCore/Private/Linux/LinuxPlatformSplash.cpp b/Engine/Source/Runtime/ApplicationCore/Private/Linux/LinuxPlatformSplash.cpp
--- a/Engine/Source/Runtime/ApplicationCore/Private/Linux/LinuxPlatformSplash.cpp
+++ b/Engine/Source/Runtime/ApplicationCore/Private/Linux/LinuxPlatformSplash.cpp
@@ -504,9 +504,21 @@ bool FLinuxSplashState::InitSplashResources(const FText &AppName, const FString
 		UE_LOG(LogHAL, Error, TEXT("Failed to create SDL properties"));
 		return false;
 	}
+
+	SDL_DisplayID DisplayID = SDL_GetPrimaryDisplay();
+	float SplashScale = SDL_GetDisplayContentScale(DisplayID);
+
+	if (SplashScale == 0.0f)
+	{
+		SplashScale = 1.0f;
+	}
+
+	const int ScaledSplashWidth = (int)((float)SplashWidth * SplashScale);
+	const int ScaledSplashHeight = (int)((float)SplashHeight * SplashScale);
+
 	SDL_SetStringProperty(Props, SDL_PROP_WINDOW_CREATE_TITLE_STRING, TCHAR_TO_UTF8(*AppName.ToString()));
-	SDL_SetNumberProperty(Props, SDL_PROP_WINDOW_CREATE_WIDTH_NUMBER, SplashWidth);
-	SDL_SetNumberProperty(Props, SDL_PROP_WINDOW_CREATE_HEIGHT_NUMBER, SplashHeight);
+	SDL_SetNumberProperty(Props, SDL_PROP_WINDOW_CREATE_WIDTH_NUMBER, ScaledSplashWidth);
+	SDL_SetNumberProperty(Props, SDL_PROP_WINDOW_CREATE_HEIGHT_NUMBER, ScaledSplashHeight);
 	SDL_SetNumberProperty(Props, SDL_PROP_WINDOW_CREATE_X_NUMBER, SDL_WINDOWPOS_CENTERED_DISPLAY(0));
 	SDL_SetNumberProperty(Props, SDL_PROP_WINDOW_CREATE_Y_NUMBER, SDL_WINDOWPOS_CENTERED_DISPLAY(0));
 	SDL_SetBooleanProperty(Props, SDL_PROP_WINDOW_CREATE_BORDERLESS_BOOLEAN, true);
