Inputs: Added missing implementations for IInputInterface::SetLightColor/ResetLightColor on Windows.

It fixes controlling the light bar (on controllers supporting it) when playing on Windows.

#Synced-CL 48634164

diff --git a/Engine/Source/Runtime/ApplicationCore/Private/Windows/WindowsApplication.cpp b/Engine/Source/Runtime/ApplicationCore/Private/Windows/WindowsApplication.cpp
--- a/Engine/Source/Runtime/ApplicationCore/Private/Windows/WindowsApplication.cpp
+++ b/Engine/Source/Runtime/ApplicationCore/Private/Windows/WindowsApplication.cpp
@@ -4253,6 +4253,20 @@ void FWindowsApplication::SetHapticFeedbackValues(int32 ControllerId, int32 Hand
 	}
 }
 
+void FWindowsApplication::SetLightColor(int32 ControllerId, FColor Color)
+{
+	const FInputDeviceLightColorProperty Property(true, Color);
+
+	SetDeviceProperty(ControllerId, &Property);
+}
+
+void FWindowsApplication::ResetLightColor(int32 ControllerId)
+{
+	const FInputDeviceLightColorProperty Property(false, FColor::Black);
+
+	SetDeviceProperty(ControllerId, &Property);
+}
+
 void FWindowsApplication::SetDeviceProperty(int32 ControllerId, const FInputDeviceProperty* Property)
 {
 	if (FApp::UseVRFocus() && !FApp::HasVRFocus())
diff --git a/Engine/Source/Runtime/ApplicationCore/Public/Windows/WindowsApplication.h b/Engine/Source/Runtime/ApplicationCore/Public/Windows/WindowsApplication.h
--- a/Engine/Source/Runtime/ApplicationCore/Public/Windows/WindowsApplication.h
+++ b/Engine/Source/Runtime/ApplicationCore/Public/Windows/WindowsApplication.h
@@ -410,8 +410,8 @@ public:
 	APPLICATIONCORE_API virtual void SetForceFeedbackChannelValue (int32 ControllerId, FForceFeedbackChannelType ChannelType, float Value) override;
 	APPLICATIONCORE_API virtual void SetForceFeedbackChannelValues(int32 ControllerId, const FForceFeedbackValues &Values) override;
 	APPLICATIONCORE_API virtual void SetHapticFeedbackValues(int32 ControllerId, int32 Hand, const FHapticFeedbackValues& Values) override;
-	virtual void SetLightColor(int32 ControllerId, FColor Color) override { }
-	virtual void ResetLightColor(int32 ControllerId) override { }
+	APPLICATIONCORE_API virtual void SetLightColor(int32 ControllerId, FColor Color) override;
+	APPLICATIONCORE_API virtual void ResetLightColor(int32 ControllerId) override;
 	APPLICATIONCORE_API virtual void SetDeviceProperty(int32 ControllerId, const FInputDeviceProperty* Property) override;
 
 	// Touchpad Sensitivity
