Fix an deprecation warning in xcode 26

#rnx

diff --git a/Engine/Source/Runtime/ApplicationCore/Private/IOS/IOSAppDelegate.cpp b/Engine/Source/Runtime/ApplicationCore/Private/IOS/IOSAppDelegate.cpp
--- a/Engine/Source/Runtime/ApplicationCore/Private/IOS/IOSAppDelegate.cpp
+++ b/Engine/Source/Runtime/ApplicationCore/Private/IOS/IOSAppDelegate.cpp
@@ -740,7 +740,7 @@ static IOSAppDelegate* CachedDelegate = nil;
 
 			if ([self IsFeatureActive:EAudioFeature::BluetoothMicrophone])
 			{
-				Options |= AVAudioSessionCategoryOptionAllowBluetooth;
+				Options |= AVAudioSessionCategoryOptionAllowBluetoothHFP;
 			}
 #endif
 
@@ -765,7 +765,7 @@ static IOSAppDelegate* CachedDelegate = nil;
 #if !PLATFORM_TVOS
 			if ([self IsFeatureActive:EAudioFeature::BluetoothMicrophone])
 			{
-				Options |= AVAudioSessionCategoryOptionAllowBluetooth;
+				Options |= AVAudioSessionCategoryOptionAllowBluetoothHFP;
 			}
 #endif
 		}
