move to more efficient and direct game controller updates vs controller capture per frame.

diff --git a/Engine/Source/Runtime/ApplicationCore/Private/Apple/AppleControllerInterface.cpp b/Engine/Source/Runtime/ApplicationCore/Private/Apple/AppleControllerInterface.cpp
--- a/Engine/Source/Runtime/ApplicationCore/Private/Apple/AppleControllerInterface.cpp
+++ b/Engine/Source/Runtime/ApplicationCore/Private/Apple/AppleControllerInterface.cpp
@@ -39,17 +39,52 @@ FAppleControllerInterface::FAppleControllerInterface( const TSharedRef< FGeneric
 		{
 			FUserController& UserController = Controllers[ControllerIndex];
 			UserController.PlayerIndex = PlayerIndex::PlayerUnset;
+
+			// Setup the buttons for this controller. We pre-allocate for the worse case of MAX_NUM_CONTROLLERS (4) connected controllers.
+			// Doing it here avoids creating a new TMap on controller connect and destruction on disconnect which
+			// could happen many times during a play session. 
+			// Allocate for both current and previous captured controller state.
+			for (int32 ActiveIndex = 0; ActiveIndex < MAX_CONTROLLER_CAPTURE_STATES; ActiveIndex++)
+			{
+				ControllerButtons[ActiveIndex][ControllerIndex].Add(FGamepadKeyNames::FaceButtonLeft);
+				ControllerButtons[ActiveIndex][ControllerIndex].Add(FGamepadKeyNames::FaceButtonBottom);
+				ControllerButtons[ActiveIndex][ControllerIndex].Add(FGamepadKeyNames::FaceButtonRight);
+				ControllerButtons[ActiveIndex][ControllerIndex].Add(FGamepadKeyNames::FaceButtonTop);
+				ControllerButtons[ActiveIndex][ControllerIndex].Add(FGamepadKeyNames::LeftShoulder);
+				ControllerButtons[ActiveIndex][ControllerIndex].Add(FGamepadKeyNames::RightShoulder);
+				ControllerButtons[ActiveIndex][ControllerIndex].Add(FGamepadKeyNames::LeftTriggerThreshold);
+				ControllerButtons[ActiveIndex][ControllerIndex].Add(FGamepadKeyNames::RightTriggerThreshold);
+				ControllerButtons[ActiveIndex][ControllerIndex].Add(FGamepadKeyNames::DPadUp);
+				ControllerButtons[ActiveIndex][ControllerIndex].Add(FGamepadKeyNames::DPadDown);
+				ControllerButtons[ActiveIndex][ControllerIndex].Add(FGamepadKeyNames::DPadRight);
+				ControllerButtons[ActiveIndex][ControllerIndex].Add(FGamepadKeyNames::DPadLeft);
+				ControllerButtons[ActiveIndex][ControllerIndex].Add(FGamepadKeyNames::LeftThumb);
+				ControllerButtons[ActiveIndex][ControllerIndex].Add(FGamepadKeyNames::RightThumb);
+				ControllerButtons[ActiveIndex][ControllerIndex].Add(FGamepadKeyNames::SpecialRight);
+				ControllerButtons[ActiveIndex][ControllerIndex].Add(FGamepadKeyNames::SpecialLeft);
+				ControllerButtons[ActiveIndex][ControllerIndex].Add(FGamepadKeyNames::LeftStickUp);
+				ControllerButtons[ActiveIndex][ControllerIndex].Add(FGamepadKeyNames::LeftStickDown);
+				ControllerButtons[ActiveIndex][ControllerIndex].Add(FGamepadKeyNames::LeftStickRight);
+				ControllerButtons[ActiveIndex][ControllerIndex].Add(FGamepadKeyNames::LeftStickLeft);
+				ControllerButtons[ActiveIndex][ControllerIndex].Add(FGamepadKeyNames::RightAnalogX);
+				ControllerButtons[ActiveIndex][ControllerIndex].Add(FGamepadKeyNames::RightAnalogY);
+				ControllerButtons[ActiveIndex][ControllerIndex].Add(FGamepadKeyNames::LeftAnalogX);
+				ControllerButtons[ActiveIndex][ControllerIndex].Add(FGamepadKeyNames::LeftAnalogY);
+				ControllerButtons[ActiveIndex][ControllerIndex].Add(FGamepadKeyNames::RightTriggerAnalog);
+				ControllerButtons[ActiveIndex][ControllerIndex].Add(FGamepadKeyNames::LeftTriggerAnalog);
+			}
 		}
-		
+
 		for (GCController* Cont in [GCController controllers])
 		{
 			HandleConnection(Cont);
 		}
-		
+
 		NSNotificationCenter* notificationCenter = [NSNotificationCenter defaultCenter];
-		
-		// Not in an operation queue, [NSOperationQueue currentQueue] will return nil on macOS and iOS
-		// Notification callback will always be on the app main thread - defer events for Unreal Engine update thread
+
+		// Not in an operation queue, [NSOperationQueue currentQueue] will return
+		// nil on macOS and iOS Notification callback will always be on the app main
+		// thread - defer events for Unreal Engine update thread
 
 		[notificationCenter addObserverForName:GCControllerDidDisconnectNotification object:nil queue:nil usingBlock:^(NSNotification* Notification)
 		{
@@ -84,7 +119,7 @@ void FAppleControllerInterface::SignalEvent(EAppleControllerEventType InEventTyp
 void FAppleControllerInterface::Tick( float DeltaTime )
 {
 	FScopeLock Lock(&DeferredEventCS);
-	
+
 	for(uint32_t Index = 0;Index < DeferredEvents.Num();++Index)
 	{
 		FDeferredAppleControllerEvent& Event = DeferredEvents[Index];
@@ -119,51 +154,51 @@ void FAppleControllerInterface::Tick( float DeltaTime )
 
 void FAppleControllerInterface::SetControllerType(uint32 ControllerIndex)
 {
-    GCController *Controller = Controllers[ControllerIndex].Controller;
+	GCController *Controller = Controllers[ControllerIndex].Controller;
 
-    if ([Controller.productCategory isEqualToString:@"DualShock 4"])
-    {
-    	// validate it's of class GCDualShockGamepad so we can cast it later
-	    if ([Controller.extendedGamepad isKindOfClass:[GCDualShockGamepad class]])
-	    {
+	if ([Controller.productCategory isEqualToString:@"DualShock 4"])
+	{
+		// validate it's of class GCDualShockGamepad so we can cast it later
+		if ([Controller.extendedGamepad isKindOfClass:[GCDualShockGamepad class]])
+		{
 			Controllers[ControllerIndex].ControllerType = ControllerType::DualShockGamepad;
-	    }
-    	else
-    	{
+		}
+		else
+		{
 			UE_LOG(LogAppleController, Warning, TEXT("DualShock4 detected, but invalid class"));
-		    Controllers[ControllerIndex].ControllerType = ControllerType::ExtendedGamepad;
-    	}
-    }
-    else if ([Controller.productCategory isEqualToString:@"Xbox One"])
-    {
-        Controllers[ControllerIndex].ControllerType = ControllerType::XboxGamepad;
-    }
-    else if ([Controller.productCategory isEqualToString:@"DualSense"])
-    {
-    	// validate it is of class GCDualSenseGamepad so we can cast it later
-	    if ([Controller.extendedGamepad isKindOfClass:[GCDualSenseGamepad class]])
-	    {
+			Controllers[ControllerIndex].ControllerType = ControllerType::ExtendedGamepad;
+		}
+	}
+	else if ([Controller.productCategory isEqualToString:@"Xbox One"])
+	{
+		Controllers[ControllerIndex].ControllerType = ControllerType::XboxGamepad;
+	}
+	else if ([Controller.productCategory isEqualToString:@"DualSense"])
+	{
+		// validate it is of class GCDualSenseGamepad so we can cast it later
+		if ([Controller.extendedGamepad isKindOfClass:[GCDualSenseGamepad class]])
+		{
 			Controllers[ControllerIndex].ControllerType = ControllerType::DualSenseGamepad;
-	    }
-    	else
-    	{
+		}
+		else
+		{
 			UE_LOG(LogAppleController, Warning, TEXT("DualSense detected, but invalid class"));
-		    Controllers[ControllerIndex].ControllerType = ControllerType::ExtendedGamepad;
-    	}
-    }
-    else if (Controller.extendedGamepad != nil)
-    {
-        Controllers[ControllerIndex].ControllerType = ControllerType::ExtendedGamepad;
-    }
-    else if (Controller.microGamepad != nil)
-    {
-        Controllers[ControllerIndex].ControllerType = ControllerType::SiriRemote;
-    }
-    else
-    {
-        Controllers[ControllerIndex].ControllerType = ControllerType::Unassigned;
-        UE_LOG(LogAppleController, Warning, TEXT("Controller type is not recognized"));
-    }
+			Controllers[ControllerIndex].ControllerType = ControllerType::ExtendedGamepad;
+		}
+	}
+	else if (Controller.extendedGamepad != nil)
+	{
+		Controllers[ControllerIndex].ControllerType = ControllerType::ExtendedGamepad;
+	}
+	else if (Controller.microGamepad != nil)
+	{
+		Controllers[ControllerIndex].ControllerType = ControllerType::SiriRemote;
+	}
+	else
+	{
+		Controllers[ControllerIndex].ControllerType = ControllerType::Unassigned;
+		UE_LOG(LogAppleController, Warning, TEXT("Controller type is not recognized"));
+	}
 }
 
 FName FAppleControllerInterface::GetControllerTypeName(const ControllerType InControllerType)
@@ -187,6 +222,7 @@ FName FAppleControllerInterface::GetControllerTypeName(const ControllerType InCo
 			ControllerTypeName = DualSenseControllerTypeName;
 			break;
 		default:
+			UE_LOG(LogAppleController, Warning, TEXT("[GetControllerTypeName] Controller type is not recognized"));
 			break;
 	}
 
@@ -195,36 +231,124 @@ FName FAppleControllerInterface::GetControllerTypeName(const ControllerType InCo
 
 void FAppleControllerInterface::SetCurrentController(GCController* Controller)
 {
-    int32 ControllerIndex = 0;
-    PlayerIndex PreviousIndex = PlayerIndex::PlayerUnset;
+	int32 ControllerIndex = 0;
+	PlayerIndex PreviousIndex = PlayerIndex::PlayerUnset;
 
 	if (Controller == nil)
 	{
 		return;
 	}
+
+	for (ControllerIndex = 0; ControllerIndex < UE_ARRAY_COUNT(Controllers); ControllerIndex++)
+	{
+		if (Controllers[ControllerIndex].Controller == Controller)
+		{
+			if (Controllers[ControllerIndex].PlayerIndex == PlayerIndex::PlayerOne)
+			{
+				// Already set as CurrentController
+				return;
+			}
+			PreviousIndex = Controllers[ControllerIndex].PlayerIndex;
+			Controllers[ControllerIndex].PlayerIndex = PlayerIndex::PlayerOne;
+		}
+	}
 	
-    for (ControllerIndex = 0; ControllerIndex < UE_ARRAY_COUNT(Controllers); ControllerIndex++)
-    {
-        if (Controllers[ControllerIndex].Controller == Controller)
-        {
-            if (Controllers[ControllerIndex].PlayerIndex == PlayerIndex::PlayerOne)
-            {
-                // Already set as CurrentController
-                return;
-            }
-            PreviousIndex = Controllers[ControllerIndex].PlayerIndex;
-            Controllers[ControllerIndex].PlayerIndex = PlayerIndex::PlayerOne;
-        }
-    }
-    
-    for (ControllerIndex = 0; ControllerIndex < UE_ARRAY_COUNT(Controllers); ControllerIndex++)
-    {
-        if (Controllers[ControllerIndex].PlayerIndex == PlayerIndex::PlayerOne && Controllers[ControllerIndex].Controller != Controller)
-        {
-            // The old PlayerOne, should swap place with the new PlayerOne
-            Controllers[ControllerIndex].PlayerIndex = PreviousIndex;
-        }
-    }
+	for (ControllerIndex = 0; ControllerIndex < UE_ARRAY_COUNT(Controllers); ControllerIndex++)
+	{
+		if (Controllers[ControllerIndex].PlayerIndex == PlayerIndex::PlayerOne &&
+			Controllers[ControllerIndex].Controller != Controller)
+		{
+			// The old PlayerOne, should swap place with the new PlayerOne
+			Controllers[ControllerIndex].PlayerIndex = PreviousIndex;
+		}
+	}
+}
+
+void FAppleControllerInterface::CaptureControllerState(const GCController* Controller, uint32 ControllerIndex)
+{
+	if (!bAllowControllers || Controller == nil)
+	{
+		return;
+	}
+
+	GCController* Cont = Controllers[ControllerIndex].Controller;
+	TMap<FName, FControllerButton>& State = ControllerButtons[CurrentButtonIndex][ControllerIndex];
+
+	switch (Controllers[ControllerIndex].ControllerType)
+	{
+		case ControllerType::ExtendedGamepad:
+		case ControllerType::DualShockGamepad:
+		case ControllerType::XboxGamepad:
+		case ControllerType::DualSenseGamepad:
+		{
+			// Capture buttons state
+			const GCExtendedGamepad* ExtendedGamepad = Cont.extendedGamepad;
+			State[FGamepadKeyNames::FaceButtonLeft].bPressed = ExtendedGamepad.buttonX.pressed;
+			State[FGamepadKeyNames::FaceButtonBottom].bPressed = ExtendedGamepad.buttonA.pressed;
+			State[FGamepadKeyNames::FaceButtonRight].bPressed = ExtendedGamepad.buttonB.pressed;
+			State[FGamepadKeyNames::FaceButtonTop].bPressed = ExtendedGamepad.buttonY.pressed;
+			State[FGamepadKeyNames::LeftShoulder].bPressed = ExtendedGamepad.leftShoulder.pressed;
+			State[FGamepadKeyNames::RightShoulder].bPressed = ExtendedGamepad.rightShoulder.pressed;
+			State[FGamepadKeyNames::LeftTriggerThreshold].bPressed = ExtendedGamepad.leftTrigger.pressed;
+			State[FGamepadKeyNames::RightTriggerThreshold].bPressed = ExtendedGamepad.rightTrigger.pressed;
+			State[FGamepadKeyNames::DPadUp].bPressed = ExtendedGamepad.dpad.up.pressed;
+			State[FGamepadKeyNames::DPadDown].bPressed = ExtendedGamepad.dpad.down.pressed;
+			State[FGamepadKeyNames::DPadRight].bPressed = ExtendedGamepad.dpad.right.pressed;
+			State[FGamepadKeyNames::DPadLeft].bPressed = ExtendedGamepad.dpad.left.pressed;
+			State[FGamepadKeyNames::LeftThumb].bPressed = ExtendedGamepad.leftThumbstickButton.pressed;
+			State[FGamepadKeyNames::RightThumb].bPressed = ExtendedGamepad.rightThumbstickButton.pressed;
+			State[FGamepadKeyNames::SpecialRight].bPressed = ExtendedGamepad.buttonMenu.pressed;
+
+			// DualSense and DualShock4 use the touchpad button for the "SpecialLeft" button
+			if (Controllers[ControllerIndex].ControllerType == ControllerType::DualSenseGamepad)
+			{
+				// SetControllerType() verified that ExtendedGamepad is of type GCDualSenseGamepad
+				State[FGamepadKeyNames::SpecialLeft].bPressed = static_cast<GCDualShockGamepad*>(ExtendedGamepad).touchpadButton.pressed;
+			}
+			else if (Controllers[ControllerIndex].ControllerType == ControllerType::DualShockGamepad)
+			{
+				State[FGamepadKeyNames::SpecialLeft].bPressed = static_cast<GCDualShockGamepad*>(ExtendedGamepad).touchpadButton.pressed;
+			}
+			else
+			{
+				State[FGamepadKeyNames::SpecialLeft].bPressed = ExtendedGamepad.buttonOptions.pressed;
+			}
+
+			// Capture axis values
+			State[FGamepadKeyNames::LeftAnalogX].AxisValue = ExtendedGamepad.leftThumbstick.xAxis.value;
+			State[FGamepadKeyNames::LeftAnalogY].AxisValue = ExtendedGamepad.leftThumbstick.yAxis.value;
+			State[FGamepadKeyNames::RightAnalogX].AxisValue = ExtendedGamepad.rightThumbstick.xAxis.value;
+			State[FGamepadKeyNames::RightAnalogY].AxisValue = ExtendedGamepad.rightThumbstick.yAxis.value;
+			State[FGamepadKeyNames::RightTriggerAnalog].AxisValue = ExtendedGamepad.rightTrigger.value;
+			State[FGamepadKeyNames::LeftTriggerAnalog].AxisValue = ExtendedGamepad.leftTrigger.value;
+			break;
+		}
+
+		case ControllerType::SiriRemote:
+		{
+			const GCMicroGamepad* MicroGamepad = Cont.microGamepad;
+			State[FGamepadKeyNames::FaceButtonLeft].bPressed = MicroGamepad.buttonX.pressed;
+			State[FGamepadKeyNames::FaceButtonBottom].bPressed = MicroGamepad.buttonA.pressed;
+			State[FGamepadKeyNames::DPadUp].bPressed = MicroGamepad.dpad.up.pressed;
+			State[FGamepadKeyNames::DPadDown].bPressed = MicroGamepad.dpad.down.pressed;
+			State[FGamepadKeyNames::DPadRight].bPressed = MicroGamepad.dpad.right.pressed;
+			State[FGamepadKeyNames::DPadLeft].bPressed = MicroGamepad.dpad.left.pressed;
+			State[FGamepadKeyNames::LeftStickUp].bPressed = MicroGamepad.dpad.up.pressed;
+			State[FGamepadKeyNames::LeftStickDown].bPressed = MicroGamepad.dpad.down.pressed;
+			State[FGamepadKeyNames::LeftStickRight].bPressed = MicroGamepad.dpad.right.pressed;
+			State[FGamepadKeyNames::LeftStickLeft].bPressed = MicroGamepad.dpad.left.pressed;
+			State[FGamepadKeyNames::SpecialRight].bPressed = MicroGamepad.buttonMenu.pressed;
+			
+			// Capture "axis" values
+			State[FGamepadKeyNames::LeftAnalogX].AxisValue = MicroGamepad.dpad.xAxis.value;
+			State[FGamepadKeyNames::LeftAnalogY].AxisValue = MicroGamepad.dpad.yAxis.value;
+			break;
+		}
+
+		default:
+			UE_LOG(LogAppleController, Warning, TEXT("UNSUPPORTED CONTROLLER"));
+			break;
+	}
 }
 
 void FAppleControllerInterface::HandleConnection(GCController* Controller)
@@ -240,27 +364,37 @@ void FAppleControllerInterface::HandleConnection(GCController* Controller)
 	bool bFoundSlot = false;
 	for (int32 ControllerIndex = 0; ControllerIndex < UE_ARRAY_COUNT(Controllers); ControllerIndex++)
 	{
-        if (Controllers[ControllerIndex].ControllerType != ControllerType::Unassigned)
-        {
-            continue;
-        }
-        
-        Controllers[ControllerIndex].PlayerIndex = (PlayerIndex)ControllerIndex;
-        Controllers[ControllerIndex].Controller = [Controller retain];
-        SetControllerType(ControllerIndex);
-        
-        bFoundSlot = true;
-        
-        IPlatformInputDeviceMapper& DeviceMapper = IPlatformInputDeviceMapper::Get();
-        FPlatformUserId UserId = FGenericPlatformMisc::GetPlatformUserForUserIndex(Controllers[ControllerIndex].PlayerIndex);
-        FInputDeviceId DeviceId = INPUTDEVICEID_NONE;
-        DeviceMapper.RemapControllerIdToPlatformUserAndDevice(Controllers[ControllerIndex].PlayerIndex, OUT UserId, OUT DeviceId);
-        DeviceMapper.Internal_MapInputDeviceToUser(DeviceId, UserId, EInputDeviceConnectionState::Connected);
-        
-        UE_LOG(LogAppleController, Log, TEXT("New %s controller inserted, assigned to playerIndex %d"),
-               Controllers[ControllerIndex].ControllerType == ControllerType::SiriRemote
-               ? TEXT("SiriRemote") : *FString(Controller.productCategory), Controllers[ControllerIndex].PlayerIndex);
-        break;
+		if (Controllers[ControllerIndex].ControllerType != ControllerType::Unassigned)
+		{
+			continue;
+		}
+		
+		Controllers[ControllerIndex].PlayerIndex = (PlayerIndex)ControllerIndex;
+		Controllers[ControllerIndex].Controller = [Controller retain];
+		SetControllerType(ControllerIndex);
+		
+		// Reset values for contoller buttons and axis
+		for (int32 ActiveIndex = 0; ActiveIndex < MAX_CONTROLLER_CAPTURE_STATES; ActiveIndex++)
+		{
+			TMap<FName, FControllerButton>& ButtonStates = ControllerButtons[ActiveIndex][ControllerIndex];
+			for (TTuple<FName, FAppleControllerInterface::FControllerButton>& Pair : ButtonStates)
+			{
+				Pair.Value.Reset();
+			}
+		}
+		
+		bFoundSlot = true;
+		
+		IPlatformInputDeviceMapper& DeviceMapper = IPlatformInputDeviceMapper::Get();
+		FPlatformUserId UserId = FGenericPlatformMisc::GetPlatformUserForUserIndex(Controllers[ControllerIndex].PlayerIndex);
+		FInputDeviceId DeviceId = INPUTDEVICEID_NONE;
+		DeviceMapper.RemapControllerIdToPlatformUserAndDevice(Controllers[ControllerIndex].PlayerIndex, OUT UserId, OUT DeviceId);
+		DeviceMapper.Internal_MapInputDeviceToUser(DeviceId, UserId, EInputDeviceConnectionState::Connected);
+		
+		UE_LOG(LogAppleController, Log, TEXT("New %s controller inserted, assigned to playerIndex %d"),
+			   Controllers[ControllerIndex].ControllerType == ControllerType::SiriRemote
+			   ? TEXT("SiriRemote") : *FString(Controller.productCategory), Controllers[ControllerIndex].PlayerIndex);
+		break;
 	}
 	checkf(bFoundSlot, TEXT("Used a fifth controller somehow!"));
 }
@@ -270,38 +404,68 @@ void FAppleControllerInterface::HandleDisconnect(GCController* Controller)
 	// if we don't allow controllers, there could be unset player index here
 	if (!bAllowControllers || Controller == nil)
 	{
-        return;
+		return;
 	}
 	
-    for (int32 ControllerIndex = 0; ControllerIndex < UE_ARRAY_COUNT(Controllers); ControllerIndex++)
-    {
+	for (int32 ControllerIndex = 0; ControllerIndex < UE_ARRAY_COUNT(Controllers); ControllerIndex++)
+	{
 		FUserController& UserController = Controllers[ControllerIndex];
-        if (UserController.Controller == Controller)
-        {
-			// Player index of unset(-1) would indicate that it has become unset even though it is now trying to disconnect
-			// This can occur on iOS when bGameSupportsMultipleActiveControllers is false
+		if (UserController.Controller == Controller)
+		{
+			// Player index of unset(-1) would indicate that it has become unset even though it is now 
+			// trying to disconnect. This can occur on iOS when bGameSupportsMultipleActiveControllers is false
 			UE_LOG(LogAppleController, Log, TEXT("Controller for playerIndex %d, controller Index %d removed"), UserController.PlayerIndex, ControllerIndex);
-            
-            IPlatformInputDeviceMapper& DeviceMapper = IPlatformInputDeviceMapper::Get();
-            FPlatformUserId UserId = FGenericPlatformMisc::GetPlatformUserForUserIndex(UserController.PlayerIndex);
-            FInputDeviceId DeviceId = INPUTDEVICEID_NONE;
-            DeviceMapper.RemapControllerIdToPlatformUserAndDevice(UserController.PlayerIndex, OUT UserId, OUT DeviceId);
-            DeviceMapper.Internal_MapInputDeviceToUser(DeviceId, UserId, EInputDeviceConnectionState::Disconnected);
+
+			// Check if there's any buttons/joysticks currently pressed and/or active analog axis.
+			// If so, send off events as if the button was unpressed and axis back to 0.
+			// This avoids issues like having a player stuck moving forward and/or stuck firing a weapon.
+			TMap<FName, FControllerButton>& CurrentButtons = ControllerButtons[CurrentButtonIndex][ControllerIndex];
+			TMap<FName, FControllerButton>& PrevButtons = ControllerButtons[PrevButtonIndex][ControllerIndex];
+			for (auto& Pair : CurrentButtons)
+			{
+				Pair.Value.Reset();
+				
+				const FControllerButton* PrevButtonState = PrevButtons.Find(Pair.Key);
+				if (Pair.Value.bPressed != PrevButtonState->bPressed)
+				{
+					HandleButtonGamepad(Pair.Key, ControllerIndex);
+				}
+				if (Pair.Value.AxisValue != PrevButtonState->AxisValue)
+				{
+					HandleAnalogGamepad(Pair.Key, ControllerIndex);
+				}
+			}
+			HandleVirtualButtonGamepad(FGamepadKeyNames::LeftStickRight, FGamepadKeyNames::LeftStickLeft, ControllerIndex);
+			HandleVirtualButtonGamepad(FGamepadKeyNames::LeftStickDown, FGamepadKeyNames::LeftStickUp, ControllerIndex);
+			HandleVirtualButtonGamepad(FGamepadKeyNames::RightStickLeft, FGamepadKeyNames::RightStickRight, ControllerIndex);
+			HandleVirtualButtonGamepad(FGamepadKeyNames::RightStickDown, FGamepadKeyNames::RightStickUp, ControllerIndex);
+			
+			// Clear out any old Previous captured button state. Above cleared the Current state.
+			for (auto& Pair : PrevButtons)
+			{
+				Pair.Value.Reset();
+			}
+			
+			IPlatformInputDeviceMapper& DeviceMapper = IPlatformInputDeviceMapper::Get();
+			FPlatformUserId UserId = FGenericPlatformMisc::GetPlatformUserForUserIndex(UserController.PlayerIndex);
+			FInputDeviceId DeviceId = INPUTDEVICEID_NONE;
+			DeviceMapper.RemapControllerIdToPlatformUserAndDevice(UserController.PlayerIndex, OUT UserId, OUT DeviceId);
+			DeviceMapper.Internal_MapInputDeviceToUser(DeviceId, UserId, EInputDeviceConnectionState::Disconnected);
 			
 			[UserController.Controller release];
-			[UserController.PreviousExtendedGamepad release];
 			
-            FMemory::Memzero(&UserController, sizeof(Controllers[ControllerIndex]));
-            UserController.PlayerIndex = PlayerIndex::PlayerUnset;
-            
-            return;
-        }
-    }
+			FMemory::Memzero(&UserController, sizeof(Controllers[ControllerIndex]));
+			UserController.PlayerIndex = PlayerIndex::PlayerUnset;
+			
+			return;
+		}
+	}
 }
 
 void FAppleControllerInterface::SendControllerEvents()
 {
-	@autoreleasepool{
+	@autoreleasepool
+	{
 	for(int32 i = 0; i < UE_ARRAY_COUNT(Controllers); ++i)
 	{
 		FUserController& Controller = Controllers[i];
@@ -316,11 +480,18 @@ void FAppleControllerInterface::SendControllerEvents()
 		IPlatformInputDeviceMapper& DeviceMapper = IPlatformInputDeviceMapper::Get();
 		FPlatformUserId UserId = FGenericPlatformMisc::GetPlatformUserForUserIndex(Controller.PlayerIndex);
 		FInputDeviceId DeviceId = DeviceMapper.GetPrimaryInputDeviceForUser(UserId);
+
 		///The FInputDeviceScope::HardwareDeviceIdentifier has to be one of values in UInputPlatformSettings::HardwareDevices.
 		// This is a temp solution with a hardcoded string which can be mapped to FHardwareDeviceIdentifier::DefaultGamepad.
 		// TODO: Future improvement is needed to acquire them by values in the iOS input device info.
 		FInputDeviceScope InputScope(nullptr, GetControllerTypeName(Controller.ControllerType), DeviceId.GetId(), HardwareDeviceIdentifier_DefaultGamepad);
-		
+
+		uint32 ButtonIndex = PrevButtonIndex;	
+		PrevButtonIndex = CurrentButtonIndex;
+		CurrentButtonIndex = ButtonIndex;
+
+		CaptureControllerState(ControllerImpl, i);
+
 		HandleButtonGamepad(FGamepadKeyNames::FaceButtonBottom, i);
 		HandleButtonGamepad(FGamepadKeyNames::FaceButtonLeft, i);
 		HandleButtonGamepad(FGamepadKeyNames::FaceButtonRight, i);
@@ -349,28 +520,7 @@ void FAppleControllerInterface::SendControllerEvents()
 		HandleVirtualButtonGamepad(FGamepadKeyNames::RightStickDown, FGamepadKeyNames::RightStickUp, i);
 		HandleButtonGamepad(FGamepadKeyNames::LeftThumb, i);
 		HandleButtonGamepad(FGamepadKeyNames::RightThumb, i);
-
-		// Update the saved previous captured game controller. This shouldn't fail, but it does occasionally in the
-		// GameController Framework, so add some extra error checking
-		GCController* CapturedController = [ControllerImpl capture];
-		if (CapturedController != nil)
-		{
-			GCExtendedGamepad* ExtendedGamepad = CapturedController.extendedGamepad;
-			if (ExtendedGamepad != nil)
-			{
-				[Controller.PreviousExtendedGamepad release];
-				Controller.PreviousExtendedGamepad = [ExtendedGamepad retain];
-			}
-			else
-			{
-				UE_LOG(LogAppleController, Warning, TEXT("Unable to access extendedGamepad"));
-		    }
 		}
-		else
-		{
-			UE_LOG(LogAppleController, Warning, TEXT("Unable to capture controller"));
-		}
-	}
 	} //@autoreleasepool
 }
 
@@ -392,34 +542,34 @@ bool FAppleControllerInterface::IsGamepadAttached() const
 
 GCControllerButtonInput* FAppleControllerInterface::GetGCControllerButton(const FGamepadKeyNames::Type& ButtonKey, uint32 ControllerIndex)
 {
-    GCController* Cont = Controllers[ControllerIndex].Controller;
-    
-    GCExtendedGamepad* ExtendedGamepad = Cont.extendedGamepad;
-    GCControllerButtonInput *ButtonToReturn = nullptr;
-
-    if (ButtonKey == FGamepadKeyNames::FaceButtonBottom){ButtonToReturn = ExtendedGamepad.buttonA;}
-    else if (ButtonKey == FGamepadKeyNames::FaceButtonRight){ButtonToReturn = ExtendedGamepad.buttonB;}
-    else if (ButtonKey == FGamepadKeyNames::FaceButtonLeft){ButtonToReturn = ExtendedGamepad.buttonX;}
-    else if (ButtonKey == FGamepadKeyNames::FaceButtonTop){ButtonToReturn = ExtendedGamepad.buttonY;}
-    else if (ButtonKey == FGamepadKeyNames::LeftShoulder){ButtonToReturn = ExtendedGamepad.leftShoulder;}
-    else if (ButtonKey == FGamepadKeyNames::RightShoulder){ButtonToReturn = ExtendedGamepad.rightShoulder;}
-    else if (ButtonKey == FGamepadKeyNames::LeftTriggerThreshold){ButtonToReturn = ExtendedGamepad.leftTrigger;}
-    else if (ButtonKey == FGamepadKeyNames::RightTriggerThreshold){ButtonToReturn = ExtendedGamepad.rightTrigger;}
-    else if (ButtonKey == FGamepadKeyNames::LeftTriggerAnalog){ButtonToReturn = ExtendedGamepad.leftTrigger;}
-    else if (ButtonKey == FGamepadKeyNames::RightTriggerAnalog){ButtonToReturn = ExtendedGamepad.rightTrigger;}
-    else if (ButtonKey == FGamepadKeyNames::LeftThumb){ButtonToReturn = ExtendedGamepad.leftThumbstickButton;}
-    else if (ButtonKey == FGamepadKeyNames::RightThumb){ButtonToReturn = ExtendedGamepad.rightThumbstickButton;}
-
-    return ButtonToReturn;
+	GCController* Cont = Controllers[ControllerIndex].Controller;
+	
+	GCExtendedGamepad* ExtendedGamepad = Cont.extendedGamepad;
+	GCControllerButtonInput *ButtonToReturn = nullptr;
+	
+	if (ButtonKey == FGamepadKeyNames::FaceButtonBottom){ButtonToReturn = ExtendedGamepad.buttonA;}
+	else if (ButtonKey == FGamepadKeyNames::FaceButtonRight){ButtonToReturn = ExtendedGamepad.buttonB;}
+	else if (ButtonKey == FGamepadKeyNames::FaceButtonLeft){ButtonToReturn = ExtendedGamepad.buttonX;}
+	else if (ButtonKey == FGamepadKeyNames::FaceButtonTop){ButtonToReturn = ExtendedGamepad.buttonY;}
+	else if (ButtonKey == FGamepadKeyNames::LeftShoulder){ButtonToReturn = ExtendedGamepad.leftShoulder;}
+	else if (ButtonKey == FGamepadKeyNames::RightShoulder){ButtonToReturn = ExtendedGamepad.rightShoulder;}
+	else if (ButtonKey == FGamepadKeyNames::LeftTriggerThreshold){ButtonToReturn = ExtendedGamepad.leftTrigger;}
+	else if (ButtonKey == FGamepadKeyNames::RightTriggerThreshold){ButtonToReturn = ExtendedGamepad.rightTrigger;}
+	else if (ButtonKey == FGamepadKeyNames::LeftTriggerAnalog){ButtonToReturn = ExtendedGamepad.leftTrigger;}
+	else if (ButtonKey == FGamepadKeyNames::RightTriggerAnalog){ButtonToReturn = ExtendedGamepad.rightTrigger;}
+	else if (ButtonKey == FGamepadKeyNames::LeftThumb){ButtonToReturn = ExtendedGamepad.leftThumbstickButton;}
+	else if (ButtonKey == FGamepadKeyNames::RightThumb){ButtonToReturn = ExtendedGamepad.rightThumbstickButton;}
+	
+	return ButtonToReturn;
 }
 
 const ControllerType FAppleControllerInterface::GetControllerType(uint32 ControllerIndex)
 {
-    if (Controllers[ControllerIndex].Controller != nullptr)
-    {
-        return Controllers[ControllerIndex].ControllerType;
-    }
-    return ControllerType::Unassigned;
+	if (Controllers[ControllerIndex].Controller != nullptr)
+	{
+		return Controllers[ControllerIndex].ControllerType;
+	}
+	return ControllerType::Unassigned;
 }
 
 void FAppleControllerInterface::HandleInputInternal(const FGamepadKeyNames::Type& UEButton, uint32 ControllerIndex, bool bIsPressed, bool bWasPressed)
@@ -428,10 +578,10 @@ void FAppleControllerInterface::HandleInputInternal(const FGamepadKeyNames::Type
     const float InitialRepeatDelay = 0.2f;
     const float RepeatDelay = 0.1f;
     GCController* Cont = Controllers[ControllerIndex].Controller;
-    
+
 	IPlatformInputDeviceMapper& DeviceMapper = IPlatformInputDeviceMapper::Get();
 	FPlatformUserId UserId = FGenericPlatformMisc::GetPlatformUserForUserIndex(Controllers[ControllerIndex].PlayerIndex);
-    FInputDeviceId DeviceId = DeviceMapper.GetPrimaryInputDeviceForUser(UserId);
+	FInputDeviceId DeviceId = DeviceMapper.GetPrimaryInputDeviceForUser(UserId);
 
     if (bWasPressed != bIsPressed)
     {
@@ -458,252 +608,139 @@ void FAppleControllerInterface::HandleInputInternal(const FGamepadKeyNames::Type
 
 void FAppleControllerInterface::HandleVirtualButtonGamepad(const FGamepadKeyNames::Type& UEButtonNegative, const FGamepadKeyNames::Type& UEButtonPositive, uint32 ControllerIndex)
 {
-    GCController* Cont = Controllers[ControllerIndex].Controller;
-    GCExtendedGamepad *ExtendedGamepad = Cont.extendedGamepad;
-    GCExtendedGamepad *ExtendedPreviousGamepad = Controllers[ControllerIndex].PreviousExtendedGamepad;;
-
-    // Send controller events any time we are passed the given input threshold similarly to PC/Console (see: XInputInterface.cpp)
-    const float RepeatDeadzone = 0.24f;
-    
-    bool bWasNegativePressed = false;
-    bool bNegativePressed = false;
-    bool bWasPositivePressed = false;
-    bool bPositivePressed = false;
-    
-    if (UEButtonNegative == FGamepadKeyNames::LeftStickLeft && UEButtonPositive == FGamepadKeyNames::LeftStickRight)
-    {
-        bWasNegativePressed = ExtendedPreviousGamepad != nil && ExtendedPreviousGamepad.leftThumbstick.xAxis.value <= -RepeatDeadzone;
-        bNegativePressed = ExtendedGamepad.leftThumbstick.xAxis.value <= -RepeatDeadzone;
-        bWasPositivePressed = ExtendedPreviousGamepad != nil && ExtendedPreviousGamepad.leftThumbstick.xAxis.value >= RepeatDeadzone;
-        bPositivePressed = ExtendedGamepad.leftThumbstick.xAxis.value >= RepeatDeadzone;
+	GCController* Cont = Controllers[ControllerIndex].Controller;
 
-        HandleInputInternal(FGamepadKeyNames::LeftStickDown, ControllerIndex, bNegativePressed, bWasNegativePressed);
-        HandleInputInternal(FGamepadKeyNames::LeftStickUp, ControllerIndex, bPositivePressed, bWasPositivePressed);
-    }
-    else if (UEButtonNegative == FGamepadKeyNames::LeftStickDown && UEButtonPositive == FGamepadKeyNames::LeftStickUp)
-    {
-        bWasNegativePressed = ExtendedPreviousGamepad != nil && ExtendedPreviousGamepad.leftThumbstick.yAxis.value <= -RepeatDeadzone;
-        bNegativePressed = ExtendedGamepad.leftThumbstick.yAxis.value <= -RepeatDeadzone;
-        bWasPositivePressed = ExtendedPreviousGamepad != nil && ExtendedPreviousGamepad.leftThumbstick.yAxis.value >= RepeatDeadzone;
-        bPositivePressed = ExtendedGamepad.leftThumbstick.yAxis.value >= RepeatDeadzone;
+	// Send controller events any time we are passed the given input threshold similarly to PC/Console (see: XInputInterface.cpp)
+	const float RepeatDeadzone = 0.24f;
 
-        HandleInputInternal(FGamepadKeyNames::LeftStickDown, ControllerIndex, bNegativePressed, bWasNegativePressed);
-        HandleInputInternal(FGamepadKeyNames::LeftStickUp, ControllerIndex, bPositivePressed, bWasPositivePressed);
-    }
-    else if (UEButtonNegative == FGamepadKeyNames::RightStickLeft && UEButtonPositive == FGamepadKeyNames::RightStickRight)
-    {
-        bWasNegativePressed = ExtendedPreviousGamepad != nil && ExtendedPreviousGamepad.rightThumbstick.xAxis.value <= -RepeatDeadzone;
-        bNegativePressed = ExtendedGamepad.rightThumbstick.xAxis.value <= -RepeatDeadzone;
-        bWasPositivePressed = ExtendedPreviousGamepad != nil && ExtendedPreviousGamepad.rightThumbstick.xAxis.value >= RepeatDeadzone;
-        bPositivePressed = ExtendedGamepad.rightThumbstick.xAxis.value >= RepeatDeadzone;
+	bool bWasNegativePressed = false;
+	bool bNegativePressed = false;
+	bool bWasPositivePressed = false;
+	bool bPositivePressed = false;
 
-        HandleInputInternal(FGamepadKeyNames::LeftStickDown, ControllerIndex, bNegativePressed, bWasNegativePressed);
-        HandleInputInternal(FGamepadKeyNames::LeftStickUp, ControllerIndex, bPositivePressed, bWasPositivePressed);
-    }
-    else if (UEButtonNegative == FGamepadKeyNames::RightStickDown && UEButtonPositive == FGamepadKeyNames::RightStickUp)
-    {
-        bWasNegativePressed = ExtendedPreviousGamepad != nil && ExtendedPreviousGamepad.rightThumbstick.yAxis.value <= -RepeatDeadzone;
-        bNegativePressed = ExtendedGamepad.rightThumbstick.yAxis.value <= -RepeatDeadzone;
-        bWasPositivePressed = ExtendedPreviousGamepad != nil && ExtendedPreviousGamepad.rightThumbstick.yAxis.value >= RepeatDeadzone;
-        bPositivePressed = ExtendedGamepad.rightThumbstick.yAxis.value >= RepeatDeadzone;
+	const TMap<FName, FControllerButton>& PrevState = ControllerButtons[PrevButtonIndex][ControllerIndex];
+	const TMap<FName, FControllerButton>& CurrentState = ControllerButtons[CurrentButtonIndex][ControllerIndex];
 
-        HandleInputInternal(FGamepadKeyNames::LeftStickDown, ControllerIndex, bNegativePressed, bWasNegativePressed);
-        HandleInputInternal(FGamepadKeyNames::LeftStickUp, ControllerIndex, bPositivePressed, bWasPositivePressed);
-    }
+	if (UEButtonNegative == FGamepadKeyNames::LeftStickLeft && UEButtonPositive == FGamepadKeyNames::LeftStickRight)
+	{
+		// Uses the "leftThumbstick.xAxis" of the controller, which is captured by "LeftAnalogX"
+		bWasNegativePressed = PrevState.Find(FGamepadKeyNames::LeftAnalogX)->AxisValue <= -RepeatDeadzone;
+		bNegativePressed = CurrentState.Find(FGamepadKeyNames::LeftAnalogX)->AxisValue <= -RepeatDeadzone;
+		bWasPositivePressed = PrevState.Find(FGamepadKeyNames::LeftAnalogX)->AxisValue >= RepeatDeadzone;
+		bPositivePressed = CurrentState.Find(FGamepadKeyNames::LeftAnalogX)->AxisValue >= RepeatDeadzone;
+		
+		HandleInputInternal(FGamepadKeyNames::LeftStickDown, ControllerIndex, bNegativePressed, bWasNegativePressed);
+		HandleInputInternal(FGamepadKeyNames::LeftStickUp, ControllerIndex, bPositivePressed, bWasPositivePressed);
+	}
+	else if (UEButtonNegative == FGamepadKeyNames::LeftStickDown && UEButtonPositive == FGamepadKeyNames::LeftStickUp)
+	{
+		// Uses the "leftThumbstick.yAxis" of the controller, which is captured by "LeftAnalogY"
+		bWasNegativePressed = PrevState.Find(FGamepadKeyNames::LeftAnalogY)->AxisValue <= -RepeatDeadzone;
+		bNegativePressed = CurrentState.Find(FGamepadKeyNames::LeftAnalogY)->AxisValue <= -RepeatDeadzone;
+		bWasPositivePressed = PrevState.Find(FGamepadKeyNames::LeftAnalogY)->AxisValue >= RepeatDeadzone;
+		bPositivePressed = CurrentState.Find(FGamepadKeyNames::LeftAnalogY)->AxisValue >= RepeatDeadzone;
+		
+		HandleInputInternal(FGamepadKeyNames::LeftStickDown, ControllerIndex, bNegativePressed, bWasNegativePressed);
+		HandleInputInternal(FGamepadKeyNames::LeftStickUp, ControllerIndex, bPositivePressed, bWasPositivePressed);
+	}
+	else if (UEButtonNegative == FGamepadKeyNames::RightStickLeft && UEButtonPositive == FGamepadKeyNames::RightStickRight)
+	{
+		// Uses the "rightThumbstick.xAxis" of the controller, which is captured by "RightAnalogX"
+		bWasNegativePressed = PrevState.Find(FGamepadKeyNames::RightAnalogX)->AxisValue <= -RepeatDeadzone;
+		bNegativePressed = CurrentState.Find(FGamepadKeyNames::RightAnalogX)->AxisValue <= -RepeatDeadzone;
+		bWasPositivePressed = PrevState.Find(FGamepadKeyNames::RightAnalogX)->AxisValue >= RepeatDeadzone;
+		bPositivePressed = CurrentState.Find(FGamepadKeyNames::RightAnalogX)->AxisValue >= RepeatDeadzone;
+		
+		HandleInputInternal(FGamepadKeyNames::LeftStickDown, ControllerIndex, bNegativePressed, bWasNegativePressed);
+		HandleInputInternal(FGamepadKeyNames::LeftStickUp, ControllerIndex, bPositivePressed, bWasPositivePressed);
+	}
+	else if (UEButtonNegative == FGamepadKeyNames::RightStickDown && UEButtonPositive == FGamepadKeyNames::RightStickUp)
+	{
+		// Uses the "rightThumbstick.yAxis" of the controller, which is captured by "RightAnalogY"
+		bWasNegativePressed = PrevState.Find(FGamepadKeyNames::RightAnalogY)->AxisValue <= -RepeatDeadzone;
+		bNegativePressed = CurrentState.Find(FGamepadKeyNames::RightAnalogY)->AxisValue <= -RepeatDeadzone;
+		bWasPositivePressed = PrevState.Find(FGamepadKeyNames::RightAnalogY)->AxisValue >= RepeatDeadzone;
+		bPositivePressed = CurrentState.Find(FGamepadKeyNames::RightAnalogY)->AxisValue >= RepeatDeadzone;
+		
+		HandleInputInternal(FGamepadKeyNames::LeftStickDown, ControllerIndex, bNegativePressed, bWasNegativePressed);
+		HandleInputInternal(FGamepadKeyNames::LeftStickUp, ControllerIndex, bPositivePressed, bWasPositivePressed);
+	}
 }
 
 void FAppleControllerInterface::HandleButtonGamepad(const FGamepadKeyNames::Type& UEButton, uint32 ControllerIndex)
 {
-    GCController* Cont = Controllers[ControllerIndex].Controller;
-
-    
-    bool bWasPressed = false;
-    bool bIsPressed = false;
+	bool bWasPressed = false;
+	bool bIsPressed = false;
+	const TMap<FName, FControllerButton>& PrevState = ControllerButtons[PrevButtonIndex][ControllerIndex];
+	const TMap<FName, FControllerButton>& CurrentState = ControllerButtons[CurrentButtonIndex][ControllerIndex];
 
-    GCExtendedGamepad *ExtendedGamepad = nil;
-    GCExtendedGamepad *ExtendedPreviousGamepad = nil;
+	bIsPressed = CurrentState.Find(UEButton)->bPressed;
+	bWasPressed = PrevState.Find(UEButton)->bPressed;
 
-    GCMicroGamepad *MicroGamepad = nil;
-    GCMicroGamepad *MicroPreviousGamepad = nil;
-
-#define SET_PRESSED(Gamepad, PreviousGamepad, GCButton, UEButton) \
-{ \
-bWasPressed = PreviousGamepad.GCButton.pressed; \
-bIsPressed = Gamepad.GCButton.pressed; \
-}
-    switch (Controllers[ControllerIndex].ControllerType)
-    {
-        case ControllerType::ExtendedGamepad:
-        case ControllerType::DualShockGamepad:
-        case ControllerType::XboxGamepad:
-        case ControllerType::DualSenseGamepad:
-            
-            ExtendedGamepad = Cont.extendedGamepad;
-            ExtendedPreviousGamepad = Controllers[ControllerIndex].PreviousExtendedGamepad;
-       
-            if (UEButton == FGamepadKeyNames::FaceButtonLeft){SET_PRESSED(ExtendedGamepad, ExtendedPreviousGamepad, buttonX, UEButton);}
-            else if (UEButton == FGamepadKeyNames::FaceButtonBottom){SET_PRESSED(ExtendedGamepad, ExtendedPreviousGamepad, buttonA, UEButton);}
-            else if (UEButton == FGamepadKeyNames::FaceButtonRight){SET_PRESSED(ExtendedGamepad, ExtendedPreviousGamepad, buttonB, UEButton);}
-            else if (UEButton == FGamepadKeyNames::FaceButtonTop){SET_PRESSED(ExtendedGamepad, ExtendedPreviousGamepad, buttonY, UEButton);}
-            else if (UEButton == FGamepadKeyNames::LeftShoulder){SET_PRESSED(ExtendedGamepad, ExtendedPreviousGamepad, leftShoulder, UEButton);}
-            else if (UEButton == FGamepadKeyNames::RightShoulder){SET_PRESSED(ExtendedGamepad, ExtendedPreviousGamepad, rightShoulder, UEButton);}
-            else if (UEButton == FGamepadKeyNames::LeftTriggerThreshold){SET_PRESSED(ExtendedGamepad, ExtendedPreviousGamepad, leftTrigger, UEButton);}
-            else if (UEButton == FGamepadKeyNames::RightTriggerThreshold){SET_PRESSED(ExtendedGamepad, ExtendedPreviousGamepad, rightTrigger, UEButton);}
-            else if (UEButton == FGamepadKeyNames::DPadUp){SET_PRESSED(ExtendedGamepad, ExtendedPreviousGamepad, dpad.up, UEButton);}
-            else if (UEButton == FGamepadKeyNames::DPadDown){SET_PRESSED(ExtendedGamepad, ExtendedPreviousGamepad, dpad.down, UEButton);}
-            else if (UEButton == FGamepadKeyNames::DPadRight){SET_PRESSED(ExtendedGamepad, ExtendedPreviousGamepad, dpad.right, UEButton);}
-            else if (UEButton == FGamepadKeyNames::DPadLeft){SET_PRESSED(ExtendedGamepad, ExtendedPreviousGamepad, dpad.left, UEButton);}
-            else if (UEButton == FGamepadKeyNames::SpecialRight){SET_PRESSED(ExtendedGamepad, ExtendedPreviousGamepad, buttonMenu, UEButton);}
-            else if (UEButton == FGamepadKeyNames::LeftThumb){SET_PRESSED(ExtendedGamepad, ExtendedPreviousGamepad, leftThumbstickButton, UEButton);}
-            else if (UEButton == FGamepadKeyNames::RightThumb){SET_PRESSED(ExtendedGamepad, ExtendedPreviousGamepad, rightThumbstickButton, UEButton);}
-			else if (UEButton == FGamepadKeyNames::SpecialLeft)
-			{
-				// DualSense and DualShock4 use the touchpad button for the "SpecialLeft" button
-				if (Controllers[ControllerIndex].ControllerType == ControllerType::DualSenseGamepad)
-				{
-					// SetControllerType() verified that ExtendedGamepad is of type GCDualSenseGamepad
-					SET_PRESSED(static_cast<GCDualSenseGamepad*>(ExtendedGamepad), static_cast<GCDualSenseGamepad*>(ExtendedPreviousGamepad), touchpadButton, UEButton);
-				}
-				else if (Controllers[ControllerIndex].ControllerType == ControllerType::DualShockGamepad)
-				{
-					SET_PRESSED(static_cast<GCDualShockGamepad*>(ExtendedGamepad), static_cast<GCDualShockGamepad*>(ExtendedPreviousGamepad), touchpadButton, UEButton);
-				}
-				else
-				{
-					SET_PRESSED(ExtendedGamepad, ExtendedPreviousGamepad, buttonOptions, UEButton);
-				}
-			}
-            break;
-        case ControllerType::SiriRemote:
-            
-            MicroGamepad = Cont.microGamepad;
-            MicroPreviousGamepad = Controllers[ControllerIndex].PreviousMicroGamepad;
-
-            if (UEButton == FGamepadKeyNames::LeftStickUp){SET_PRESSED(MicroGamepad, MicroPreviousGamepad, dpad.up, UEButton);}
-            else if (UEButton == FGamepadKeyNames::LeftStickDown){SET_PRESSED(MicroGamepad, MicroPreviousGamepad, dpad.down, UEButton);}
-            else if (UEButton == FGamepadKeyNames::LeftStickRight){SET_PRESSED(MicroGamepad, MicroPreviousGamepad, dpad.right, UEButton);}
-            else if (UEButton == FGamepadKeyNames::LeftStickLeft){SET_PRESSED(MicroGamepad, MicroPreviousGamepad, dpad.left, UEButton);}
-            else if (UEButton == FGamepadKeyNames::FaceButtonBottom){SET_PRESSED(MicroGamepad, MicroPreviousGamepad, buttonA, UEButton);}
-            else if (UEButton == FGamepadKeyNames::FaceButtonLeft){SET_PRESSED(MicroGamepad, MicroPreviousGamepad, buttonX, UEButton);}
-            else if (UEButton == FGamepadKeyNames::SpecialRight){SET_PRESSED(MicroGamepad, MicroPreviousGamepad, buttonMenu, UEButton);}
-            break;
-    }
-    HandleInputInternal(UEButton, ControllerIndex, bIsPressed, bWasPressed);
+	HandleInputInternal(UEButton, ControllerIndex, bIsPressed, bWasPressed);
 }
 
 void FAppleControllerInterface::HandleAnalogGamepad(const FGamepadKeyNames::Type& UEAxis, uint32 ControllerIndex)
 {
-    GCController* Cont = Controllers[ControllerIndex].Controller;
-    
-    IPlatformInputDeviceMapper& DeviceMapper = IPlatformInputDeviceMapper::Get();
+	GCController* Cont = Controllers[ControllerIndex].Controller;
+	IPlatformInputDeviceMapper& DeviceMapper = IPlatformInputDeviceMapper::Get();
 	FPlatformUserId UserId = FGenericPlatformMisc::GetPlatformUserForUserIndex(Controllers[ControllerIndex].PlayerIndex);
-    FInputDeviceId DeviceId = DeviceMapper.GetPrimaryInputDeviceForUser(UserId);
-    
-    // Send controller events any time we are passed the given input threshold similarly to PC/Console (see: XInputInterface.cpp)
-    const float RepeatDeadzone = 0.24f;
-    bool bWasPositivePressed = false;
-    bool bPositivePressed = false;
-    bool bWasNegativePressed = false;
-    bool bNegativePressed = false;
-    float axisValue = 0;
+	FInputDeviceId DeviceId = DeviceMapper.GetPrimaryInputDeviceForUser(UserId);
+
+	// Send controller events any time we are passed the given input threshold
+	// similarly to PC/Console (see: XInputInterface.cpp)
+	const float RepeatDeadzone = 0.24f;
+	float AxisValue = 0;
 	bool bSendAnalogChangedMessage = false;
 
-    GCExtendedGamepad *ExtendedGamepad = Cont.extendedGamepad;
-    GCExtendedGamepad *ExtendedPreviousGamepad = Controllers[ControllerIndex].PreviousExtendedGamepad;;
+	const TMap<FName, FControllerButton>& PrevState = ControllerButtons[PrevButtonIndex][ControllerIndex];
+	const TMap<FName, FControllerButton>& CurrentState = ControllerButtons[CurrentButtonIndex][ControllerIndex];
 
-    GCMicroGamepad *MicroGamepad = Cont.microGamepad;
-    GCMicroGamepad *MicroPreviousGamepad = Controllers[ControllerIndex].PreviousMicroGamepad;
-    
-    switch (Controllers[ControllerIndex].ControllerType)
-    {
-        case ControllerType::ExtendedGamepad:
-        case ControllerType::DualShockGamepad:
-        case ControllerType::XboxGamepad:
-        case ControllerType::DualSenseGamepad:
-            
-            if (UEAxis == FGamepadKeyNames::LeftAnalogX)
-			{
-				if ((ExtendedPreviousGamepad != nil && ExtendedGamepad.leftThumbstick.xAxis.value != ExtendedPreviousGamepad.leftThumbstick.xAxis.value) ||
-                	(ExtendedGamepad.leftThumbstick.xAxis.value < -RepeatDeadzone || ExtendedGamepad.leftThumbstick.xAxis.value > RepeatDeadzone))
-					{
-						axisValue = ExtendedGamepad.leftThumbstick.xAxis.value;
-						bSendAnalogChangedMessage = true;
-					}
-			}
-            else if (UEAxis == FGamepadKeyNames::LeftAnalogY)
-			{
-				if ((ExtendedPreviousGamepad != nil && ExtendedGamepad.leftThumbstick.yAxis.value != ExtendedPreviousGamepad.leftThumbstick.yAxis.value) ||
-                    (ExtendedGamepad.leftThumbstick.yAxis.value < -RepeatDeadzone || ExtendedGamepad.leftThumbstick.yAxis.value > RepeatDeadzone))
-				{
-					axisValue = ExtendedGamepad.leftThumbstick.yAxis.value;
-					bSendAnalogChangedMessage = true;
-				}
-			}
-            else if (UEAxis == FGamepadKeyNames::RightAnalogX)
-			{
-				if ((ExtendedPreviousGamepad != nil && ExtendedGamepad.rightThumbstick.xAxis.value != ExtendedPreviousGamepad.rightThumbstick.xAxis.value) ||
-                	(ExtendedGamepad.rightThumbstick.xAxis.value < -RepeatDeadzone || ExtendedGamepad.rightThumbstick.xAxis.value > RepeatDeadzone))
-				{
-					axisValue = ExtendedGamepad.rightThumbstick.xAxis.value;
-					bSendAnalogChangedMessage = true;
-				}
-			}
-            else if (UEAxis == FGamepadKeyNames::RightAnalogY)
-			{
-				if ((ExtendedPreviousGamepad != nil && ExtendedGamepad.rightThumbstick.yAxis.value != ExtendedPreviousGamepad.rightThumbstick.yAxis.value) ||
-                    (ExtendedGamepad.rightThumbstick.yAxis.value < -RepeatDeadzone || ExtendedGamepad.rightThumbstick.yAxis.value > RepeatDeadzone))
-				{
-					axisValue = ExtendedGamepad.rightThumbstick.yAxis.value;
-					bSendAnalogChangedMessage = true;
-				}
-			}
-            else if (UEAxis == FGamepadKeyNames::LeftTriggerAnalog)
-			{
-				if ((ExtendedPreviousGamepad != nil && ExtendedGamepad.leftTrigger.value != ExtendedPreviousGamepad.leftTrigger.value) ||
-                	(ExtendedGamepad.leftTrigger.value < -RepeatDeadzone || ExtendedGamepad.leftTrigger.value > RepeatDeadzone))
-				{
-					axisValue = ExtendedGamepad.leftTrigger.value;
-					bSendAnalogChangedMessage = true;
-				}
-			}
-            else if (UEAxis == FGamepadKeyNames::RightTriggerAnalog)
-			{
-				if ((ExtendedPreviousGamepad != nil && ExtendedGamepad.rightTrigger.value != ExtendedPreviousGamepad.rightTrigger.value) ||
-                	(ExtendedGamepad.rightTrigger.value < -RepeatDeadzone || ExtendedGamepad.rightTrigger.value > RepeatDeadzone))
-				{
-					axisValue = ExtendedGamepad.rightTrigger.value;
-					bSendAnalogChangedMessage = true;
-				}
-			}
-            break;
-            
-        case ControllerType::SiriRemote:
-            
-            if (UEAxis == FGamepadKeyNames::LeftAnalogX)
+	switch (Controllers[ControllerIndex].ControllerType)
+	{
+		case ControllerType::ExtendedGamepad:
+		case ControllerType::DualShockGamepad:
+		case ControllerType::XboxGamepad:
+		case ControllerType::DualSenseGamepad:
+		{
+			if ((UEAxis == FGamepadKeyNames::LeftAnalogX) || (UEAxis == FGamepadKeyNames::RightAnalogX) || (UEAxis == FGamepadKeyNames::LeftTriggerAnalog) ||
+				(UEAxis == FGamepadKeyNames::LeftAnalogY) || (UEAxis == FGamepadKeyNames::RightAnalogY) || (UEAxis == FGamepadKeyNames::RightTriggerAnalog))
 			{
-				if ((ExtendedPreviousGamepad != nil && MicroGamepad.dpad.xAxis.value != ExtendedPreviousGamepad.dpad.xAxis.value) ||
-                (MicroGamepad.dpad.xAxis.value < -RepeatDeadzone || MicroGamepad.dpad.xAxis.value > RepeatDeadzone))
+				float CurrentValue = CurrentState.Find(UEAxis)->AxisValue;
+				float PrevValue = PrevState.Find(UEAxis)->AxisValue;
+				if ((CurrentValue != PrevValue) || (CurrentValue < -RepeatDeadzone || CurrentValue > RepeatDeadzone))
 				{
-					axisValue = MicroGamepad.dpad.xAxis.value;
+					AxisValue = CurrentValue;
 					bSendAnalogChangedMessage = true;
 				}
 			}
-            else if (UEAxis == FGamepadKeyNames::LeftAnalogY)
+			break;
+		}
+
+		case ControllerType::SiriRemote:
+		{
+			if ((UEAxis == FGamepadKeyNames::LeftAnalogX) || (UEAxis == FGamepadKeyNames::LeftAnalogY))
 			{
-				if ((MicroPreviousGamepad != nil && MicroGamepad.dpad.yAxis.value != MicroPreviousGamepad.dpad.yAxis.value) ||
-                (MicroGamepad.dpad.yAxis.value < -RepeatDeadzone || MicroGamepad.dpad.yAxis.value > RepeatDeadzone))
+				float CurrentValue = CurrentState.Find(UEAxis)->AxisValue;
+				float PrevValue = CurrentState.Find(UEAxis)->AxisValue;
+				if ((CurrentValue != PrevValue) || (CurrentValue < -RepeatDeadzone || CurrentValue > RepeatDeadzone))
 				{
-					axisValue = MicroGamepad.dpad.yAxis.value;
+					AxisValue = CurrentValue;
 					bSendAnalogChangedMessage = true;
 				}
 			}
-            break;
-    }
+			break;
+		}
+
+		default:
+			UE_LOG(LogAppleController, Warning, TEXT("Controller type is recognized you analogue"));
+			break;
+	}
 #if APPLE_CONTROLLER_DEBUG
-    NSLog(@"Axis %s is %f", TCHAR_TO_ANSI(*UEAxis.ToString()), axisValue);
+	NSLog(@"Axis %s is %f", TCHAR_TO_ANSI(*UEAxis.ToString()), AxisValue);
 #endif
 	if (bSendAnalogChangedMessage)
 	{
-		MessageHandler->OnControllerAnalog(UEAxis, UserId, DeviceId, axisValue);
+		MessageHandler->OnControllerAnalog(UEAxis, UserId, DeviceId, AxisValue);
 	}
 }
diff --git a/Engine/Source/Runtime/ApplicationCore/Private/Apple/AppleControllerInterface.h b/Engine/Source/Runtime/ApplicationCore/Private/Apple/AppleControllerInterface.h
--- a/Engine/Source/Runtime/ApplicationCore/Private/Apple/AppleControllerInterface.h
+++ b/Engine/Source/Runtime/ApplicationCore/Private/Apple/AppleControllerInterface.h
@@ -10,6 +10,9 @@
 
 DECLARE_LOG_CATEGORY_EXTERN(LogAppleController, Log, All);
 
+#define MAX_NUM_CONTROLLERS					4	// reasonable limit for now
+#define MAX_CONTROLLER_CAPTURE_STATES		2	// track current and previous button states
+
 enum ControllerType
 {
 	Unassigned,
@@ -17,7 +20,7 @@ enum ControllerType
 	ExtendedGamepad,
 	XboxGamepad,
 	DualShockGamepad,
-    DualSenseGamepad
+	DualSenseGamepad
 };
 
 enum PlayerIndex
@@ -32,10 +35,10 @@ enum PlayerIndex
 
 enum class EAppleControllerEventType : int32
 {
-    Invalid,
-    Connect,
-    Disconnect,
-    BecomeCurrent
+	Invalid,
+	Connect,
+	Disconnect,
+	BecomeCurrent
 };
 
 struct FDeferredAppleControllerEvent
@@ -53,8 +56,8 @@ struct FDeferredAppleControllerEvent
 	{
 		[Controller release];
 	}
-    EAppleControllerEventType EventType;
-    GCController* Controller;
+	EAppleControllerEventType EventType;
+	GCController* Controller;
 };
 
 /**
@@ -93,15 +96,15 @@ public:
 	bool IsGamepadAttached() const;
 
 	const ControllerType GetControllerType(uint32 ControllerIndex);
-    void SetControllerType(uint32 ControllerIndex);
+	void SetControllerType(uint32 ControllerIndex);
 	FName GetControllerTypeName(const ControllerType InControllerType);
-    
-    GCControllerButtonInput* GetGCControllerButton(const FGamepadKeyNames::Type& ButtonKey, uint32 ControllerIndex);
-    
-    void HandleInputInternal(const FGamepadKeyNames::Type& UEButton, uint32 ControllerIndex, bool bIsPressed, bool bWasPressed);
-    void HandleButtonGamepad(const FGamepadKeyNames::Type& UEButton, uint32 ControllerIndex);
-    void HandleAnalogGamepad(const FGamepadKeyNames::Type& UEAxis, uint32 ControllerIndex);
-    void HandleVirtualButtonGamepad(const FGamepadKeyNames::Type& UEButtonNegative, const FGamepadKeyNames::Type& UEButtonPositive, uint32 ControllerIndex);
+	
+	GCControllerButtonInput* GetGCControllerButton(const FGamepadKeyNames::Type& ButtonKey, uint32 ControllerIndex);
+	
+	void HandleInputInternal(const FGamepadKeyNames::Type& UEButton, uint32 ControllerIndex, bool bIsPressed, bool bWasPressed);
+	void HandleButtonGamepad(const FGamepadKeyNames::Type& UEButton, uint32 ControllerIndex);
+	void HandleAnalogGamepad(const FGamepadKeyNames::Type& UEAxis, uint32 ControllerIndex);
+	void HandleVirtualButtonGamepad(const FGamepadKeyNames::Type& UEButtonNegative, const FGamepadKeyNames::Type& UEButtonPositive, uint32 ControllerIndex);
 
 protected:
 
@@ -113,42 +116,59 @@ private:
 
 	void HandleConnection(GCController* Controller);
 	void HandleDisconnect(GCController* Controller);
-	
-    void SetCurrentController(GCController* Controller);
+
+	void SetCurrentController(GCController* Controller);
+
+	void CaptureControllerState(const GCController* Controller, uint32 ControllerIndex);
 
 protected:
-	
-    TSharedRef< FGenericApplicationMessageHandler > MessageHandler;
 
+	TSharedRef< FGenericApplicationMessageHandler > MessageHandler;
 
 	/** Game controller objects (per user)*/
 	struct FUserController
 	{
-        GCController* Controller;
-        
-        ControllerType ControllerType;
-        
-        PlayerIndex PlayerIndex;
-
-        GCExtendedGamepad* PreviousExtendedGamepad;
-        GCMicroGamepad* PreviousMicroGamepad;
-        
+		GCController* Controller;
+		ControllerType ControllerType;
+		PlayerIndex PlayerIndex;
+
 #if !PLATFORM_MAC
 		FQuat ReferenceAttitude;
 		bool bNeedsReferenceAttitude;
 		bool bHasReferenceAttitude;
 #endif
 	};
-	
+
+	/** Track individual buttons and axis values */
+	struct FControllerButton
+	{
+		bool bPressed = false;
+		float PressedValue = 0;	// the level of pressure being applied to the button
+		float AxisValue = 0;
+
+		void Reset()
+		{
+			bPressed = false;
+			PressedValue = 0;
+			AxisValue = 0;
+		}
+	};
+
 	// Controller Event Callbacks are on the main thread - defer to tick processing
 	FCriticalSection DeferredEventCS;
 	TArray<FDeferredAppleControllerEvent> DeferredEvents;
-	
-    // there is a hardcoded limit of 4 controllers in the API
-	FUserController Controllers[4];
+
+	// there is a hardcoded limit of 4 controllers in the API
+	FUserController Controllers[MAX_NUM_CONTROLLERS];
+
+	uint32 CurrentButtonIndex = 0;
+	uint32 PrevButtonIndex = 1;
+	// We track 2 sets of buttons - the current state and the previous state
+	// And there is a hardcoded limit of 4 controllers in the API
+	TMap<FName, FControllerButton> ControllerButtons[MAX_CONTROLLER_CAPTURE_STATES][MAX_NUM_CONTROLLERS];
 	
 	TMap<FName, double> NextKeyRepeatTime;
-	
+
 	// should we allow controllers to send input
 	bool bAllowControllers;
 
