fix deprecation warning in Xcode 26

#rnx

diff --git a/Engine/Source/Runtime/ApplicationCore/Private/IOS/IOSAppDelegate.cpp b/Engine/Source/Runtime/ApplicationCore/Private/IOS/IOSAppDelegate.cpp
--- a/Engine/Source/Runtime/ApplicationCore/Private/IOS/IOSAppDelegate.cpp
+++ b/Engine/Source/Runtime/ApplicationCore/Private/IOS/IOSAppDelegate.cpp
@@ -740,7 +740,13 @@ static IOSAppDelegate* CachedDelegate = nil;
 
 			if ([self IsFeatureActive:EAudioFeature::BluetoothMicrophone])
 			{
+				// iOS sdk 26 deprecated AVAudioSessionCategoryOptionAllowBluetooth in favor of AVAudioSessionCategoryOptionAllowBluetoothHFP
+				// but AVAudioSessionCategoryOptionAllowBluetoothHFP is not present in sdk 16 despite what the docs/header say
+#if (defined(__IPHONE_26_0) && __IPHONE_OS_VERSION_MAX_ALLOWED >= __IPHONE_26_0)
+				Options |= AVAudioSessionCategoryOptionAllowBluetoothHFP;
+#else
 				Options |= AVAudioSessionCategoryOptionAllowBluetooth;
+#endif
 			}
 #endif
 
@@ -765,7 +771,13 @@ static IOSAppDelegate* CachedDelegate = nil;
 #if !PLATFORM_TVOS
 			if ([self IsFeatureActive:EAudioFeature::BluetoothMicrophone])
 			{
+				// iOS sdk 26 deprecated AVAudioSessionCategoryOptionAllowBluetooth in favor of AVAudioSessionCategoryOptionAllowBluetoothHFP
+				// but AVAudioSessionCategoryOptionAllowBluetoothHFP is not present in sdk 16 despite what the docs/header say
+#if (defined(__IPHONE_26_0) && __IPHONE_OS_VERSION_MAX_ALLOWED >= __IPHONE_26_0)
+				Options |= AVAudioSessionCategoryOptionAllowBluetoothHFP;
+#else
 				Options |= AVAudioSessionCategoryOptionAllowBluetooth;
+#endif
 			}
 #endif
 		}
