tvOS microGamepad handling is now performed in AppleInterfaceController as the SiriRemote

#rnx

#Synced-CL Unknown

diff --git a/Engine/Source/Runtime/ApplicationCore/Private/IOS/IOSInputInterface.cpp b/Engine/Source/Runtime/ApplicationCore/Private/IOS/IOSInputInterface.cpp
--- a/Engine/Source/Runtime/ApplicationCore/Private/IOS/IOSInputInterface.cpp
+++ b/Engine/Source/Runtime/ApplicationCore/Private/IOS/IOSInputInterface.cpp
@@ -645,39 +645,6 @@ void FIOSInputInterface::SendControllerEvents()
     
     // Generic Controller update
     FAppleControllerInterface::SendControllerEvents();
-    
-#if PLATFORM_TVOS
-    // tvOS specific controller enumeration to handle micro input if no gamepad
-    for(int32 i = 0; i < UE_ARRAY_COUNT(Controllers); ++i)
-    {
-        FUserController& Controller = Controllers[i];
-        if(Controller.PreviousExtendedGamepad == nil)
-        {
-            GCController* Cont = Controller.Controller;
-            GCMicroGamepad* MicroGamepad = [Cont capture].microGamepad;
-            if (MicroGamepad != nil && Controller.PlayerIndex != PlayerIndex::PlayerUnset)
-            {
-                const GCMicroGamepad* PreviousMicroGamepad = Controller.PreviousMicroGamepad;
-
-                HandleButtonGamepad(FGamepadKeyNames::FaceButtonBottom, i);
-                HandleButtonGamepad(FGamepadKeyNames::FaceButtonLeft, i);
-                HandleButtonGamepad(FGamepadKeyNames::SpecialRight, i);
-                
-                HandleAnalogGamepad(FGamepadKeyNames::LeftAnalogX, i);
-                HandleAnalogGamepad(FGamepadKeyNames::LeftAnalogY, i);
-
-                HandleButtonGamepad(FGamepadKeyNames::LeftStickUp, i);
-                HandleButtonGamepad(FGamepadKeyNames::LeftStickDown, i);
-                HandleButtonGamepad(FGamepadKeyNames::LeftStickRight, i);
-                HandleButtonGamepad(FGamepadKeyNames::LeftStickLeft, i);
-                
-                [Controller.PreviousMicroGamepad release];
-                Controller.PreviousMicroGamepad = MicroGamepad;
-                [Controller.PreviousMicroGamepad retain];
-            }
-        }
-    }
-#endif
 }
 
 void FIOSInputInterface::QueueTouchInput(const TArray<TouchInput>& InTouchEvents)
