Fixed typo in Android manual resoultion override parameter
#Synced-CL 48601912

diff --git a/Engine/Source/Runtime/ApplicationCore/Private/Android/AndroidWindow.cpp b/Engine/Source/Runtime/ApplicationCore/Private/Android/AndroidWindow.cpp
--- a/Engine/Source/Runtime/ApplicationCore/Private/Android/AndroidWindow.cpp
+++ b/Engine/Source/Runtime/ApplicationCore/Private/Android/AndroidWindow.cpp
@@ -588,7 +588,7 @@ FPlatformRect FAndroidWindow::GetScreenRect(bool bUseEventThreadWindow)
 	FAndroidCachedWindowRectParams CurrentParams;
 	CurrentParams.ContentScaleFactor = RequestedContentScaleFactor;
 	CurrentParams.MobileResX = RequestedResX;
-	CurrentParams.MobileResY = RequestedResX;
+	CurrentParams.MobileResY = RequestedResY;
 	CurrentParams.Window_EventThread = Window;
 	CurrentParams.WindowDPI = GAndroidWindowDPI;
 	CurrentParams.SceneMinDPI = FMath::Min(GAndroid3DSceneMinDPI, GAndroidWindowDPI);
