[iOS] Improving iPad external keyboard and mouse support.

This CL reworked the current usage GCkeyboard and GC Mouse to

- Not be iOS specific (this sets us up to use it in macOS in the futture)
- Queue all inputs in a lock free manner
- For keyboard, added basic character input support (GCKeyboard does not provide proper keyboard layout and character translation information). This means it cannot be used for proper text input support as is, but the console command will work now.
- Removed the old GCKeyboard and GCMouse implementation.

#Synced-CL 48538939

diff --git a/Engine/Source/Runtime/ApplicationCore/Private/Apple/AppleControllerInterface.h b/Engine/Source/Runtime/ApplicationCore/Private/Apple/AppleControllerInterface.h
--- a/Engine/Source/Runtime/ApplicationCore/Private/Apple/AppleControllerInterface.h
+++ b/Engine/Source/Runtime/ApplicationCore/Private/Apple/AppleControllerInterface.h
@@ -74,7 +74,7 @@ public:
 
 	virtual ~FAppleControllerInterface() {}
 
-	void SetMessageHandler( const TSharedRef< FGenericApplicationMessageHandler >& InMessageHandler );
+	virtual void SetMessageHandler( const TSharedRef< FGenericApplicationMessageHandler >& InMessageHandler );
 
 	/** Tick the interface (i.e check for new controllers) */
 	void Tick( float DeltaTime );
diff --git a/Engine/Source/Runtime/ApplicationCore/Private/Apple/AppleKeyboardController.cpp b/Engine/Source/Runtime/ApplicationCore/Private/Apple/AppleKeyboardController.cpp
new file mode 100644
--- /dev/null
+++ b/Engine/Source/Runtime/ApplicationCore/Private/Apple/AppleKeyboardController.cpp
@@ -0,0 +1,387 @@
+// Copyright Epic Games, Inc. All Rights Reserved.
+
+#include "AppleKeyboardController.h"
+
+#include "AppleControllerInterface.h"
+
+#include "GenericPlatform/GenericApplicationMessageHandler.h"
+#include "HAL/PlatformTime.h"
+
+#import <GameController/GCControllerButtonInput.h>
+#import <GameController/GCKeyboard.h>
+#import <GameController/GCKeyboardInput.h>
+#import <GameController/GCKeyCodes.h>
+
+#include "HAL/IConsoleManager.h"
+#include "Misc/CoreDelegates.h"
+#include "Misc/ScopeRWLock.h"
+
+namespace Apple::Cvars
+{
+	static bool bEnableGCKeyboardSupport = true;
+	static FAutoConsoleVariableRef CVarEnableGCKeyboardSupport(
+				TEXT("input.apple.enablegckeyboardsupport"),
+				bEnableGCKeyboardSupport,
+				TEXT("If set to false, the implementation for GCKeyboard support will tell the rest of the system that no keyboard is connected"));
+}
+
+FAppleKeyboardController::FAppleKeyboardController(const TSharedRef<FGenericApplicationMessageHandler>& InMessageHandler) : MessageHandler(InMessageHandler)
+{
+	if (@available(macOS 11.0, iOS 14.0, tvOS 14.0, *))
+	{
+		RegisterNonPrintableKeys();
+		
+		ConnectedEventObserverID = [[NSNotificationCenter defaultCenter] addObserverForName:GCKeyboardDidConnectNotification 
+																		object:nil
+																		queue:[NSOperationQueue mainQueue]
+																		usingBlock:^(NSNotification* Notification) { HandleKeyboardConnected(); }];
+		
+		DisconnectedEventObserverID = [[NSNotificationCenter defaultCenter] addObserverForName:GCKeyboardDidDisconnectNotification
+																			object:nil 
+																			queue:[NSOperationQueue mainQueue]
+																			usingBlock:^(NSNotification* Notification) { HandleKeyboardDisconnected(); }];
+		
+		FCoreDelegates::ApplicationHasEnteredForegroundDelegate.AddRaw(this, &FAppleKeyboardController::ForceUpdatePressedKeysState);
+		
+		if (!bHandledKeyboardConnection && GCKeyboard.coalescedKeyboard)
+		{
+			HandleKeyboardConnected();
+		}
+		
+		bIsInitialized = true;
+	}
+	else
+	{
+		UE_LOG(LogAppleController, Warning, TEXT("GCKeyboard API is not available in the current OS version. Physical keyboad dupport via FAppleKeyboardController will not be available."));	
+	}
+}
+
+FAppleKeyboardController::~FAppleKeyboardController()
+{
+	if (@available(macOS 11.0, iOS 14.0, tvOS 14.0, *))
+	{
+		if (ConnectedEventObserverID)
+		{
+			[[NSNotificationCenter defaultCenter] removeObserver:ConnectedEventObserverID];
+		}
+		
+		if (DisconnectedEventObserverID)
+		{
+			[[NSNotificationCenter defaultCenter] removeObserver:DisconnectedEventObserverID];
+		}
+		
+		GCKeyboardInput* KeyboardInput = GCKeyboard.coalescedKeyboard ? GCKeyboard.coalescedKeyboard.keyboardInput : nullptr;
+		if (KeyboardInput)
+		{
+			KeyboardInput.keyChangedHandler = nil;
+		}
+		
+		FCoreDelegates::ApplicationHasEnteredForegroundDelegate.RemoveAll(this);
+	}
+}
+
+void FAppleKeyboardController::SetMessageHandler(const TSharedRef<FGenericApplicationMessageHandler>& InMessageHandler)
+{
+	MessageHandler = InMessageHandler;
+}
+
+bool FAppleKeyboardController::IsPrintableKey(uint32 KeyCode) const
+{
+	return !NonPrintableKeys.Contains(KeyCode);
+}
+
+void FAppleKeyboardController::RegisterNonPrintableKeys()
+{
+	// If more keys are added, increase this number to match
+	constexpr int32 NonPrintableKeysNum = 29;
+	NonPrintableKeys.Reserve(NonPrintableKeysNum);
+	
+	NonPrintableKeys.Emplace(GCKeyCodeF1);
+	NonPrintableKeys.Emplace(GCKeyCodeF2);
+	NonPrintableKeys.Emplace(GCKeyCodeF3);
+	NonPrintableKeys.Emplace(GCKeyCodeF4);
+	NonPrintableKeys.Emplace(GCKeyCodeF5);
+	NonPrintableKeys.Emplace(GCKeyCodeF6);
+	NonPrintableKeys.Emplace(GCKeyCodeF7);
+	NonPrintableKeys.Emplace(GCKeyCodeF8);
+	NonPrintableKeys.Emplace(GCKeyCodeF9);
+	NonPrintableKeys.Emplace(GCKeyCodeF10);
+	NonPrintableKeys.Emplace(GCKeyCodeF11);
+	NonPrintableKeys.Emplace(GCKeyCodeF12);
+	NonPrintableKeys.Emplace(GCKeyCodePause);
+	NonPrintableKeys.Emplace(GCKeyCodeInsert);
+	NonPrintableKeys.Emplace(GCKeyCodeHome);
+	NonPrintableKeys.Emplace(GCKeyCodePageUp);
+	NonPrintableKeys.Emplace(GCKeyCodeEnd);
+	NonPrintableKeys.Emplace(GCKeyCodePageDown);
+	NonPrintableKeys.Emplace(GCKeyCodeRightArrow);
+	NonPrintableKeys.Emplace(GCKeyCodeLeftArrow);
+	NonPrintableKeys.Emplace(GCKeyCodeDownArrow);
+	NonPrintableKeys.Emplace(GCKeyCodeUpArrow);
+	NonPrintableKeys.Emplace(GCKeyCodeDeleteForward);
+	NonPrintableKeys.Emplace(GCKeyCodeLeftShift);
+	NonPrintableKeys.Emplace(GCKeyCodeRightShift);
+	NonPrintableKeys.Emplace(GCKeyCodeLeftAlt);
+	NonPrintableKeys.Emplace(GCKeyCodeRightAlt);
+	NonPrintableKeys.Emplace(GCKeyCodeRightGUI);
+	NonPrintableKeys.Emplace(GCKeyCodeLeftGUI);
+}
+
+void FAppleKeyboardController::SendControllerEvents()
+{
+	if (!bIsInitialized)
+	{
+		return;
+	}
+
+	const double CurrentTime = FPlatformTime::Seconds();
+	FDeferredEvent CurrentEvent;
+	while (DeferredKeyEvents.Dequeue(CurrentEvent))
+	{
+		UpdateEventTracking(CurrentEvent, CurrentTime);
+		DispatchEvent(CurrentEvent);
+	}
+	
+	UpdateAndDispatchRepeatEvents();
+}
+
+void FAppleKeyboardController::HandleKeyChangeEvent(GCKeyCode keyCode, BOOL pressed) 
+{ 
+	if (bIsKeyboardInhibited)
+	{
+		return;
+	}
+
+	FDeferredEvent NewKeyEvent;
+	NewKeyEvent.KeyEventType = pressed ? EKeyEvent::KeyDown : EKeyEvent::KeyUp;
+	NewKeyEvent.GCKeyCode = keyCode;
+	NewKeyEvent.TranslatedKeyCode = KeyCodeToPlatformCodeTranslator.IsBound() ? KeyCodeToPlatformCodeTranslator.Execute(NewKeyEvent.GCKeyCode) : NewKeyEvent.GCKeyCode;
+	NewKeyEvent.CharCode = KeyCodeToCharTranslator.IsBound() ? KeyCodeToCharTranslator.Execute(NewKeyEvent.GCKeyCode) : '?';
+	
+	UE_LOG(LogAppleController, Verbose, TEXT("[GC AppleKeyboardController %hs] Key State Changed | GC KeyCode [%u] | State [%s]"), __func__, NewKeyEvent.GCKeyCode, pressed ? TEXT("True") : TEXT("False"));
+	
+	DeferredKeyEvents.Enqueue(NewKeyEvent);
+}
+
+FModifierKeysState FAppleKeyboardController::GetModifierKeys() const
+{
+	return CurrentModifierKeys.ToUEModifierKeyStates();
+}
+
+void FAppleKeyboardController::SetGCKeyCodeToPlatformCodeTranslator(const FKeyCodeTranslator& InExternalTranslatorDelegate)
+{
+	KeyCodeToPlatformCodeTranslator = InExternalTranslatorDelegate;
+}
+
+void FAppleKeyboardController::SetGCKeyCodeToCharTranslator(const FKeyCodeTranslator& InExternalTranslatorDelegate)
+{
+	KeyCodeToCharTranslator = InExternalTranslatorDelegate;
+}
+
+void FAppleKeyboardController::ProcessModifierKeyEvent(const FDeferredEvent& InDeferredEvent)
+{
+	UpdateModifierKeyState(CurrentModifierKeys.bAreCapsLocked, GCKeyCodeCapsLock, InDeferredEvent);
+	UpdateModifierKeyState(CurrentModifierKeys.bIsLeftAltDown, GCKeyCodeLeftAlt, InDeferredEvent);
+	UpdateModifierKeyState(CurrentModifierKeys.bIsLeftCommandDown, GCKeyCodeLeftGUI, InDeferredEvent);
+	UpdateModifierKeyState(CurrentModifierKeys.bIsLeftControlDown, GCKeyCodeLeftControl, InDeferredEvent);
+	UpdateModifierKeyState(CurrentModifierKeys.bIsLeftShiftDown, GCKeyCodeLeftShift, InDeferredEvent);
+	UpdateModifierKeyState(CurrentModifierKeys.bIsRightAltDown, GCKeyCodeRightAlt, InDeferredEvent);
+	UpdateModifierKeyState(CurrentModifierKeys.bIsRightCommandDown, GCKeyCodeRightGUI, InDeferredEvent);
+	UpdateModifierKeyState(CurrentModifierKeys.bIsRightControlDown, GCKeyCodeRightControl, InDeferredEvent);
+	UpdateModifierKeyState(CurrentModifierKeys.bIsRightShiftDown, GCKeyCodeRightShift, InDeferredEvent);
+}
+
+FAppleKeyboardController::EKeyEvent FAppleKeyboardController::GetUpdatedKeyStateForKey(GCKeyCode KeyCode)
+{
+	GCKeyboardInput* KeyboardInput = GCKeyboard.coalescedKeyboard ? GCKeyboard.coalescedKeyboard.keyboardInput : nullptr;
+	if (!KeyboardInput)
+	{
+		return EKeyEvent::KeyUp;
+	}
+
+	GCControllerButtonInput* KeyInput = [KeyboardInput buttonForKeyCode: KeyCode];
+
+	return KeyInput && KeyInput.pressed ? EKeyEvent::KeyDown : EKeyEvent::KeyUp;
+}
+
+void FAppleKeyboardController::UpdateModifierKeyState(bool& OutModifierKeyFlag, GCKeyCode ModifierKeyCode, const FDeferredEvent& InKeyEvent)
+{
+	if (InKeyEvent.GCKeyCode == ModifierKeyCode)
+	{
+		OutModifierKeyFlag = InKeyEvent.KeyEventType == EKeyEvent::KeyDown;
+	}
+}
+
+void FAppleKeyboardController::DispatchEvent(const FDeferredEvent& InDeferredEvent)
+{
+	UE_LOG(LogAppleController,
+		   Verbose,
+		   TEXT("[GC AppleKeyboardController %hs] GC KeyCode [%u] | Type [%s] | Is Repeat [%s]"),
+		   __func__,
+		   InDeferredEvent.GCKeyCode,
+		   InDeferredEvent.KeyEventType == EKeyEvent::KeyUp ? TEXT("UP") : TEXT("DOWN"),
+		   InDeferredEvent.bIsRepeat ? TEXT("True") : TEXT("False"));
+	
+	ProcessModifierKeyEvent(InDeferredEvent);
+
+	switch (InDeferredEvent.KeyEventType)
+	{
+		case EKeyEvent::KeyUp:
+		{
+			MessageHandler->OnKeyUp(InDeferredEvent.TranslatedKeyCode, InDeferredEvent.CharCode, false);
+			break;
+		}
+		case EKeyEvent::KeyDown:
+		{
+			MessageHandler->OnKeyDown(InDeferredEvent.TranslatedKeyCode, InDeferredEvent.CharCode, InDeferredEvent.bIsRepeat);
+			
+			// Mimic how CMD is handled in MacApplication
+			const bool IsCMDPressed = CurrentModifierKeys.bIsLeftCommandDown || CurrentModifierKeys.bIsRightCommandDown;
+			if (!IsCMDPressed && IsPrintableKey(InDeferredEvent.GCKeyCode))
+			{
+				MessageHandler->OnKeyChar((TCHAR) InDeferredEvent.CharCode, InDeferredEvent.bIsRepeat);
+			}
+			break;
+		}
+	case EKeyEvent::Invalid:
+	default:
+		ensure(false);
+		break;
+	}
+}
+
+void FAppleKeyboardController::UpdateEventTracking(const FDeferredEvent& InDeferredEvent, double CurrentTime)
+{
+	switch (InDeferredEvent.KeyEventType)
+	{
+		case EKeyEvent::KeyUp:
+		{
+			CurrentlyPressedKeyStates.Remove(InDeferredEvent.GCKeyCode);
+			break;
+		}
+		case EKeyEvent::KeyDown:
+		{
+			FDeferredEvent PressedKeyEventCopy = InDeferredEvent;
+			PressedKeyEventCopy.LastRepeatTime = CurrentTime;
+			CurrentlyPressedKeyStates.Add(InDeferredEvent.GCKeyCode, PressedKeyEventCopy);
+			break;
+		}
+	case EKeyEvent::Invalid:
+	default:
+		ensure(false);
+		break;
+	}
+}
+
+void FAppleKeyboardController::UpdateAndDispatchRepeatEvents()
+{
+	if (bIsKeyboardInhibited)
+	{
+		return;
+	}
+
+	const double CurrentTime = FPlatformTime::Seconds();
+	
+	bool bForceStateUpdate = bNeedUpdatePressedKeyState.exchange(false);
+	
+	// Force state update is used to double-check what is the actual state of a key in case it changed while the app was in the background
+	// If it was unpressed we need to remove it from the current pressed key list, and dispatch the key up event
+	TArray<FDeferredEvent, TInlineAllocator<32>> KeysUnpressedOutsideNormalLoop;
+
+	for (TPair<uint32, FDeferredEvent>& PressedKeyEventWithId : CurrentlyPressedKeyStates)
+	{
+		double ElapsedTime = CurrentTime - PressedKeyEventWithId.Value.LastRepeatTime;
+	
+		// TODO : These are usually system settings. Find a way to read them that works in macOS and iOS
+		constexpr double InitialRepeatDelay = 0.5;
+		constexpr double RepeatInterval = 0.1;
+	
+		const double DesiredRepeatInterval = FMath::IsNearlyZero(PressedKeyEventWithId.Value.LastRepeatTime) ? InitialRepeatDelay : RepeatInterval;
+		if (ElapsedTime > DesiredRepeatInterval)
+		{
+			PressedKeyEventWithId.Value.bIsRepeat = true;
+			PressedKeyEventWithId.Value.LastRepeatTime = CurrentTime;
+			DispatchEvent(PressedKeyEventWithId.Value);
+		}
+		
+		if (bForceStateUpdate)
+		{
+			EKeyEvent CurrentKeyState = GetUpdatedKeyStateForKey(PressedKeyEventWithId.Value.GCKeyCode);
+			
+			if (CurrentKeyState == EKeyEvent::KeyUp)
+			{
+				FDeferredEvent KeyEventCopy = PressedKeyEventWithId.Value;
+				KeyEventCopy.KeyEventType = CurrentKeyState;
+				KeysUnpressedOutsideNormalLoop.Emplace(KeyEventCopy);
+			}
+		}
+	}
+	
+	for (const FDeferredEvent& UnpressedKeyEvent : KeysUnpressedOutsideNormalLoop)
+	{
+		DispatchEvent(UnpressedKeyEvent);
+		CurrentlyPressedKeyStates.Remove(UnpressedKeyEvent.GCKeyCode);
+	}
+}
+
+void FAppleKeyboardController::HandleKeyboardConnected() 
+{
+	UE_LOG(LogAppleController, Verbose, TEXT("[GC AppleKeyboardController %hs] Handling Keyboard Connection "), __func__);
+
+	GCKeyboardInput* KeyboardInput = GCKeyboard.coalescedKeyboard ? GCKeyboard.coalescedKeyboard.keyboardInput : nullptr;
+	
+	if (ensure(KeyboardInput))
+	{
+
+		KeyboardInput.keyChangedHandler = ^(GCKeyboardInput* keyboard, GCDeviceButtonInput* key, GCKeyCode keyCode, BOOL pressed)
+		{
+			this->HandleKeyChangeEvent(keyCode, pressed); 
+		};
+		
+		OnAvailabilityChanged().Broadcast();
+
+		bHandledKeyboardConnection = true;
+	}
+}
+
+void FAppleKeyboardController::HandleKeyboardDisconnected() 
+{ 
+	UE_LOG(LogAppleController, Verbose, TEXT("[GC AppleKeyboardController %hs] Handling Keyboard Disconnection "), __func__);
+
+	OnAvailabilityChanged().Broadcast();
+	
+	if (!GCKeyboard.coalescedKeyboard)
+	{
+		ForceUpdatePressedKeysState();
+	}
+	
+	bHandledKeyboardConnection = false;
+}
+
+void FAppleKeyboardController::ForceUpdatePressedKeysState() 
+{ 	
+	// This is called in the OS main thread, but we want to update the data in the Game Thread
+	bNeedUpdatePressedKeyState = true;
+}
+
+void FAppleKeyboardController::SetInhibited(bool bNewIsInhibited) 
+{ 
+	if (bIsKeyboardInhibited != bNewIsInhibited)
+	{
+		bIsKeyboardInhibited = bNewIsInhibited;
+		
+		ForceUpdatePressedKeysState();
+	}
+}
+
+bool FAppleKeyboardController::IsAnyKeyboardConnected() const 
+{ 
+	if (@available(macOS 11.0, iOS 14.0, tvOS 14.0, *))
+	{
+		return Apple::Cvars::bEnableGCKeyboardSupport && GCKeyboard.coalescedKeyboard != nullptr;
+	}
+	else
+	{
+		return false;
+	}
+}
diff --git a/Engine/Source/Runtime/ApplicationCore/Private/Apple/AppleKeyboardController.h b/Engine/Source/Runtime/ApplicationCore/Private/Apple/AppleKeyboardController.h
new file mode 100644
--- /dev/null
+++ b/Engine/Source/Runtime/ApplicationCore/Private/Apple/AppleKeyboardController.h
@@ -0,0 +1,157 @@
+// Copyright Epic Games, Inc. All Rights Reserved.
+
+#pragma once
+
+#include "Containers/SpscQueue.h"
+#include "GenericPlatform/GenericApplication.h"
+
+#include "Templates/SharedPointer.h"
+
+#import <GameController/GCKeyCodes.h>
+
+class FGenericApplicationMessageHandler;
+
+/**
+ * Class that intercepts, queues and forwards keyboard events from GCKeyboard to UE's input system 
+ */
+class FAppleKeyboardController
+{
+public:
+	explicit FAppleKeyboardController(const TSharedRef<FGenericApplicationMessageHandler>& InMessageHandler);
+	~FAppleKeyboardController();
+
+	/**
+	 * Sets the message handler to which any input even will be forwarded to
+	 * @param InMessageHandler New Message handler instance
+	 */
+	void SetMessageHandler(const TSharedRef<FGenericApplicationMessageHandler>& InMessageHandler);
+	
+	/** Process and dispatches all queued key events to UE's input system*/
+	void SendControllerEvents();
+	
+	/** Returns the current state of all modifier keys */
+	FModifierKeysState GetModifierKeys() const;
+	
+	DECLARE_DELEGATE_RetVal_OneParam(uint32, FKeyCodeTranslator, uint32);
+	
+	/**
+	 * Sets an external callback that will perform any key code translation from GC Keycodes to an ASCII character if needed.
+	 * @param InExternalTranslatorDelegate Callback that performs the transaltion
+	 */
+	void SetGCKeyCodeToCharTranslator(const FKeyCodeTranslator& InExternalTranslatorDelegate);
+	
+	/**
+	 * Sets an external callback that will perform any key code translation from GC Keycodes to other keycode format if needed
+	 * @param InExternalTranslatorDelegate Callback that performs the transaltion
+	 */
+	void SetGCKeyCodeToPlatformCodeTranslator(const FKeyCodeTranslator& InExternalTranslatorDelegate);
+
+	/**
+	 * Sets a flag to make sure in the next update look any currently key tracked as being pressed, re-queries the pressed state from the GCKeyboard API
+	 */
+	void ForceUpdatePressedKeysState();
+
+	/**
+	 * If set to true, keyboard input from GCKeyboard will be ignored 
+	 * @param bNewIsInhibited new inhibited state
+	 */
+	void SetInhibited(bool bNewIsInhibited);
+
+	/** Returns true if we are ignoring GC Keyboard inputs */
+	bool IsInhibited() const
+	{
+		return bIsKeyboardInhibited;
+	}
+
+	/** Returns true if any physical keyboard is connected */
+	bool IsAnyKeyboardConnected() const;
+	
+	DECLARE_MULTICAST_DELEGATE(FKeyboardAvailabilityChanged)
+
+	/**
+	 * Delegate that is executed if a keyboard is connected or disconnected
+	 */
+	FKeyboardAvailabilityChanged& OnAvailabilityChanged()
+	{
+		return AvailabilityChangedDelegate;
+	}
+
+protected:
+	
+	FKeyboardAvailabilityChanged AvailabilityChangedDelegate;
+	
+	FKeyCodeTranslator KeyCodeToPlatformCodeTranslator;
+	FKeyCodeTranslator KeyCodeToCharTranslator;
+
+	TSharedRef<FGenericApplicationMessageHandler> MessageHandler;
+	
+	enum class EKeyEvent : uint8
+	{
+		Invalid,
+		KeyUp,
+		KeyDown
+	};
+
+	struct FDeferredEvent
+	{
+		EKeyEvent KeyEventType = EKeyEvent::Invalid;
+		uint32 GCKeyCode = 0;
+		uint32 TranslatedKeyCode = 0;
+		uint32 CharCode = 0;
+		bool bIsRepeat = false;
+		double LastRepeatTime = 0.0;
+	};
+	
+	struct FInternalModifierKeyState
+	{
+		bool bIsLeftShiftDown = false;
+		bool bIsRightShiftDown = false;
+		bool bIsLeftControlDown = false;
+		bool bIsRightControlDown = false;
+		bool bIsLeftAltDown = false;
+		bool bIsRightAltDown = false;
+		bool bIsLeftCommandDown = false;
+		bool bIsRightCommandDown = false;
+		bool bAreCapsLocked = false;
+		
+		FModifierKeysState ToUEModifierKeyStates() const
+		{
+			return FModifierKeysState(bIsLeftShiftDown, bIsRightShiftDown, bIsLeftControlDown, bIsRightControlDown, bIsLeftAltDown, bIsRightAltDown, bIsLeftCommandDown, bIsRightCommandDown, bAreCapsLocked);
+		}
+	};
+	
+	void HandleKeyboardConnected();
+	void HandleKeyboardDisconnected();
+	void HandleKeyChangeEvent(GCKeyCode keyCode, BOOL pressed);
+	
+	EKeyEvent GetUpdatedKeyStateForKey(GCKeyCode KeyCode);
+	
+	void UpdateModifierKeyState(bool& OutModifierKeyFlag, GCKeyCode ModifierKeyCode, const FDeferredEvent& InKeyEvent);
+	
+	void ProcessModifierKeyEvent(const FDeferredEvent& InDeferredEvent);
+	void DispatchEvent(const FDeferredEvent& InDeferredEvent);
+	void UpdateAndDispatchRepeatEvents();
+	void UpdateEventTracking(const FDeferredEvent& InDeferredEvent, double CurrentTime);
+	
+	bool IsPrintableKey(uint32 KeyCode) const;
+	
+	void RegisterNonPrintableKeys();
+
+	FInternalModifierKeyState CurrentModifierKeys;
+	
+	bool bHandledKeyboardConnection = false;
+	
+	TMap<uint32, FDeferredEvent> CurrentlyPressedKeyStates;
+	
+	TSpscQueue<FDeferredEvent> DeferredKeyEvents;
+	
+	std::atomic<bool> bNeedUpdatePressedKeyState = false;
+	std::atomic<bool> bIsKeyboardInhibited = false;
+
+	bool bIsInitialized = false;
+	
+	id ConnectedEventObserverID = nil;
+	id DisconnectedEventObserverID = nil;
+	
+	TSet<uint32> NonPrintableKeys;
+};
diff --git a/Engine/Source/Runtime/ApplicationCore/Private/Apple/AppleMouseController.cpp b/Engine/Source/Runtime/ApplicationCore/Private/Apple/AppleMouseController.cpp
new file mode 100644
--- /dev/null
+++ b/Engine/Source/Runtime/ApplicationCore/Private/Apple/AppleMouseController.cpp
@@ -0,0 +1,188 @@
+// Copyright Epic Games, Inc. All Rights Reserved.
+
+#include "AppleMouseController.h"
+
+#include "AppleControllerInterface.h"
+#include "Math/UnrealMathUtility.h"
+
+#import "GameController/GCControllerButtonInput.h"
+#import "GameController/GCDeviceCursor.h"
+#import "GameController/GCMouseInput.h"
+
+FAppleMouseController::FAppleMouseController(const TSharedRef<FGenericApplicationMessageHandler>& InMessageHandler) : MessageHandler(InMessageHandler)
+{
+	if (@available(macOS 11.0, iOS 14.0, *))
+	{
+		ConnectedEventObserverID = [[NSNotificationCenter defaultCenter] addObserverForName:GCMouseDidConnectNotification 
+																		object:nil
+																		queue:[NSOperationQueue mainQueue]
+																		usingBlock:^(NSNotification* Notification) { HandleMouseConnected((GCMouse *) Notification.object); }];
+		
+		DisconnectedEventObserverID = [[NSNotificationCenter defaultCenter] addObserverForName:GCMouseDidDisconnectNotification
+																			object:nil 
+																			queue:[NSOperationQueue mainQueue]
+																			usingBlock:^(NSNotification* Notification) { HandleMouseDisconnected(); }];
+		
+		for (GCMouse* ConnectedMouse in GCMouse.mice)
+		{
+			HandleMouseConnected(ConnectedMouse);
+		}
+		
+		bIsInitialized = true;
+	}
+	else
+	{
+		UE_LOG(LogAppleController, Warning, TEXT("GCMouse API is not available in the current OS version. Physical mouse support via FAppleMouseController will not be available."));	
+	}
+}
+
+FAppleMouseController::~FAppleMouseController()
+{
+	if (@available(macOS 11.0, iOS 14.0, *))
+	{
+		if (ConnectedEventObserverID)
+		{
+			[[NSNotificationCenter defaultCenter] removeObserver:ConnectedEventObserverID];
+		}
+		
+		if (DisconnectedEventObserverID)
+		{
+			[[NSNotificationCenter defaultCenter] removeObserver:DisconnectedEventObserverID];
+		}
+		
+
+		for (GCMouse* ConnectedMouse in GCMouse.mice)
+		{
+			GCMouseInput* MouseInput = ConnectedMouse ? ConnectedMouse.mouseInput : nullptr;
+			
+			if (MouseInput)
+			{
+				MouseInput.mouseMovedHandler = nil;
+				MouseInput.scroll.valueChangedHandler = nil;
+				MouseInput.leftButton.pressedChangedHandler = nil;
+				MouseInput.rightButton.pressedChangedHandler = nil;
+				MouseInput.middleButton.pressedChangedHandler = nil;
+				
+				constexpr uint64 MaxAuxButtonIndex = EMouseButtons::Type::Invalid - EMouseButtons::Type::Thumb01;
+				const uint64 CurrentMaxAuxButtonIndex = MaxAuxButtonIndex > MouseInput.auxiliaryButtons.count ? MouseInput.auxiliaryButtons.count : MaxAuxButtonIndex;
+
+				for (uint64 AuxButtonIndex = 0; AuxButtonIndex < CurrentMaxAuxButtonIndex; AuxButtonIndex++)
+				{
+					MouseInput.auxiliaryButtons[AuxButtonIndex].pressedChangedHandler = nil;
+				}
+			}
+		}
+	}
+}
+
+void FAppleMouseController::HandleMouseConnected(GCMouse* Mouse)
+{
+	UE_LOG(LogAppleController, Log, TEXT("[%hs] Handling physical mouse connection..."), __func__);
+	
+	GCMouseInput* MouseInput = Mouse ? Mouse.mouseInput : nullptr;
+	
+	if (ensure(MouseInput))
+	{
+		MouseInput.mouseMovedHandler = ^(GCMouseInput * mouse, float deltaX, float deltaY)
+		{
+			CurrentMovementState.MouseDeltaX += deltaX;
+			CurrentMovementState.MouseDeltaY -= deltaY;
+		};
+		
+		MouseInput.scroll.valueChangedHandler = ^(GCControllerDirectionPad* dpad, float xValue, float yValue) 
+		{
+			CurrentMovementState.ScrollDelta += yValue;
+		};
+		
+		MouseInput.leftButton.pressedChangedHandler = ^(GCControllerButtonInput* eventButton, float value, BOOL pressed) 
+		{
+			HandleKeyEvent(pressed, EMouseButtons::Type::Left);
+		};
+
+		MouseInput.rightButton.pressedChangedHandler = ^(GCControllerButtonInput* eventButton, float value, BOOL pressed) {
+			HandleKeyEvent(pressed, EMouseButtons::Type::Right);
+		};
+
+		MouseInput.middleButton.pressedChangedHandler = ^(GCControllerButtonInput* eventButton, float value, BOOL pressed) 
+		{
+			HandleKeyEvent(pressed, EMouseButtons::Type::Middle);
+		};
+		
+		constexpr uint64 MaxAuxButtonIndex = EMouseButtons::Type::Invalid - EMouseButtons::Type::Thumb01;
+		const uint64 CurrentMaxAuxButtonIndex = MaxAuxButtonIndex > MouseInput.auxiliaryButtons.count ? MouseInput.auxiliaryButtons.count : MaxAuxButtonIndex;
+
+		for (uint64 AuxButtonIndex = 0; AuxButtonIndex < CurrentMaxAuxButtonIndex; AuxButtonIndex++)
+		{
+			MouseInput.auxiliaryButtons[AuxButtonIndex].pressedChangedHandler = ^(GCControllerButtonInput* eventButton, float value, BOOL pressed) 
+			{
+				const EMouseButtons::Type TargetButton = static_cast<EMouseButtons::Type>(EMouseButtons::Type::Thumb01 + AuxButtonIndex);
+				HandleKeyEvent(pressed, TargetButton);
+			};
+		}
+	}
+}
+
+void FAppleMouseController::HandleKeyEvent(bool bIsPressed, EMouseButtons::Type Button)
+{
+	FDeferredKeyEvent KeyEvent;
+	KeyEvent.ButtonID = Button;
+	KeyEvent.Type = bIsPressed ? EKeyEventType::Down : EKeyEventType::Up;
+	DeferredKeyEvents.Enqueue(KeyEvent);
+}
+
+void FAppleMouseController::DispatchMovementUpdates()
+{
+	MessageHandler->OnRawMouseMove(FMath::TruncToInt(CurrentMovementState.MouseDeltaX.exchange(0.0f)), FMath::TruncToInt(CurrentMovementState.MouseDeltaY.exchange(0.0f)));
+	MessageHandler->OnMouseWheel(CurrentMovementState.ScrollDelta.exchange(0.0f));
+}
+
+void FAppleMouseController::DispatchKeyUpdates()
+{
+	FDeferredKeyEvent DeferredEvent;
+	while(DeferredKeyEvents.Dequeue(DeferredEvent))
+	{
+		DispatchKeyUpdate(DeferredEvent);
+	}
+}
+
+void FAppleMouseController::DispatchKeyUpdate(const FDeferredKeyEvent& Event)
+{
+	switch (Event.Type) 
+	{
+		case EKeyEventType::Up:
+			{
+				MessageHandler->OnMouseUp(Event.ButtonID);
+				break;
+			}
+		case EKeyEventType::Down:
+			{
+				MessageHandler->OnMouseDown(nullptr, Event.ButtonID);
+				break;
+			}
+		case EKeyEventType::Invalid:
+		default:
+			ensure(false);
+			break;
+	}
+}
+
+void FAppleMouseController::HandleMouseDisconnected()
+{
+	UE_LOG(LogAppleController, Log, TEXT("[%hs] Handling physical mouse disconnection..."), __func__);
+}
+
+void FAppleMouseController::SetMessageHandler(const TSharedRef<FGenericApplicationMessageHandler>& InMessageHandler)
+{
+	MessageHandler = InMessageHandler;
+}
+
+void FAppleMouseController::SendControllerEvents()
+{
+	if (!bIsInitialized)
+	{
+		return;
+	}
+
+	DispatchMovementUpdates();
+	DispatchKeyUpdates();
+}
diff --git a/Engine/Source/Runtime/ApplicationCore/Private/Apple/AppleMouseController.h b/Engine/Source/Runtime/ApplicationCore/Private/Apple/AppleMouseController.h
new file mode 100644
--- /dev/null
+++ b/Engine/Source/Runtime/ApplicationCore/Private/Apple/AppleMouseController.h
@@ -0,0 +1,71 @@
+// Copyright Epic Games, Inc. All Rights Reserved.
+
+#pragma once
+
+#include "GenericPlatform/GenericApplicationMessageHandler.h"
+#include "Templates/SharedPointer.h"
+#include <atomic>
+
+#include "Containers/SpscQueue.h"
+
+#import "GameController/GCMouse.h"
+
+/**
+ * Class that intercepts, queues and dispatches GCMouse events to UE's input system
+ */
+class FAppleMouseController
+{
+public:
+	explicit FAppleMouseController(const TSharedRef<FGenericApplicationMessageHandler>& InMessageHandler);
+	
+	~FAppleMouseController();
+	
+	/**
+	 * Sets the message handler to which any input even will be forwarded to
+	 * @param InMessageHandler New Message handler instance
+	 */
+	void SetMessageHandler(const TSharedRef<FGenericApplicationMessageHandler>& InMessageHandler);
+	
+	/** Process and dispatches all queued key events to UE's input system */
+	void SendControllerEvents();
+
+protected:
+	
+	enum class EKeyEventType
+	{
+		Invalid,
+		Up,
+		Down
+	};
+	
+	struct FDeferredKeyEvent
+	{
+		EKeyEventType Type = EKeyEventType::Invalid;
+		EMouseButtons::Type ButtonID;
+	};
+
+	void DispatchMovementUpdates();
+	void DispatchKeyUpdates();
+	void DispatchKeyUpdate(const FDeferredKeyEvent& Event);
+	
+	void HandleMouseConnected(GCMouse* Mouse);
+	void HandleMouseDisconnected();
+	void HandleKeyEvent(bool bIsPressed, EMouseButtons::Type Button);
+	
+	struct FMouseMovementState
+	{
+		std::atomic<float> MouseDeltaX = 0.0f;
+		std::atomic<float> MouseDeltaY = 0.0f;
+		std::atomic<float> ScrollDelta = 0.0f;
+	};
+	
+	FMouseMovementState CurrentMovementState;
+
+	TSpscQueue<FDeferredKeyEvent> DeferredKeyEvents;
+	TSharedRef<FGenericApplicationMessageHandler> MessageHandler;
+	
+	bool bIsInitialized = false;
+	
+	id ConnectedEventObserverID = nil;
+	id DisconnectedEventObserverID = nil;
+};
diff --git a/Engine/Source/Runtime/ApplicationCore/Private/IOS/IOSApplication.cpp b/Engine/Source/Runtime/ApplicationCore/Private/IOS/IOSApplication.cpp
--- a/Engine/Source/Runtime/ApplicationCore/Private/IOS/IOSApplication.cpp
+++ b/Engine/Source/Runtime/ApplicationCore/Private/IOS/IOSApplication.cpp
@@ -34,6 +34,16 @@ FIOSApplication::FIOSApplication()
 #endif
 {
 	[IOSAppDelegate GetDelegate].IOSApplication = this;
+	
+	if (TSharedPtr<FAppleKeyboardController> KeyboardController = InputInterface->GetKeyboardController().Pin())
+	{
+		KeyboardController->OnAvailabilityChanged().AddRaw(this, &FIOSApplication::HandlePhysicalKeyboardConnectionChanged);
+		
+		if (KeyboardController->IsAnyKeyboardConnected())
+		{
+			HandlePhysicalKeyboardConnectionChanged();
+		}
+	}
 }
 
 FIOSApplication::~FIOSApplication()
@@ -47,6 +57,42 @@ FIOSApplication::~FIOSApplication()
 		[AccessibilityAnnouncementDelayTimer release ];
 	}
 #endif
+	
+	if (TSharedPtr<FAppleKeyboardController> KeyboardController = InputInterface->GetKeyboardController().Pin())
+	{
+		KeyboardController->OnAvailabilityChanged().RemoveAll(this);
+	}
+}
+
+void FIOSApplication::HandlePhysicalKeyboardConnectionChanged()
+{
+	FIOSView *View = [[IOSAppDelegate GetDelegate] IOSView];
+	
+	if (!ensure(View))
+	{
+		return;
+	}
+
+	// TODO: When the dedicated view for external keyboard is implemented, we need to make it first responder here
+
+	if (InputInterface->IsAnyPhysicalKeyboardConnected())
+	{
+		UE_LOG(LogIOS, Log, TEXT("[%hs] External keyboard connected."), __func__);
+	}
+	else
+	{
+		UE_LOG(LogIOS, Log, TEXT("[%hs] All external keyboards disconnected"), __func__);
+	}
+}
+
+FModifierKeysState FIOSApplication::GetModifierKeys() const
+{ 
+	if (TSharedPtr<FAppleKeyboardController> KeyboardController = InputInterface->GetKeyboardController().Pin())
+	{
+		return KeyboardController->GetModifierKeys();
+	}
+
+	return FModifierKeysState(); 
 }
 
 void FIOSApplication::InitializeWindow( const TSharedRef< FGenericWindow >& InWindow, const TSharedRef< FGenericWindowDefinition >& InDefinition, const TSharedPtr< FGenericWindow >& InParent, const bool bShowImmediately )
diff --git a/Engine/Source/Runtime/ApplicationCore/Private/IOS/IOSInputInterface.cpp b/Engine/Source/Runtime/ApplicationCore/Private/IOS/IOSInputInterface.cpp
--- a/Engine/Source/Runtime/ApplicationCore/Private/IOS/IOSInputInterface.cpp
+++ b/Engine/Source/Runtime/ApplicationCore/Private/IOS/IOSInputInterface.cpp
@@ -11,6 +11,9 @@
 
 #import <AudioToolbox/AudioToolbox.h>
 
+#include "Apple/AppleControllerInterface.h"
+#include "Apple/AppleMouseController.h"
+
 DECLARE_LOG_CATEGORY_EXTERN(LogIOSInput, Log, All);
 
 #define APPLE_CONTROLLER_DEBUG 0
@@ -81,7 +84,7 @@ uint32 TranslateGCKeyCodeToASCII(GCKeyCode KeyCode)
         else if (KeyCode == GCKeyCodeSlash) { c = '/';}
         else if (KeyCode == GCKeyCodeNonUSBackslash) { c = 167;} // Section (§) in FGenericPlatformInput::GetStandardPrintableKeyMap
         
-        UE_LOG(LogIOS, Log, TEXT("char: %c - GCKeycode: %d"), TCHAR(c), KeyCode);
+        UE_LOG(LogIOS, Verbose, TEXT("char: %c - GCKeycode: %d"), TCHAR(c), KeyCode);
 	}
 
     return c;
@@ -129,7 +132,7 @@ uint32 TranslateNonCharKeys(GCKeyCode KeyCode)
 	else if (KeyCode == GCKeyCodeF20) { c = KEYCODE_F20;}
 	else { c = TranslateGCKeyCodeToASCII(KeyCode); }
 
-	UE_LOG(LogIOS, Log, TEXT("KeyCode: %d - GCKeycode: %d"), c, KeyCode);
+	UE_LOG(LogIOS, Verbose, TEXT("KeyCode: %d - GCKeycode: %d"), c, KeyCode);
 
 	return c;
 }
@@ -141,7 +144,6 @@ static FCriticalSection CriticalSection;
 static TArray<TouchInput> TouchInputStack;
 static TArray<int32> KeyInputStack;
 
-bool FIOSInputInterface::bKeyboardInhibited = false;
 FName FIOSInputInterface::InputClassName_DefaultMobileTouch;
 FString FIOSInputInterface::HardwareDeviceIdentifier_DefaultMobileTouch;
 
@@ -155,12 +157,20 @@ FIOSInputInterface::FIOSInputInterface( const TSharedRef< FGenericApplicationMes
 	, bAllowRemoteRotation(false)
 	, bGameSupportsMultipleActiveControllers(false)
     , LastHapticValue(0.0f)
-    , MouseDeltaX(0)
-    , MouseDeltaY(0)
-    , ScrollDeltaY(0)
-    , bHaveMouse(false)
+	, KeyboardController(MakeShared<FAppleKeyboardController>(InMessageHandler))
+	, MouseController(MakeShared<FAppleMouseController>(InMessageHandler))
 {
     SCOPED_BOOT_TIMING("FIOSInputInterface::FIOSInputInterface");
+	
+	KeyboardController->SetGCKeyCodeToCharTranslator(FAppleKeyboardController::FKeyCodeTranslator::CreateLambda([](uint32 KeyCode)
+	{
+		return TranslateGCKeyCodeToASCII(KeyCode);
+	}));
+
+	KeyboardController->SetGCKeyCodeToPlatformCodeTranslator(FAppleKeyboardController::FKeyCodeTranslator::CreateLambda([](uint32 KeyCode)
+	{
+		return TranslateNonCharKeys(KeyCode);
+	}));
 
 #if !PLATFORM_TVOS
     MotionManager = nil;
@@ -177,37 +187,7 @@ FIOSInputInterface::FIOSInputInterface( const TSharedRef< FGenericApplicationMes
     
     NSNotificationCenter* notificationCenter = [NSNotificationCenter defaultCenter];
     NSOperationQueue* currentQueue = [NSOperationQueue currentQueue];
-    
-    [notificationCenter addObserverForName:GCMouseDidConnectNotification object:nil queue:currentQueue usingBlock:^(NSNotification* Notification)
-     {
-        HandleMouseConnection(Notification.object);
-    }];
-    
-    [notificationCenter addObserverForName:GCMouseDidDisconnectNotification object:nil queue:currentQueue usingBlock:^(NSNotification* Notification)
-     {
-        HandleMouseDisconnect(Notification.object);
-    }];
-    
-    [notificationCenter addObserverForName:GCKeyboardDidConnectNotification object:nil queue:currentQueue usingBlock:^(NSNotification* Notification)
-     {
-        HandleKeyboardConnection(Notification.object);
-        
-    }];
-    
-    [notificationCenter addObserverForName:GCKeyboardDidDisconnectNotification object:nil queue:currentQueue usingBlock:^(NSNotification* Notification)
-     {
-        HandleKeyboardDisconnect(Notification.object);
-    }];
-    
-    if ( GCMouse.current )
-    {
-        HandleMouseConnection( GCMouse.current );
-    }
-    
-    if ( GCKeyboard.coalescedKeyboard )
-    {
-        HandleKeyboardConnection( GCKeyboard.coalescedKeyboard );
-    }
+
     if (!bGameSupportsMultipleActiveControllers)
     {
         [notificationCenter addObserverForName:GCControllerDidBecomeCurrentNotification object:nil queue:currentQueue usingBlock:^(NSNotification* Notification)
@@ -252,88 +232,6 @@ FIOSInputInterface::FIOSInputInterface( const TSharedRef< FGenericApplicationMes
 	HardwareDeviceIdentifier_DefaultMobileTouch = TEXT("MobileTouch");
 }
 
-template <EIOSEventType U, EIOSEventType V>
-static inline void HandleButtons(GCControllerButtonInput* _Nonnull eventButton, float value, BOOL pressed, TArray<FDeferredIOSEvent>& DeferredEvents, FCriticalSection* EventsMutex)
-{
-    (void)eventButton;
-    (void)value;
-    FDeferredIOSEvent DeferredEvent;
-    DeferredEvent.type = pressed ? U : V;
-    FScopeLock Lock(EventsMutex);
-    DeferredEvents.Add(DeferredEvent);
-}
-
-void FIOSInputInterface::HandleMouseConnection(GCMouse* Mouse)
-{
-    bHaveMouse = true;
-    MouseDeltaX = 0.f;
-    MouseDeltaY = 0.f;
-    Mouse.mouseInput.mouseMovedHandler = ^(GCMouseInput * _Nonnull eventMouse, float deltaX, float deltaY) {
-        (void)eventMouse;
-        this->MouseDeltaX += deltaX;
-        this->MouseDeltaY -= deltaY;
-    };
-
-    Mouse.mouseInput.leftButton.pressedChangedHandler = ^(GCControllerButtonInput* _Nonnull eventButton, float value, BOOL pressed) {
-        HandleButtons< EIOSEventType::LeftMouseDown, EIOSEventType::LeftMouseUp >( eventButton, value, pressed, DeferredEvents, &EventsMutex );
-    };
-
-    Mouse.mouseInput.rightButton.pressedChangedHandler = ^(GCControllerButtonInput* _Nonnull eventButton, float value, BOOL pressed) {
-        HandleButtons< EIOSEventType::RightMouseDown, EIOSEventType::RightMouseUp >( eventButton, value, pressed, DeferredEvents, &EventsMutex );
-    };
-
-    Mouse.mouseInput.middleButton.pressedChangedHandler = ^(GCControllerButtonInput* _Nonnull eventButton, float value, BOOL pressed) {
-        HandleButtons< EIOSEventType::MiddleMouseDown, EIOSEventType::MiddleMouseUp >( eventButton, value, pressed, DeferredEvents, &EventsMutex );
-    };
-
-    NSArray<GCDeviceButtonInput*>* auxButtons = Mouse.mouseInput.auxiliaryButtons;
-    if (auxButtons && auxButtons.count > 0)
-    {
-        if ( auxButtons.count < 2 )
-        {
-            auxButtons[0].pressedChangedHandler = ^(GCControllerButtonInput* _Nonnull eventButton, float value, BOOL pressed){
-                HandleButtons< EIOSEventType::ThumbDown + 0, EIOSEventType::ThumbUp + 0 >( eventButton, value, pressed, DeferredEvents, &EventsMutex );
-            };
-        }
-        if ( auxButtons.count < 3 )
-        {
-            auxButtons[0].pressedChangedHandler = ^(GCControllerButtonInput* _Nonnull eventButton, float value, BOOL pressed){
-                HandleButtons< EIOSEventType::ThumbDown + 1, EIOSEventType::ThumbUp + 1 >( eventButton, value, pressed, DeferredEvents, &EventsMutex );
-            };
-        }
-    }
-
-    Mouse.mouseInput.scroll.valueChangedHandler = ^(GCControllerDirectionPad* _Nonnull dpad, float xValue, float yValue) {
-        (void)xValue;
-        ScrollDeltaY += yValue;
-    };
-}
-
-void FIOSInputInterface::HandleMouseDisconnect(GCMouse* Mouse)
-{
-    bHaveMouse = false;
-}
-
-void FIOSInputInterface::HandleKeyboardConnection(GCKeyboard* Keyboard)
-{
-    Keyboard.keyboardInput.keyChangedHandler = ^(GCKeyboardInput * _Nonnull keyboard, GCControllerButtonInput * _Nonnull key, GCKeyCode keyCode, BOOL pressed) {
-        if ( !(FIOSInputInterface::IsKeyboardInhibited()) )
-        {
-            FDeferredIOSEvent DeferredEvent;
-            DeferredEvent.type = pressed ? EIOSEventType::KeyDown : EIOSEventType::KeyUp;
-            DeferredEvent.keycode = TranslateNonCharKeys(keyCode);
-            DeferredEvent.charcode = TranslateGCKeyCodeToASCII(keyCode);
-
-            
-            FScopeLock Lock(&EventsMutex);
-            DeferredEvents.Add(DeferredEvent);
-        }
-    };
-}
-
-void FIOSInputInterface::HandleKeyboardDisconnect(GCKeyboard* Keyboard)
-{
-}
 
 #if !PLATFORM_TVOS
 void ModifyVectorByOrientation(FVector& Vec, bool bIsRotation)
@@ -442,93 +340,11 @@ void FIOSInputInterface::ProcessTouchesAndKeys(uint32 ControllerId, const TArray
     }
 }
 
-void FIOSInputInterface::ProcessDeferredEvents()
-{
-    TArray<FDeferredIOSEvent> EventsToProcess;
-
-    EventsMutex.Lock();
-    EventsToProcess.Append(DeferredEvents);
-    DeferredEvents.Empty();
-    EventsMutex.Unlock();
-
-    for (uint32_t Index = 0; Index < EventsToProcess.Num(); ++Index)
-    {
-        ProcessEvent(EventsToProcess[Index]);
-    }
-}
-
-void FIOSInputInterface::ProcessEvent(const FDeferredIOSEvent& Event)
-{
-    if (Event.type != EIOSEventType::Invalid)
-    {
-        switch (Event.type)
-        {
-            case EIOSEventType::KeyDown:
-            {
-                MessageHandler->OnKeyDown(Event.keycode, Event.charcode, false);
-                break;
-            }
-            case EIOSEventType::KeyUp:
-            {
-                MessageHandler->OnKeyUp(Event.keycode, Event.charcode, false);
-                break;
-            }
-            case EIOSEventType::LeftMouseDown:
-            {
-                MessageHandler->OnMouseDown(nullptr, EMouseButtons::Left);
-                break;
-            }
-            case EIOSEventType::LeftMouseUp:
-            {
-                MessageHandler->OnMouseUp(EMouseButtons::Left);
-                break;
-            }
-            case EIOSEventType::RightMouseDown:
-            {
-                MessageHandler->OnMouseDown(nullptr, EMouseButtons::Right);
-                break;
-            }
-            case EIOSEventType::RightMouseUp:
-            {
-                MessageHandler->OnMouseUp(EMouseButtons::Right);
-                break;
-            }
-            case EIOSEventType::MiddleMouseDown:
-            {
-                MessageHandler->OnMouseDown(nullptr, EMouseButtons::Middle);
-                break;
-            }
-            case EIOSEventType::MiddleMouseUp:
-            {
-                MessageHandler->OnMouseUp(EMouseButtons::Middle);
-                break;
-            }
-            case EIOSEventType::ThumbDown+0:
-            {
-                MessageHandler->OnMouseDown(nullptr, EMouseButtons::Thumb01);
-                break;
-            }
-            case EIOSEventType::ThumbUp+0:
-            {
-                MessageHandler->OnMouseUp(EMouseButtons::Thumb01);
-                break;
-            }
-            case EIOSEventType::ThumbDown+1:
-            {
-                MessageHandler->OnMouseDown(nullptr, EMouseButtons::Thumb02);
-                break;
-            }
-            case EIOSEventType::ThumbUp+1:
-            {
-                MessageHandler->OnMouseUp(EMouseButtons::Thumb02);
-                break;
-            }
-        }
-    }
-}
-
 void FIOSInputInterface::SendControllerEvents()
 {
+	KeyboardController->SendControllerEvents();
+	MouseController->SendControllerEvents();
+
     TArray<TouchInput> LocalTouchInputStack;
     TArray<int32> LocalKeyInputStack;
     {
@@ -541,7 +357,6 @@ void FIOSInputInterface::SendControllerEvents()
     // on ios, touches always go go player 0
     ProcessTouchesAndKeys(0, LocalTouchInputStack, LocalKeyInputStack);
 #endif
-    ProcessDeferredEvents();
 
     
 #if !PLATFORM_TVOS // @todo tvos: This needs to come from the Microcontroller rotation
@@ -633,20 +448,17 @@ void FIOSInputInterface::SendControllerEvents()
     }
 #endif
     
-    if ( bHaveMouse )
-    {
-        MessageHandler->OnRawMouseMove(FMath::TruncToInt(MouseDeltaX), FMath::TruncToInt(MouseDeltaY));
-        MouseDeltaX = 0.f;
-        MouseDeltaY = 0.f;
-        
-        MessageHandler->OnMouseWheel( ScrollDeltaY );
-        ScrollDeltaY = 0.f;
-    }
-    
     // Generic Controller update
     FAppleControllerInterface::SendControllerEvents();
 }
 
+void FIOSInputInterface::SetMessageHandler(const TSharedRef<FGenericApplicationMessageHandler>& InMessageHandler)
+{
+	FAppleControllerInterface::SetMessageHandler(InMessageHandler);
+	KeyboardController->SetMessageHandler(InMessageHandler);
+	MouseController->SetMessageHandler(InMessageHandler);
+}
+
 void FIOSInputInterface::QueueTouchInput(const TArray<TouchInput>& InTouchEvents)
 {
     FScopeLock Lock(&CriticalSection);
@@ -685,6 +497,26 @@ bool FIOSInputInterface::IsMotionDataEnabled() const
     return !bPauseMotion;
 }
 
+TWeakPtr<FAppleKeyboardController> FIOSInputInterface::GetKeyboardController()
+{
+	return KeyboardController;
+}
+
+void FIOSInputInterface::SetKeyboardControllerInhibited(bool bInhibited)
+{
+	return KeyboardController->SetInhibited(bInhibited);
+}
+
+bool FIOSInputInterface::IsKeyboardControllerInhibited() const
+{
+	return KeyboardController->IsInhibited();
+}
+
+bool FIOSInputInterface::IsAnyPhysicalKeyboardConnected() const
+{
+	return KeyboardController->IsAnyKeyboardConnected();
+}
+
 void FIOSInputInterface::GetMovementData(FVector& Attitude, FVector& RotationRate, FVector& Gravity, FVector& Acceleration)
 {
 #if !PLATFORM_TVOS
diff --git a/Engine/Source/Runtime/ApplicationCore/Private/IOS/IOSView.cpp b/Engine/Source/Runtime/ApplicationCore/Private/IOS/IOSView.cpp
--- a/Engine/Source/Runtime/ApplicationCore/Private/IOS/IOSView.cpp
+++ b/Engine/Source/Runtime/ApplicationCore/Private/IOS/IOSView.cpp
@@ -231,7 +231,6 @@ MTL::Device* GMetalDevice = nullptr;
 	[self CreateFramebuffer];
 	
 #endif
-	
 
 	return self;
 }
@@ -646,16 +645,6 @@ self.accessibilityElements = @[Window.accessibilityContainer];
 }
 
 
-
-
-
-
-
-
-
-
-
-
 #pragma mark Keyboard
 
 
@@ -676,6 +665,11 @@ self.accessibilityElements = @[Window.accessibilityContainer];
 		CachedMarkedText = nil;
 	}
 
+	if ([self IsAnyPhysicalKeyboardConnected])
+	{
+		return;
+	}
+
 	// insert text one key at a time, as chars, not keydowns
 	for (int32 CharIndex = 0; CharIndex < [theText length]; CharIndex++)
 	{
@@ -707,6 +701,13 @@ self.accessibilityElements = @[Window.accessibilityContainer];
 		[CachedMarkedText release];
 		CachedMarkedText = nil;
 	}
+
+
+	if ([self IsAnyPhysicalKeyboardConnected])
+	{
+		return;
+	}
+
 	FIOSInputInterface::QueueKeyInput(KEYCODE_BACKSPACE, '\b');
 }
 
@@ -742,9 +743,8 @@ self.accessibilityElements = @[Window.accessibilityContainer];
 			{
 				[self endEditing:YES];
 			}
-			[self becomeFirstResponder];
 			
-			FIOSInputInterface::SetKeyboardInhibited(true);
+			[self becomeFirstResponder];
 		}
 		
 		FPlatformAtomics::InterlockedDecrement(&KeyboardShowCount);
@@ -765,7 +765,7 @@ self.accessibilityElements = @[Window.accessibilityContainer];
 			{
 				// Dismiss the existing keyboard, if one exists.
 				[self endEditing:YES];
-                FIOSInputInterface::SetKeyboardInhibited(false);
+
 			}
 		}
 	});
@@ -780,6 +780,13 @@ self.accessibilityElements = @[Window.accessibilityContainer];
 	return FParse::Param(FCommandLine::Get(), TEXT("NewKeyboard")) || bUseIntegratedKeyboard;
 }
 
+-(bool)IsAnyPhysicalKeyboardConnected;
+{
+	FIOSInputInterface* InputInterface = [IOSAppDelegate GetDelegate].IOSApplication->GetIOSInputInterface();
+
+	return InputInterface ? InputInterface->IsAnyPhysicalKeyboardConnected() : false;
+}
+
 -(BOOL)becomeFirstResponder
 {
 	volatile int32 ShowCount = KeyboardShowCount;
diff --git a/Engine/Source/Runtime/ApplicationCore/Public/IOS/IOSApplication.h b/Engine/Source/Runtime/ApplicationCore/Public/IOS/IOSApplication.h
--- a/Engine/Source/Runtime/ApplicationCore/Public/IOS/IOSApplication.h
+++ b/Engine/Source/Runtime/ApplicationCore/Public/IOS/IOSApplication.h
@@ -5,6 +5,7 @@
 #include "GenericPlatform/Accessibility/GenericAccessibleInterfaces.h"
 #include "GenericPlatform/GenericApplication.h"
 
+class FIOSInputInterface;
 class FIOSWindow;
 class IModularFeature;
 
@@ -18,6 +19,10 @@ public:
 public:
 
 	virtual ~FIOSApplication();
+	
+	void HandlePhysicalKeyboardConnectionChanged();
+	
+	virtual FModifierKeysState GetModifierKeys() const override;
 
 	void SetMessageHandler( const TSharedRef< FGenericApplicationMessageHandler >& InMessageHandler );
 #if WITH_ACCESSIBILITY
@@ -41,6 +46,7 @@ public:
 #endif
 
 	virtual IInputInterface* GetInputInterface() override { return (IInputInterface*)InputInterface.Get(); }
+	FIOSInputInterface* GetIOSInputInterface() { return InputInterface.Get(); }
 
 	virtual bool IsGamepadAttached() const override;
 
@@ -60,7 +66,7 @@ private:
 
 private:
 
-	TSharedPtr< class FIOSInputInterface > InputInterface;
+	TSharedPtr< FIOSInputInterface > InputInterface;
 
 	/** List of input devices implemented in external modules. */
 	TArray< TSharedPtr<class IInputDevice> > ExternalInputDevices;
diff --git a/Engine/Source/Runtime/ApplicationCore/Public/IOS/IOSInputInterface.h b/Engine/Source/Runtime/ApplicationCore/Public/IOS/IOSInputInterface.h
--- a/Engine/Source/Runtime/ApplicationCore/Public/IOS/IOSInputInterface.h
+++ b/Engine/Source/Runtime/ApplicationCore/Public/IOS/IOSInputInterface.h
@@ -14,6 +14,8 @@
 #include "Math/Color.h"
 #include "GenericPlatform/GenericApplicationMessageHandler.h"
 #include "AppleControllerInterface.h"
+#include "AppleKeyboardController.h"
+#include "AppleMouseController.h"
 
 
 #define KEYCODE_ENTER 1000
@@ -113,6 +115,8 @@ public:
 	 */
 	void SendControllerEvents();
 	
+	virtual void SetMessageHandler( const TSharedRef< FGenericApplicationMessageHandler >& InMessageHandler );
+	
 	/**
 	 * IInputInterface implementation
 	 */
@@ -128,8 +132,23 @@ public:
 	void EnableMotionData(bool bEnable);
 	bool IsMotionDataEnabled() const;
     
-    static void SetKeyboardInhibited(bool bInhibited) { bKeyboardInhibited = bInhibited; }
-    static bool IsKeyboardInhibited() { return bKeyboardInhibited; }
+	UE_DEPRECATED(5.8, "Please use SetKeyboardControllerInhibited")
+	static void SetKeyboardInhibited(bool bInhibited)
+	{
+	}
+	UE_DEPRECATED(5.8, "Please use IsKeyboardControllerInhibited")
+	static bool IsKeyboardInhibited()
+	{
+		return false;
+	}
+	
+	TWeakPtr<FAppleKeyboardController> GetKeyboardController();
+	
+	void SetKeyboardControllerInhibited(bool bInhibited);
+
+	bool IsKeyboardControllerInhibited() const;
+	
+	bool IsAnyPhysicalKeyboardConnected() const;
     
     NSData* GetGamepadGlyphRawData(const FGamepadKeyNames::Type& ButtonKey, uint32 ControllerIndex);
 
@@ -143,12 +162,6 @@ private:
 
 	FIOSInputInterface( const TSharedRef< FGenericApplicationMessageHandler >& InMessageHandler );
 
-    // handle disconnect and connect events
-    void HandleMouseConnection(GCMouse* Mouse);
-    void HandleMouseDisconnect(GCMouse* Mouse);
-    void HandleKeyboardConnection(GCKeyboard* Keyboard);
-    void HandleKeyboardDisconnect(GCKeyboard* Keyboard);
-
     /**
 	 * Get the current Movement data from the device
 	 *
@@ -166,8 +179,6 @@ private:
 
 private:
 	void ProcessTouchesAndKeys(uint32 ControllerId, const TArray<TouchInput>& InTouchInputStack, const TArray<int32>& InKeyInputStack);
-    void ProcessDeferredEvents();
-    void ProcessEvent(const FDeferredIOSEvent& Event);
 	
 	// can the remote be rotated to landscape
 	bool bAllowRemoteRotation;
@@ -211,15 +222,10 @@ private:
 	float LastHapticValue;
 	
 	int HapticFeedbackSupportLevel;
-    
-    FCriticalSection EventsMutex;
-        TArray<FDeferredIOSEvent> DeferredEvents;
-        float MouseDeltaX;
-        float MouseDeltaY;
-        float ScrollDeltaY;
-        bool bHaveMouse;
-        static bool bKeyboardInhibited;
 	
 	static FName InputClassName_DefaultMobileTouch;
 	static FString HardwareDeviceIdentifier_DefaultMobileTouch;
+	
+	TSharedRef<FAppleKeyboardController> KeyboardController;
+	TSharedRef<FAppleMouseController> MouseController;
 };
diff --git a/Engine/Source/Runtime/ApplicationCore/Public/IOS/IOSView.h b/Engine/Source/Runtime/ApplicationCore/Public/IOS/IOSView.h
--- a/Engine/Source/Runtime/ApplicationCore/Public/IOS/IOSView.h
+++ b/Engine/Source/Runtime/ApplicationCore/Public/IOS/IOSView.h
@@ -117,6 +117,8 @@ APPLICATIONCORE_API
 -(void)DeactivateKeyboard;
 -(bool)ShouldUseIntegratedKeyboard;
 
+-(bool)IsAnyPhysicalKeyboardConnected;
+
 // callable from outside to fake locations
 -(void)HandleTouchAtLoc:(CGPoint)Loc PrevLoc:(CGPoint)PrevLoc TouchIndex:(int)TouchIndex Force:(float)Force Type:(TouchType)Type TouchesArray:(TArray<TouchInput>&)TouchesArray;
 #endif
