[iOS] Improving character input support for external keyboards

- Given GC keyboard does not provide keyboard layour or character translation information, now the GC keyboard implementation only handles non-character related key inputs. The rest of the inputs are consummed from UIKit
- Updated the touch character input to use the new Keyboard API which has a lock free queue system.
- Implemented the UI responder key presses handling methods in the view controller so we can foward character related key events to the new Apple keyboard API introduces in the previous change.
- Made the view controller conform to the GCViewcontroller type for all iOS platforms, not just tvOS. This allow to make sure interactions are routed to the GC Keyboard/Mouse API when we desired to do so
- Now we make the view controller the first responder when a physical keyboard is attached. This prevents the virtual keyboard from popping up among other things.
- At a platform implementation level, now they requires virtual keyboard method now dynamically returns true or false based on if we have a physical keyboard attached or now. Before that it was defined at compile time.

#Synced-CL 49059801

diff --git a/Engine/Source/Runtime/ApplicationCore/Private/Apple/AppleKeyboardController.cpp b/Engine/Source/Runtime/ApplicationCore/Private/Apple/AppleKeyboardController.cpp
--- a/Engine/Source/Runtime/ApplicationCore/Private/Apple/AppleKeyboardController.cpp
+++ b/Engine/Source/Runtime/ApplicationCore/Private/Apple/AppleKeyboardController.cpp
@@ -90,10 +90,17 @@ bool FAppleKeyboardController::IsPrintableKey(uint32 KeyCode) const
 	return !NonPrintableKeys.Contains(KeyCode);
 }
 
+bool FAppleKeyboardController::IsPrintableChar(TCHAR InChar) const
+{
+	const bool bIsPrintableASCII = InChar >= 32 && InChar <= 126;
+	const bool bIsPrintableUnicode = InChar > 160;
+	return bIsPrintableASCII || bIsPrintableUnicode;
+}
+
 void FAppleKeyboardController::RegisterNonPrintableKeys()
 {
 	// If more keys are added, increase this number to match
-	constexpr int32 NonPrintableKeysNum = 29;
+	constexpr int32 NonPrintableKeysNum = 35;
 	NonPrintableKeys.Reserve(NonPrintableKeysNum);
 	
 	NonPrintableKeys.Emplace(GCKeyCodeF1);
@@ -125,6 +132,12 @@ void FAppleKeyboardController::RegisterNonPrintableKeys()
 	NonPrintableKeys.Emplace(GCKeyCodeRightAlt);
 	NonPrintableKeys.Emplace(GCKeyCodeRightGUI);
 	NonPrintableKeys.Emplace(GCKeyCodeLeftGUI);
+	NonPrintableKeys.Emplace(GCKeyCodeRightControl);
+	NonPrintableKeys.Emplace(GCKeyCodeLeftControl);
+	NonPrintableKeys.Emplace(GCKeyCodeDeleteOrBackspace);
+	NonPrintableKeys.Emplace(GCKeyCodeReturnOrEnter);
+	NonPrintableKeys.Emplace(GCKeyCodeEscape);
+	NonPrintableKeys.Emplace(GCKeyCodeCapsLock);
 }
 
 void FAppleKeyboardController::SendControllerEvents()
@@ -151,18 +164,55 @@ void FAppleKeyboardController::HandleKeyChangeEvent(GCKeyCode keyCode, BOOL pres
 	{
 		return;
 	}
+	
+	// Currently character related Key Down and Key Up events are processed from UI Kit
+	// because GC Keyboard API does not provided the required keyboard layout and character translation information
+	if (IsPrintableKey(keyCode))
+	{
+		return;
+	}
 
 	FDeferredEvent NewKeyEvent;
 	NewKeyEvent.KeyEventType = pressed ? EKeyEvent::KeyDown : EKeyEvent::KeyUp;
 	NewKeyEvent.GCKeyCode = keyCode;
-	NewKeyEvent.TranslatedKeyCode = KeyCodeToPlatformCodeTranslator.IsBound() ? KeyCodeToPlatformCodeTranslator.Execute(NewKeyEvent.GCKeyCode) : NewKeyEvent.GCKeyCode;
-	NewKeyEvent.CharCode = KeyCodeToCharTranslator.IsBound() ? KeyCodeToCharTranslator.Execute(NewKeyEvent.GCKeyCode) : '?';
+	NewKeyEvent.TranslatedKeyCode = KeyCodeToPlatformCodeTranslator.IsBound() ? KeyCodeToPlatformCodeTranslator.Execute(*NewKeyEvent.GCKeyCode) : *NewKeyEvent.GCKeyCode;
+	NewKeyEvent.CharCode = KeyCodeToCharTranslator.IsBound() ? KeyCodeToCharTranslator.Execute(*NewKeyEvent.GCKeyCode) : 0;
 	
-	UE_LOG(LogAppleController, Verbose, TEXT("[GC AppleKeyboardController %hs] Key State Changed | GC KeyCode [%u] | State [%s]"), __func__, NewKeyEvent.GCKeyCode, pressed ? TEXT("True") : TEXT("False"));
+	UE_LOG(LogAppleController, Verbose, TEXT("[GC AppleKeyboardController %hs] Key State Changed | GC KeyCode [%u] | State [%s]"), __func__, *NewKeyEvent.GCKeyCode, pressed ? TEXT("True") : TEXT("False"));
 	
 	DeferredKeyEvents.Enqueue(NewKeyEvent);
 }
 
+void FAppleKeyboardController::QueueExternalKeyChangeEvent(const FDeferredEvent& InExternalKeyEvent)
+{
+	if (bIsKeyboardInhibited)
+	{
+		return;
+	}
+
+	UE_LOG(LogAppleController, Verbose, TEXT("[GC AppleKeyboardController %hs] Key State Changed | KeyCode [%u] | State [%s]"), __func__, InExternalKeyEvent.TranslatedKeyCode, InExternalKeyEvent.KeyEventType ==  EKeyEvent::KeyDown ? TEXT("True") : TEXT("False"));
+	
+	DeferredKeyEvents.Enqueue(InExternalKeyEvent);
+}
+
+void FAppleKeyboardController::QueueOneOffExternalKeyEvent(const FDeferredEvent& InExternalKeyEvent)
+{
+	if (bIsKeyboardInhibited)
+	{
+		return;
+	}
+	
+	FDeferredEvent EventCopy = InExternalKeyEvent;
+	
+	EventCopy.KeyEventType = EKeyEvent::KeyDown;
+	QueueExternalKeyChangeEvent(EventCopy);
+	
+	EventCopy.KeyEventType = EKeyEvent::KeyUp;
+	QueueExternalKeyChangeEvent(EventCopy);
+	
+	UE_LOG(LogAppleController, Verbose, TEXT("[GC AppleKeyboardController %hs] Key State Changed | KeyCode [%u] | State [%s]"), __func__, InExternalKeyEvent.TranslatedKeyCode, InExternalKeyEvent.KeyEventType ==  EKeyEvent::KeyDown ? TEXT("True") : TEXT("False"));	
+}
+
 FModifierKeysState FAppleKeyboardController::GetModifierKeys() const
 {
 	return CurrentModifierKeys.ToUEModifierKeyStates();
@@ -191,22 +241,30 @@ void FAppleKeyboardController::ProcessModifierKeyEvent(const FDeferredEvent& InD
 	UpdateModifierKeyState(CurrentModifierKeys.bIsRightShiftDown, GCKeyCodeRightShift, InDeferredEvent);
 }
 
-FAppleKeyboardController::EKeyEvent FAppleKeyboardController::GetUpdatedKeyStateForKey(GCKeyCode KeyCode)
+FAppleKeyboardController::EKeyEvent FAppleKeyboardController::GetUpdatedKeyStateForKey(const FDeferredEvent& InDeferredEvent)
 {
 	GCKeyboardInput* KeyboardInput = GCKeyboard.coalescedKeyboard ? GCKeyboard.coalescedKeyboard.keyboardInput : nullptr;
-	if (!KeyboardInput)
+
+	// We might be tracking external key events, for these we don't have an easy way to track the current state.
+	// This method is used mainly for refreshing the local state after something outside the normal update loop happened, so assume key up.
+	
+	// @note : The only provider of external events at the moment is UI Kit, and that has a Press Cancelled callback we are using to clean up the state
+	// when the app goes into the background or any other unexpected acction in the middle of a press happened, so Key Up is safe to assume. But if this API is started to being use
+	// in macOS it might need a specific way to handle it.
+
+	if (!InDeferredEvent.GCKeyCode.IsSet() || !KeyboardInput || bIsKeyboardInhibited)
 	{
 		return EKeyEvent::KeyUp;
 	}
 
-	GCControllerButtonInput* KeyInput = [KeyboardInput buttonForKeyCode: KeyCode];
+	GCControllerButtonInput* KeyInput = [KeyboardInput buttonForKeyCode: *InDeferredEvent.GCKeyCode];
 
 	return KeyInput && KeyInput.pressed ? EKeyEvent::KeyDown : EKeyEvent::KeyUp;
 }
 
 void FAppleKeyboardController::UpdateModifierKeyState(bool& OutModifierKeyFlag, GCKeyCode ModifierKeyCode, const FDeferredEvent& InKeyEvent)
 {
-	if (InKeyEvent.GCKeyCode == ModifierKeyCode)
+	if (InKeyEvent.GCKeyCode.IsSet() && *InKeyEvent.GCKeyCode == ModifierKeyCode)
 	{
 		OutModifierKeyFlag = InKeyEvent.KeyEventType == EKeyEvent::KeyDown;
 	}
@@ -216,9 +274,9 @@ void FAppleKeyboardController::DispatchEvent(const FDeferredEvent& InDeferredEve
 {
 	UE_LOG(LogAppleController,
 		   Verbose,
-		   TEXT("[GC AppleKeyboardController %hs] GC KeyCode [%u] | Type [%s] | Is Repeat [%s]"),
+		   TEXT("[GC AppleKeyboardController %hs] KeyCode [%u] | Type [%s] | Is Repeat [%s]"),
 		   __func__,
-		   InDeferredEvent.GCKeyCode,
+		   InDeferredEvent.TranslatedKeyCode,
 		   InDeferredEvent.KeyEventType == EKeyEvent::KeyUp ? TEXT("UP") : TEXT("DOWN"),
 		   InDeferredEvent.bIsRepeat ? TEXT("True") : TEXT("False"));
 	
@@ -237,7 +295,7 @@ void FAppleKeyboardController::DispatchEvent(const FDeferredEvent& InDeferredEve
 			
 			// Mimic how CMD is handled in MacApplication
 			const bool IsCMDPressed = CurrentModifierKeys.bIsLeftCommandDown || CurrentModifierKeys.bIsRightCommandDown;
-			if (!IsCMDPressed && IsPrintableKey(InDeferredEvent.GCKeyCode))
+			if (!IsCMDPressed && IsPrintableChar((TCHAR)InDeferredEvent.CharCode))
 			{
 				MessageHandler->OnKeyChar((TCHAR) InDeferredEvent.CharCode, InDeferredEvent.bIsRepeat);
 			}
@@ -256,14 +314,14 @@ void FAppleKeyboardController::UpdateEventTracking(const FDeferredEvent& InDefer
 	{
 		case EKeyEvent::KeyUp:
 		{
-			CurrentlyPressedKeyStates.Remove(InDeferredEvent.GCKeyCode);
+			CurrentlyPressedKeyStates.Remove(InDeferredEvent.TranslatedKeyCode);
 			break;
 		}
 		case EKeyEvent::KeyDown:
 		{
 			FDeferredEvent PressedKeyEventCopy = InDeferredEvent;
 			PressedKeyEventCopy.LastRepeatTime = CurrentTime;
-			CurrentlyPressedKeyStates.Add(InDeferredEvent.GCKeyCode, PressedKeyEventCopy);
+			CurrentlyPressedKeyStates.Add(InDeferredEvent.TranslatedKeyCode, PressedKeyEventCopy);
 			break;
 		}
 	case EKeyEvent::Invalid:
@@ -293,8 +351,8 @@ void FAppleKeyboardController::UpdateAndDispatchRepeatEvents()
 		double ElapsedTime = CurrentTime - PressedKeyEventWithId.Value.LastRepeatTime;
 	
 		// TODO : These are usually system settings. Find a way to read them that works in macOS and iOS
-		constexpr double InitialRepeatDelay = 0.5;
-		constexpr double RepeatInterval = 0.1;
+		constexpr double InitialRepeatDelay = 0.7;
+		constexpr double RepeatInterval = 0.4;
 	
 		const double DesiredRepeatInterval = FMath::IsNearlyZero(PressedKeyEventWithId.Value.LastRepeatTime) ? InitialRepeatDelay : RepeatInterval;
 		if (ElapsedTime > DesiredRepeatInterval)
@@ -306,7 +364,7 @@ void FAppleKeyboardController::UpdateAndDispatchRepeatEvents()
 		
 		if (bForceStateUpdate)
 		{
-			EKeyEvent CurrentKeyState = GetUpdatedKeyStateForKey(PressedKeyEventWithId.Value.GCKeyCode);
+			EKeyEvent CurrentKeyState = GetUpdatedKeyStateForKey(PressedKeyEventWithId.Value);
 			
 			if (CurrentKeyState == EKeyEvent::KeyUp)
 			{
@@ -320,7 +378,7 @@ void FAppleKeyboardController::UpdateAndDispatchRepeatEvents()
 	for (const FDeferredEvent& UnpressedKeyEvent : KeysUnpressedOutsideNormalLoop)
 	{
 		DispatchEvent(UnpressedKeyEvent);
-		CurrentlyPressedKeyStates.Remove(UnpressedKeyEvent.GCKeyCode);
+		CurrentlyPressedKeyStates.Remove(UnpressedKeyEvent.TranslatedKeyCode);
 	}
 }
 
diff --git a/Engine/Source/Runtime/ApplicationCore/Private/Apple/AppleKeyboardController.h b/Engine/Source/Runtime/ApplicationCore/Private/Apple/AppleKeyboardController.h
--- a/Engine/Source/Runtime/ApplicationCore/Private/Apple/AppleKeyboardController.h
+++ b/Engine/Source/Runtime/ApplicationCore/Private/Apple/AppleKeyboardController.h
@@ -4,6 +4,7 @@
 
 #include "Containers/SpscQueue.h"
 #include "GenericPlatform/GenericApplication.h"
+#include "Misc/Optional.h"
 
 #include "Templates/SharedPointer.h"
 
@@ -26,7 +27,7 @@ public:
 	 */
 	void SetMessageHandler(const TSharedRef<FGenericApplicationMessageHandler>& InMessageHandler);
 	
-	/** Process and dispatches all queued key events to UE's input system*/
+	/** Process and dispatches all queued key events to UE's input system */
 	void SendControllerEvents();
 	
 	/** Returns the current state of all modifier keys */
@@ -52,12 +53,12 @@ public:
 	void ForceUpdatePressedKeysState();
 
 	/**
-	 * If set to true, keyboard input from GCKeyboard will be ignored 
+	 * If set to true, keyboard input will be ignored 
 	 * @param bNewIsInhibited new inhibited state
 	 */
 	void SetInhibited(bool bNewIsInhibited);
 
-	/** Returns true if we are ignoring GC Keyboard inputs */
+	/** Returns true if we are ignoring keyboard inputs */
 	bool IsInhibited() const
 	{
 		return bIsKeyboardInhibited;
@@ -75,16 +76,8 @@ public:
 	{
 		return AvailabilityChangedDelegate;
 	}
-
-protected:
-	
-	FKeyboardAvailabilityChanged AvailabilityChangedDelegate;
-	
-	FKeyCodeTranslator KeyCodeToPlatformCodeTranslator;
-	FKeyCodeTranslator KeyCodeToCharTranslator;
-
-	TSharedRef<FGenericApplicationMessageHandler> MessageHandler;
 	
+	/** Type of key events */
 	enum class EKeyEvent : uint8
 	{
 		Invalid,
@@ -92,16 +85,46 @@ protected:
 		KeyDown
 	};
 
+	/**
+	 * Structure containing information about a key event that will be dispatched in the game thread
+	 */
 	struct FDeferredEvent
 	{
 		EKeyEvent KeyEventType = EKeyEvent::Invalid;
-		uint32 GCKeyCode = 0;
+		TOptional<uint32> GCKeyCode;
 		uint32 TranslatedKeyCode = 0;
 		uint32 CharCode = 0;
 		bool bIsRepeat = false;
 		double LastRepeatTime = 0.0;
 	};
 	
+	/**
+	 * Queues a a key event built externally  to be dispatched in the game thread
+	 * @param Event to queue
+	 */
+	void QueueExternalKeyChangeEvent(const FDeferredEvent& InExternalKeyEvent);
+	
+	/**
+	 * Queues a a key event built externally to be dispatched in the game thread. The event will be transformed in its respective key up and key down events, to act as a one off key event
+	 * @param Event to queue
+	 */
+	void QueueOneOffExternalKeyEvent(const FDeferredEvent& InExternalKeyEvent);
+	
+	/**
+	 * Evaluates a character to see if it is a printable character (ASCII or Unicode)
+	 * @return true if the character is printable
+	 */
+	bool IsPrintableChar(TCHAR InChar) const;
+
+protected:
+	
+	FKeyboardAvailabilityChanged AvailabilityChangedDelegate;
+	
+	FKeyCodeTranslator KeyCodeToPlatformCodeTranslator;
+	FKeyCodeTranslator KeyCodeToCharTranslator;
+
+	TSharedRef<FGenericApplicationMessageHandler> MessageHandler;
+	
 	struct FInternalModifierKeyState
 	{
 		bool bIsLeftShiftDown = false;
@@ -124,7 +147,7 @@ protected:
 	void HandleKeyboardDisconnected();
 	void HandleKeyChangeEvent(GCKeyCode keyCode, BOOL pressed);
 	
-	EKeyEvent GetUpdatedKeyStateForKey(GCKeyCode KeyCode);
+	EKeyEvent GetUpdatedKeyStateForKey(const FDeferredEvent& InDeferredEvent);
 	
 	void UpdateModifierKeyState(bool& OutModifierKeyFlag, GCKeyCode ModifierKeyCode, const FDeferredEvent& InKeyEvent);
 	
diff --git a/Engine/Source/Runtime/ApplicationCore/Private/IOS/IOSApplication.cpp b/Engine/Source/Runtime/ApplicationCore/Private/IOS/IOSApplication.cpp
--- a/Engine/Source/Runtime/ApplicationCore/Private/IOS/IOSApplication.cpp
+++ b/Engine/Source/Runtime/ApplicationCore/Private/IOS/IOSApplication.cpp
@@ -66,21 +66,24 @@ FIOSApplication::~FIOSApplication()
 
 void FIOSApplication::HandlePhysicalKeyboardConnectionChanged()
 {
-	FIOSView *View = [[IOSAppDelegate GetDelegate] IOSView];
+	IOSViewController* ViewController = [[IOSAppDelegate GetDelegate] IOSController];
 	
-	if (!ensure(View))
+	if (!ensure(ViewController))
 	{
 		return;
 	}
 
-	// TODO: When the dedicated view for external keyboard is implemented, we need to make it first responder here
-
 	if (InputInterface->IsAnyPhysicalKeyboardConnected())
 	{
+		ViewController.controllerUserInteractionEnabled = false;
+		[ViewController becomeFirstResponder];
 		UE_LOG(LogIOS, Log, TEXT("[%hs] External keyboard connected."), __func__);
 	}
 	else
 	{
+		ViewController.controllerUserInteractionEnabled = true;
+		[ViewController resignFirstResponder];
+
 		UE_LOG(LogIOS, Log, TEXT("[%hs] All external keyboards disconnected"), __func__);
 	}
 }
diff --git a/Engine/Source/Runtime/ApplicationCore/Private/IOS/IOSInputInterface.cpp b/Engine/Source/Runtime/ApplicationCore/Private/IOS/IOSInputInterface.cpp
--- a/Engine/Source/Runtime/ApplicationCore/Private/IOS/IOSInputInterface.cpp
+++ b/Engine/Source/Runtime/ApplicationCore/Private/IOS/IOSInputInterface.cpp
@@ -142,7 +142,6 @@ constexpr EIOSEventType operator+(EIOSEventType type, int Index) { return (EIOSE
 // protects the input stack used on 2 threads
 static FCriticalSection CriticalSection;
 static TArray<TouchInput> TouchInputStack;
-static TArray<int32> KeyInputStack;
 
 FName FIOSInputInterface::InputClassName_DefaultMobileTouch;
 FString FIOSInputInterface::HardwareDeviceIdentifier_DefaultMobileTouch;
@@ -164,7 +163,8 @@ FIOSInputInterface::FIOSInputInterface( const TSharedRef< FGenericApplicationMes
 	
 	KeyboardController->SetGCKeyCodeToCharTranslator(FAppleKeyboardController::FKeyCodeTranslator::CreateLambda([](uint32 KeyCode)
 	{
-		return TranslateGCKeyCodeToASCII(KeyCode);
+		// Currently character events and injected to the controller using UIKit Events as these have the layout and character translation we need.
+		return 0;
 	}));
 
 	KeyboardController->SetGCKeyCodeToPlatformCodeTranslator(FAppleKeyboardController::FKeyCodeTranslator::CreateLambda([](uint32 KeyCode)
@@ -295,7 +295,7 @@ void ModifyVectorByOrientation(FVector& Vec, bool bIsRotation)
 }
 #endif
 
-void FIOSInputInterface::ProcessTouchesAndKeys(uint32 ControllerId, const TArray<TouchInput>& InTouchInputStack, const TArray<int32>& InKeyInputStack)
+void FIOSInputInterface::ProcessTouchesAndKeys(uint32 ControllerId, const TArray<TouchInput>& InTouchInputStack)
 {
     for(int i = 0; i < InTouchInputStack.Num(); ++i)
     {
@@ -328,16 +328,6 @@ void FIOSInputInterface::ProcessTouchesAndKeys(uint32 ControllerId, const TArray
             MessageHandler->OnTouchFirstMove(Touch.Position, Touch.Force, Touch.Handle, ControllerId);
         }
     }
-    
-    // these come in pairs
-    for(int32 KeyIndex = 0; KeyIndex < InKeyInputStack.Num(); KeyIndex+=2)
-    {
-        int32 KeyCode = InKeyInputStack[KeyIndex];
-        int32 CharCode = InKeyInputStack[KeyIndex + 1];
-        MessageHandler->OnKeyDown(KeyCode, CharCode, false);
-        MessageHandler->OnKeyChar((TCHAR)CharCode,  false);
-        MessageHandler->OnKeyUp  (KeyCode, CharCode, false);
-    }
 }
 
 void FIOSInputInterface::SendControllerEvents()
@@ -346,16 +336,14 @@ void FIOSInputInterface::SendControllerEvents()
 	MouseController->SendControllerEvents();
 
     TArray<TouchInput> LocalTouchInputStack;
-    TArray<int32> LocalKeyInputStack;
     {
         FScopeLock Lock(&CriticalSection);
         Exchange(LocalTouchInputStack, TouchInputStack);
-        Exchange(LocalKeyInputStack, KeyInputStack);
     }
     
 #if !PLATFORM_TVOS
     // on ios, touches always go go player 0
-    ProcessTouchesAndKeys(0, LocalTouchInputStack, LocalKeyInputStack);
+    ProcessTouchesAndKeys(0, LocalTouchInputStack);
 #endif
 
     
@@ -466,15 +454,6 @@ void FIOSInputInterface::QueueTouchInput(const TArray<TouchInput>& InTouchEvents
     TouchInputStack.Append(InTouchEvents);
 }
 
-void FIOSInputInterface::QueueKeyInput(int32 Key, int32 Char)
-{
-    FScopeLock Lock(&CriticalSection);
-    
-    // put the key and char into the array
-    KeyInputStack.Add(Key);
-    KeyInputStack.Add(Char);
-}
-
 void FIOSInputInterface::EnableMotionData(bool bEnable)
 {
     bPauseMotion = !bEnable;
diff --git a/Engine/Source/Runtime/ApplicationCore/Private/IOS/IOSPlatformApplicationMisc.cpp b/Engine/Source/Runtime/ApplicationCore/Private/IOS/IOSPlatformApplicationMisc.cpp
--- a/Engine/Source/Runtime/ApplicationCore/Private/IOS/IOSPlatformApplicationMisc.cpp
+++ b/Engine/Source/Runtime/ApplicationCore/Private/IOS/IOSPlatformApplicationMisc.cpp
@@ -360,9 +360,15 @@ EScreenPhysicalAccuracy FIOSPlatformApplicationMisc::ComputePhysicalScreenDimens
 
 bool FIOSPlatformApplicationMisc::RequiresVirtualKeyboard()
 {
-#if !PLATFORM_TVOS
-    return PLATFORM_HAS_TOUCH_MAIN_SCREEN;
-#else
-    return true;
-#endif
+	{
+		FIOSApplication* IOSApplication = [IOSAppDelegate GetDelegate].IOSApplication;
+		FIOSInputInterface* InputInterface = IOSApplication ? IOSApplication->GetIOSInputInterface() : nullptr;
+		if (ensure(InputInterface))
+		{
+			return !InputInterface->IsAnyPhysicalKeyboardConnected();
+		}
+	}
+	
+	// Falback to showing the virtual keyboard
+	return true;
 }
diff --git a/Engine/Source/Runtime/ApplicationCore/Private/IOS/IOSView.cpp b/Engine/Source/Runtime/ApplicationCore/Private/IOS/IOSView.cpp
--- a/Engine/Source/Runtime/ApplicationCore/Private/IOS/IOSView.cpp
+++ b/Engine/Source/Runtime/ApplicationCore/Private/IOS/IOSView.cpp
@@ -664,8 +664,9 @@ self.accessibilityElements = @[Window.accessibilityContainer];
 		[CachedMarkedText release];
 		CachedMarkedText = nil;
 	}
-
-	if ([self IsAnyPhysicalKeyboardConnected])
+	
+	TSharedPtr<FAppleKeyboardController> KeyboardController = ([self GetKeyboardController]).Pin();
+	if (!KeyboardController)
 	{
 		return;
 	}
@@ -675,12 +676,14 @@ self.accessibilityElements = @[Window.accessibilityContainer];
 	{
 		int32 Char = [theText characterAtIndex:CharIndex];
 
-		// FPlatformMisc::LowLevelOutputDebugStringf(TEXT("sending key '%c' to game\n"), Char);
+		FAppleKeyboardController::FDeferredEvent KeyEvent;
+		
+		KeyEvent.CharCode = Char;
 
 		if (Char == '\n')
 		{
 			// send the enter keypress
-			FIOSInputInterface::QueueKeyInput(KEYCODE_ENTER, Char);
+			KeyEvent.TranslatedKeyCode = KEYCODE_ENTER;
 			
 			// hide the keyboard
 			if (!bIsUsingIntegratedKeyboard)
@@ -690,8 +693,10 @@ self.accessibilityElements = @[Window.accessibilityContainer];
 		}
 		else
 		{
-			FIOSInputInterface::QueueKeyInput(Char, Char);
+			KeyEvent.TranslatedKeyCode = Char;
 		}
+
+		KeyboardController->QueueOneOffExternalKeyEvent(KeyEvent);
 	}
 }
 
@@ -703,12 +708,17 @@ self.accessibilityElements = @[Window.accessibilityContainer];
 	}
 
 
-	if ([self IsAnyPhysicalKeyboardConnected])
+	TSharedPtr<FAppleKeyboardController> KeyboardController = ([self GetKeyboardController]).Pin();
+	if (!KeyboardController)
 	{
 		return;
 	}
-
-	FIOSInputInterface::QueueKeyInput(KEYCODE_BACKSPACE, '\b');
+	
+	FAppleKeyboardController::FDeferredEvent KeyEvent;
+	KeyEvent.CharCode = '\b';
+	KeyEvent.TranslatedKeyCode = KEYCODE_BACKSPACE;
+	
+	KeyboardController->QueueOneOffExternalKeyEvent(KeyEvent);
 }
 
 -(void)ActivateKeyboard:(bool)bInSendEscapeOnClose
@@ -780,17 +790,21 @@ self.accessibilityElements = @[Window.accessibilityContainer];
 	return FParse::Param(FCommandLine::Get(), TEXT("NewKeyboard")) || bUseIntegratedKeyboard;
 }
 
--(bool)IsAnyPhysicalKeyboardConnected;
+-(TWeakPtr<FAppleKeyboardController>)GetKeyboardController
 {
-	FIOSInputInterface* InputInterface = [IOSAppDelegate GetDelegate].IOSApplication->GetIOSInputInterface();
-
-	return InputInterface ? InputInterface->IsAnyPhysicalKeyboardConnected() : false;
+	FIOSApplication* IOSApplication = [IOSAppDelegate GetDelegate].IOSApplication;
+	FIOSInputInterface* InputInterface = IOSApplication ? IOSApplication->GetIOSInputInterface() : nullptr;
+	
+	return InputInterface ? InputInterface->GetKeyboardController() : nullptr;
 }
 
 -(BOOL)becomeFirstResponder
 {
+	TSharedPtr<FAppleKeyboardController> KeyboardController = ([self GetKeyboardController]).Pin();
+	bool bIsAnyPhysicalKeyboardConnected = KeyboardController && KeyboardController->IsAnyKeyboardConnected();
+
 	volatile int32 ShowCount = KeyboardShowCount;
-	if (ShowCount >= 1)
+	if (ShowCount >= 1 && !bIsAnyPhysicalKeyboardConnected)
 	{
 		return [super becomeFirstResponder];
 	}
@@ -804,8 +818,16 @@ self.accessibilityElements = @[Window.accessibilityContainer];
 {
 	if (bSendEscapeOnClose)
 	{
-		// tell the console to close itself
-		FIOSInputInterface::QueueKeyInput(KEYCODE_ESCAPE, 0);
+		TSharedPtr<FAppleKeyboardController> KeyboardController = ([self GetKeyboardController]).Pin();
+		if (KeyboardController)
+		{
+			FAppleKeyboardController::FDeferredEvent KeyEvent;
+			KeyEvent.CharCode = 0;
+			KeyEvent.TranslatedKeyCode = KEYCODE_ESCAPE;
+
+			// tell the console to close itself
+			KeyboardController->QueueOneOffExternalKeyEvent(KeyEvent);
+		}
 	}
 	
 	return [super resignFirstResponder];
@@ -1291,4 +1313,92 @@ self.accessibilityElements = @[Window.accessibilityContainer];
  } 
  #endif
 
+- (BOOL)canBecomeFirstResponder
+{
+	return true;
+}
+
+-(TWeakPtr<FAppleKeyboardController>)GetKeyboardController
+{
+	FIOSApplication* IOSApplication = [IOSAppDelegate GetDelegate].IOSApplication;
+	FIOSInputInterface* InputInterface = IOSApplication ? IOSApplication->GetIOSInputInterface() : nullptr;
+	return InputInterface ? InputInterface->GetKeyboardController() : nullptr;
+}
+
+- (bool)IsAnyPhysicalKeyboardConnected
+{
+	TSharedPtr<FAppleKeyboardController> KeyboardController = ([self GetKeyboardController]).Pin();
+	return KeyboardController ? KeyboardController->IsAnyKeyboardConnected() : false;
+}
+
+- (bool)ProcessKeyPresses:(NSSet<UIPress *> *)keypresses eventType:(FAppleKeyboardController::EKeyEvent)eventType
+{
+	TSharedPtr<FAppleKeyboardController> KeyboardController = ([self GetKeyboardController]).Pin();
+	if (KeyboardController && KeyboardController->IsAnyKeyboardConnected())
+	{
+		for (UIPress* Press in keypresses.allObjects)
+		{
+			if (UIKey* PressedKey = Press ? Press.key : nil)
+			{
+				// Sequences like backspace or new lines are handled by the controller internally using GC Key events
+				// So it is safe (and desired) ignore them here
+
+				unichar Char = [PressedKey.characters length] == 1 ? [PressedKey.characters characterAtIndex:0] : 0;
+				if (KeyboardController->IsPrintableChar((TCHAR) Char))
+				{
+					FAppleKeyboardController::FDeferredEvent KeyEvent;
+					KeyEvent.KeyEventType = eventType;
+					KeyEvent.TranslatedKeyCode = Char;
+					KeyEvent.CharCode = Char;
+					
+					KeyboardController->QueueExternalKeyChangeEvent(KeyEvent);
+				}
+			}
+		}
+		
+		return true;
+	}
+	
+	return false;
+}
+
+- (void)pressesBegan:(NSSet<UIPress *> *)presses withEvent:(nullable UIPressesEvent *)event
+{
+	if ([self ProcessKeyPresses:presses eventType:FAppleKeyboardController::EKeyEvent::KeyDown])
+	{
+		return;
+	}
+	
+	[super pressesBegan:presses withEvent:event];
+}
+
+- (void)pressesChanged:(NSSet<UIPress *> *)presses withEvent:(nullable UIPressesEvent *)event
+{
+	if ([self IsAnyPhysicalKeyboardConnected])
+	{
+		return;
+	}
+	[super pressesChanged:presses withEvent:event];
+}
+
+- (void)pressesEnded:(NSSet<UIPress *> *)presses withEvent:(nullable UIPressesEvent *)event
+{
+	if ([self ProcessKeyPresses:presses eventType:FAppleKeyboardController::EKeyEvent::KeyUp])
+	{
+		return;
+	}
+
+	[super pressesEnded:presses withEvent:event];
+}
+
+- (void)pressesCancelled:(NSSet<UIPress *> *)presses withEvent:(nullable UIPressesEvent *)event
+{
+	if ([self ProcessKeyPresses:presses eventType:FAppleKeyboardController::EKeyEvent::KeyUp])
+	{
+		return;
+	}
+
+	[super pressesCancelled:presses withEvent:event];
+}
+
 @end
diff --git a/Engine/Source/Runtime/ApplicationCore/Public/IOS/IOSInputInterface.h b/Engine/Source/Runtime/ApplicationCore/Public/IOS/IOSInputInterface.h
--- a/Engine/Source/Runtime/ApplicationCore/Public/IOS/IOSInputInterface.h
+++ b/Engine/Source/Runtime/ApplicationCore/Public/IOS/IOSInputInterface.h
@@ -124,7 +124,6 @@ public:
 	virtual void SetForceFeedbackChannelValues(int32 ControllerId, const FForceFeedbackValues &values) override;
 
 	static APPLICATIONCORE_API void QueueTouchInput(const TArray<TouchInput>& InTouchEvents);
-	static void QueueKeyInput(int32 Key, int32 Char);
 
 	void SetGamepadsAllowed(bool bAllowed) { bAllowControllers = bAllowed; }
 	void SetGamepadsBlockDeviceFeedback(bool bBlock) { bControllersBlockDeviceFeedback = bBlock; }
@@ -178,7 +177,7 @@ private:
 	void CalibrateMotion(uint32 PlayerIndex);
 
 private:
-	void ProcessTouchesAndKeys(uint32 ControllerId, const TArray<TouchInput>& InTouchInputStack, const TArray<int32>& InKeyInputStack);
+	void ProcessTouchesAndKeys(uint32 ControllerId, const TArray<TouchInput>& InTouchInputStack);
 	
 	// can the remote be rotated to landscape
 	bool bAllowRemoteRotation;
diff --git a/Engine/Source/Runtime/ApplicationCore/Public/IOS/IOSView.h b/Engine/Source/Runtime/ApplicationCore/Public/IOS/IOSView.h
--- a/Engine/Source/Runtime/ApplicationCore/Public/IOS/IOSView.h
+++ b/Engine/Source/Runtime/ApplicationCore/Public/IOS/IOSView.h
@@ -3,9 +3,11 @@
 #pragma once
 
 #ifndef SWIFT_IMPORT
+#include "AppleKeyboardController.h"
 #include "CoreTypes.h"
 #include "Delegates/Delegate.h"
 #include "IOS/IOSInputInterface.h"
+#include "Templates/SharedPointer.h"
 #endif
 
 #import <UIKit/UIKit.h>
@@ -116,8 +118,7 @@ APPLICATIONCORE_API
 -(void)ActivateKeyboard:(bool)bInSendEscapeOnClose keyboardConfig:(FKeyboardConfig)KeyboardConfig;
 -(void)DeactivateKeyboard;
 -(bool)ShouldUseIntegratedKeyboard;
-
--(bool)IsAnyPhysicalKeyboardConnected;
+-(TWeakPtr<FAppleKeyboardController>) GetKeyboardController;
 
 // callable from outside to fake locations
 -(void)HandleTouchAtLoc:(CGPoint)Loc PrevLoc:(CGPoint)PrevLoc TouchIndex:(int)TouchIndex Force:(float)Force Type:(TouchType)Type TouchesArray:(TArray<TouchInput>&)TouchesArray;
@@ -134,13 +135,10 @@ APPLICATIONCORE_API
 /**
  * A view controller subclass that handles loading our IOS view as well as autorotation
  */
-#if PLATFORM_TVOS
 #import <GameController/GameController.h>
-// if TVOS doesn't use the GCEventViewController, it will background the app when the user presses Menu/Pause
+// We use a GCEvent view controller to dynamically control how Game Controller API events flow in the responder chain, based on the presence of a physical keyboard or not
 @interface IOSViewController : GCEventViewController
-#else
-@interface IOSViewController : UIViewController
-#endif
+
 {
 }
 
@@ -151,4 +149,15 @@ APPLICATIONCORE_API
 - (UIInterfaceOrientationMask)supportedInterfaceOrientations;
 #endif
 
+- (TWeakPtr<FAppleKeyboardController>)GetKeyboardController;
+- (bool)IsAnyPhysicalKeyboardConnected;
+
+- (bool)ProcessKeyPresses:(NSSet<UIPress *> *)keypresses eventType:(FAppleKeyboardController::EKeyEvent)eventType;
+
+- (void)pressesBegan:(NSSet<UIPress *> *)presses withEvent:(nullable UIPressesEvent *)event;
+- (void)pressesChanged:(NSSet<UIPress *> *)presses withEvent:(nullable UIPressesEvent *)event;
+- (void)pressesEnded:(NSSet<UIPress *> *)presses withEvent:(nullable UIPressesEvent *)event;
+- (void)pressesCancelled:(NSSet<UIPress *> *)presses withEvent:(nullable UIPressesEvent *)event;
+- (bool)canBecomeFirstResponder;
+
 @end
diff --git a/Engine/Source/Runtime/Launch/Private/IOS/LaunchIOS.cpp b/Engine/Source/Runtime/Launch/Private/IOS/LaunchIOS.cpp
--- a/Engine/Source/Runtime/Launch/Private/IOS/LaunchIOS.cpp
+++ b/Engine/Source/Runtime/Launch/Private/IOS/LaunchIOS.cpp
@@ -320,6 +320,10 @@ void FAppEntry::PreInit(IOSAppDelegate* AppDelegate, UIApplication* Application)
 	// make a controller object
 	IOSViewController* IOSController = [[IOSViewController alloc] init];
 	
+	// IOSViewController is a GCViewController so we can have proper GC Controller support
+	// But default we want all inputs to be forwaded to UI Kit. This will only change when we have physical keyborad or mouse connected
+	IOSController.controllerUserInteractionEnabled = true;
+	
 #if PLATFORM_TVOS
 	// @todo tvos: This may need to be exposed to the game so that when you click Menu it will background the app
 	// this is basically the same way Android handles the Back button (maybe we should pass Menu button as back... maybe)
@@ -382,6 +386,8 @@ static void MainThreadInit()
 #else
 	AppDelegate.IOSView = [[FIOSView alloc] initWithFrame:FullResolutionRect];
 	AppDelegate.IOSView.clearsContextBeforeDrawing = NO;
+	
+	
 #if !PLATFORM_TVOS
 	AppDelegate.IOSView.multipleTouchEnabled = YES;
 #endif
