[Backout] - CL49144466
#rnx
Original CL Desc
-----------------------------------------------------------------
[iOS] Improving character input support for external keyboards

- Given GC keyboard does not provide keyboard layour or character translation information, now the GC keyboard implementation only handles non-character related key inputs. The rest of the inputs are consummed from UIKit
- Updated the touch character input to use the new Keyboard API which has a lock free queue system.
- Implemented the UI responder key presses handling methods in the view controller so we can foward character related key events to the new Apple keyboard API introduces in the previous change.
- Made the view controller conform to the GCViewcontroller type for all iOS platforms, not just tvOS. This allow to make sure interactions are routed to the GC Keyboard/Mouse API when we desired to do so
- Now we make the view controller the first responder when a physical keyboard is attached. This prevents the virtual keyboard from popping up among other things.
- At a platform implementation level, now they requires virtual keyboard method now dynamically returns true or false based on if we have a physical keyboard attached or now. Before that it was defined at compile time.

#Synced-CL 49059801

diff --git a/Engine/Source/Runtime/ApplicationCore/Private/Apple/AppleKeyboardController.cpp b/Engine/Source/Runtime/ApplicationCore/Private/Apple/AppleKeyboardController.cpp
--- a/Engine/Source/Runtime/ApplicationCore/Private/Apple/AppleKeyboardController.cpp
+++ b/Engine/Source/Runtime/ApplicationCore/Private/Apple/AppleKeyboardController.cpp
@@ -90,17 +90,10 @@ bool FAppleKeyboardController::IsPrintableKey(uint32 KeyCode) const
 	return !NonPrintableKeys.Contains(KeyCode);
 }
 
-bool FAppleKeyboardController::IsPrintableChar(TCHAR InChar) const
-{
-	const bool bIsPrintableASCII = InChar >= 32 && InChar <= 126;
-	const bool bIsPrintableUnicode = InChar > 160;
-	return bIsPrintableASCII || bIsPrintableUnicode;
-}
-
 void FAppleKeyboardController::RegisterNonPrintableKeys()
 {
 	// If more keys are added, increase this number to match
-	constexpr int32 NonPrintableKeysNum = 35;
+	constexpr int32 NonPrintableKeysNum = 29;
 	NonPrintableKeys.Reserve(NonPrintableKeysNum);
 	
 	NonPrintableKeys.Emplace(GCKeyCodeF1);
@@ -132,12 +125,6 @@ void FAppleKeyboardController::RegisterNonPrintableKeys()
 	NonPrintableKeys.Emplace(GCKeyCodeRightAlt);
 	NonPrintableKeys.Emplace(GCKeyCodeRightGUI);
 	NonPrintableKeys.Emplace(GCKeyCodeLeftGUI);
-	NonPrintableKeys.Emplace(GCKeyCodeRightControl);
-	NonPrintableKeys.Emplace(GCKeyCodeLeftControl);
-	NonPrintableKeys.Emplace(GCKeyCodeDeleteOrBackspace);
-	NonPrintableKeys.Emplace(GCKeyCodeReturnOrEnter);
-	NonPrintableKeys.Emplace(GCKeyCodeEscape);
-	NonPrintableKeys.Emplace(GCKeyCodeCapsLock);
 }
 
 void FAppleKeyboardController::SendControllerEvents()
@@ -164,55 +151,18 @@ void FAppleKeyboardController::HandleKeyChangeEvent(GCKeyCode keyCode, BOOL pres
 	{
 		return;
 	}
-	
-	// Currently character related Key Down and Key Up events are processed from UI Kit
-	// because GC Keyboard API does not provided the required keyboard layout and character translation information
-	if (IsPrintableKey(keyCode))
-	{
-		return;
-	}
 
 	FDeferredEvent NewKeyEvent;
 	NewKeyEvent.KeyEventType = pressed ? EKeyEvent::KeyDown : EKeyEvent::KeyUp;
 	NewKeyEvent.GCKeyCode = keyCode;
-	NewKeyEvent.TranslatedKeyCode = KeyCodeToPlatformCodeTranslator.IsBound() ? KeyCodeToPlatformCodeTranslator.Execute(*NewKeyEvent.GCKeyCode) : *NewKeyEvent.GCKeyCode;
-	NewKeyEvent.CharCode = KeyCodeToCharTranslator.IsBound() ? KeyCodeToCharTranslator.Execute(*NewKeyEvent.GCKeyCode) : 0;
+	NewKeyEvent.TranslatedKeyCode = KeyCodeToPlatformCodeTranslator.IsBound() ? KeyCodeToPlatformCodeTranslator.Execute(NewKeyEvent.GCKeyCode) : NewKeyEvent.GCKeyCode;
+	NewKeyEvent.CharCode = KeyCodeToCharTranslator.IsBound() ? KeyCodeToCharTranslator.Execute(NewKeyEvent.GCKeyCode) : '?';
 	
-	UE_LOG(LogAppleController, Verbose, TEXT("[GC AppleKeyboardController %hs] Key State Changed | GC KeyCode [%u] | State [%s]"), __func__, *NewKeyEvent.GCKeyCode, pressed ? TEXT("True") : TEXT("False"));
+	UE_LOG(LogAppleController, Verbose, TEXT("[GC AppleKeyboardController %hs] Key State Changed | GC KeyCode [%u] | State [%s]"), __func__, NewKeyEvent.GCKeyCode, pressed ? TEXT("True") : TEXT("False"));
 	
 	DeferredKeyEvents.Enqueue(NewKeyEvent);
 }
 
-void FAppleKeyboardController::QueueExternalKeyChangeEvent(const FDeferredEvent& InExternalKeyEvent)
-{
-	if (bIsKeyboardInhibited)
-	{
-		return;
-	}
-
-	UE_LOG(LogAppleController, Verbose, TEXT("[GC AppleKeyboardController %hs] Key State Changed | KeyCode [%u] | State [%s]"), __func__, InExternalKeyEvent.TranslatedKeyCode, InExternalKeyEvent.KeyEventType ==  EKeyEvent::KeyDown ? TEXT("True") : TEXT("False"));
-	
-	DeferredKeyEvents.Enqueue(InExternalKeyEvent);
-}
-
-void FAppleKeyboardController::QueueOneOffExternalKeyEvent(const FDeferredEvent& InExternalKeyEvent)
-{
-	if (bIsKeyboardInhibited)
-	{
-		return;
-	}
-	
-	FDeferredEvent EventCopy = InExternalKeyEvent;
-	
-	EventCopy.KeyEventType = EKeyEvent::KeyDown;
-	QueueExternalKeyChangeEvent(EventCopy);
-	
-	EventCopy.KeyEventType = EKeyEvent::KeyUp;
-	QueueExternalKeyChangeEvent(EventCopy);
-	
-	UE_LOG(LogAppleController, Verbose, TEXT("[GC AppleKeyboardController %hs] Key State Changed | KeyCode [%u] | State [%s]"), __func__, InExternalKeyEvent.TranslatedKeyCode, InExternalKeyEvent.KeyEventType ==  EKeyEvent::KeyDown ? TEXT("True") : TEXT("False"));	
-}
-
 FModifierKeysState FAppleKeyboardController::GetModifierKeys() const
 {
 	return CurrentModifierKeys.ToUEModifierKeyStates();
@@ -241,30 +191,22 @@ void FAppleKeyboardController::ProcessModifierKeyEvent(const FDeferredEvent& InD
 	UpdateModifierKeyState(CurrentModifierKeys.bIsRightShiftDown, GCKeyCodeRightShift, InDeferredEvent);
 }
 
-FAppleKeyboardController::EKeyEvent FAppleKeyboardController::GetUpdatedKeyStateForKey(const FDeferredEvent& InDeferredEvent)
+FAppleKeyboardController::EKeyEvent FAppleKeyboardController::GetUpdatedKeyStateForKey(GCKeyCode KeyCode)
 {
 	GCKeyboardInput* KeyboardInput = GCKeyboard.coalescedKeyboard ? GCKeyboard.coalescedKeyboard.keyboardInput : nullptr;
-
-	// We might be tracking external key events, for these we don't have an easy way to track the current state.
-	// This method is used mainly for refreshing the local state after something outside the normal update loop happened, so assume key up.
-	
-	// @note : The only provider of external events at the moment is UI Kit, and that has a Press Cancelled callback we are using to clean up the state
-	// when the app goes into the background or any other unexpected acction in the middle of a press happened, so Key Up is safe to assume. But if this API is started to being use
-	// in macOS it might need a specific way to handle it.
-
-	if (!InDeferredEvent.GCKeyCode.IsSet() || !KeyboardInput || bIsKeyboardInhibited)
+	if (!KeyboardInput)
 	{
 		return EKeyEvent::KeyUp;
 	}
 
-	GCControllerButtonInput* KeyInput = [KeyboardInput buttonForKeyCode: *InDeferredEvent.GCKeyCode];
+	GCControllerButtonInput* KeyInput = [KeyboardInput buttonForKeyCode: KeyCode];
 
 	return KeyInput && KeyInput.pressed ? EKeyEvent::KeyDown : EKeyEvent::KeyUp;
 }
 
 void FAppleKeyboardController::UpdateModifierKeyState(bool& OutModifierKeyFlag, GCKeyCode ModifierKeyCode, const FDeferredEvent& InKeyEvent)
 {
-	if (InKeyEvent.GCKeyCode.IsSet() && *InKeyEvent.GCKeyCode == ModifierKeyCode)
+	if (InKeyEvent.GCKeyCode == ModifierKeyCode)
 	{
 		OutModifierKeyFlag = InKeyEvent.KeyEventType == EKeyEvent::KeyDown;
 	}
@@ -274,9 +216,9 @@ void FAppleKeyboardController::DispatchEvent(const FDeferredEvent& InDeferredEve
 {
 	UE_LOG(LogAppleController,
 		   Verbose,
-		   TEXT("[GC AppleKeyboardController %hs] KeyCode [%u] | Type [%s] | Is Repeat [%s]"),
+		   TEXT("[GC AppleKeyboardController %hs] GC KeyCode [%u] | Type [%s] | Is Repeat [%s]"),
 		   __func__,
-		   InDeferredEvent.TranslatedKeyCode,
+		   InDeferredEvent.GCKeyCode,
 		   InDeferredEvent.KeyEventType == EKeyEvent::KeyUp ? TEXT("UP") : TEXT("DOWN"),
 		   InDeferredEvent.bIsRepeat ? TEXT("True") : TEXT("False"));
 	
@@ -295,7 +237,7 @@ void FAppleKeyboardController::DispatchEvent(const FDeferredEvent& InDeferredEve
 			
 			// Mimic how CMD is handled in MacApplication
 			const bool IsCMDPressed = CurrentModifierKeys.bIsLeftCommandDown || CurrentModifierKeys.bIsRightCommandDown;
-			if (!IsCMDPressed && IsPrintableChar((TCHAR)InDeferredEvent.CharCode))
+			if (!IsCMDPressed && IsPrintableKey(InDeferredEvent.GCKeyCode))
 			{
 				MessageHandler->OnKeyChar((TCHAR) InDeferredEvent.CharCode, InDeferredEvent.bIsRepeat);
 			}
@@ -314,14 +256,14 @@ void FAppleKeyboardController::UpdateEventTracking(const FDeferredEvent& InDefer
 	{
 		case EKeyEvent::KeyUp:
 		{
-			CurrentlyPressedKeyStates.Remove(InDeferredEvent.TranslatedKeyCode);
+			CurrentlyPressedKeyStates.Remove(InDeferredEvent.GCKeyCode);
 			break;
 		}
 		case EKeyEvent::KeyDown:
 		{
 			FDeferredEvent PressedKeyEventCopy = InDeferredEvent;
 			PressedKeyEventCopy.LastRepeatTime = CurrentTime;
-			CurrentlyPressedKeyStates.Add(InDeferredEvent.TranslatedKeyCode, PressedKeyEventCopy);
+			CurrentlyPressedKeyStates.Add(InDeferredEvent.GCKeyCode, PressedKeyEventCopy);
 			break;
 		}
 	case EKeyEvent::Invalid:
@@ -351,8 +293,8 @@ void FAppleKeyboardController::UpdateAndDispatchRepeatEvents()
 		double ElapsedTime = CurrentTime - PressedKeyEventWithId.Value.LastRepeatTime;
 	
 		// TODO : These are usually system settings. Find a way to read them that works in macOS and iOS
-		constexpr double InitialRepeatDelay = 0.7;
-		constexpr double RepeatInterval = 0.4;
+		constexpr double InitialRepeatDelay = 0.5;
+		constexpr double RepeatInterval = 0.1;
 	
 		const double DesiredRepeatInterval = FMath::IsNearlyZero(PressedKeyEventWithId.Value.LastRepeatTime) ? InitialRepeatDelay : RepeatInterval;
 		if (ElapsedTime > DesiredRepeatInterval)
@@ -364,7 +306,7 @@ void FAppleKeyboardController::UpdateAndDispatchRepeatEvents()
 		
 		if (bForceStateUpdate)
 		{
-			EKeyEvent CurrentKeyState = GetUpdatedKeyStateForKey(PressedKeyEventWithId.Value);
+			EKeyEvent CurrentKeyState = GetUpdatedKeyStateForKey(PressedKeyEventWithId.Value.GCKeyCode);
 			
 			if (CurrentKeyState == EKeyEvent::KeyUp)
 			{
@@ -378,7 +320,7 @@ void FAppleKeyboardController::UpdateAndDispatchRepeatEvents()
 	for (const FDeferredEvent& UnpressedKeyEvent : KeysUnpressedOutsideNormalLoop)
 	{
 		DispatchEvent(UnpressedKeyEvent);
-		CurrentlyPressedKeyStates.Remove(UnpressedKeyEvent.TranslatedKeyCode);
+		CurrentlyPressedKeyStates.Remove(UnpressedKeyEvent.GCKeyCode);
 	}
 }
 
diff --git a/Engine/Source/Runtime/ApplicationCore/Private/Apple/AppleKeyboardController.h b/Engine/Source/Runtime/ApplicationCore/Private/Apple/AppleKeyboardController.h
--- a/Engine/Source/Runtime/ApplicationCore/Private/Apple/AppleKeyboardController.h
+++ b/Engine/Source/Runtime/ApplicationCore/Private/Apple/AppleKeyboardController.h
@@ -4,7 +4,6 @@
 
 #include "Containers/SpscQueue.h"
 #include "GenericPlatform/GenericApplication.h"
-#include "Misc/Optional.h"
 
 #include "Templates/SharedPointer.h"
 
@@ -27,7 +26,7 @@ public:
 	 */
 	void SetMessageHandler(const TSharedRef<FGenericApplicationMessageHandler>& InMessageHandler);
 	
-	/** Process and dispatches all queued key events to UE's input system */
+	/** Process and dispatches all queued key events to UE's input system*/
 	void SendControllerEvents();
 	
 	/** Returns the current state of all modifier keys */
@@ -53,12 +52,12 @@ public:
 	void ForceUpdatePressedKeysState();
 
 	/**
-	 * If set to true, keyboard input will be ignored 
+	 * If set to true, keyboard input from GCKeyboard will be ignored 
 	 * @param bNewIsInhibited new inhibited state
 	 */
 	void SetInhibited(bool bNewIsInhibited);
 
-	/** Returns true if we are ignoring keyboard inputs */
+	/** Returns true if we are ignoring GC Keyboard inputs */
 	bool IsInhibited() const
 	{
 		return bIsKeyboardInhibited;
@@ -76,8 +75,16 @@ public:
 	{
 		return AvailabilityChangedDelegate;
 	}
+
+protected:
+	
+	FKeyboardAvailabilityChanged AvailabilityChangedDelegate;
+	
+	FKeyCodeTranslator KeyCodeToPlatformCodeTranslator;
+	FKeyCodeTranslator KeyCodeToCharTranslator;
+
+	TSharedRef<FGenericApplicationMessageHandler> MessageHandler;
 	
-	/** Type of key events */
 	enum class EKeyEvent : uint8
 	{
 		Invalid,
@@ -85,46 +92,16 @@ public:
 		KeyDown
 	};
 
-	/**
-	 * Structure containing information about a key event that will be dispatched in the game thread
-	 */
 	struct FDeferredEvent
 	{
 		EKeyEvent KeyEventType = EKeyEvent::Invalid;
-		TOptional<uint32> GCKeyCode;
+		uint32 GCKeyCode = 0;
 		uint32 TranslatedKeyCode = 0;
 		uint32 CharCode = 0;
 		bool bIsRepeat = false;
 		double LastRepeatTime = 0.0;
 	};
 	
-	/**
-	 * Queues a a key event built externally  to be dispatched in the game thread
-	 * @param Event to queue
-	 */
-	void QueueExternalKeyChangeEvent(const FDeferredEvent& InExternalKeyEvent);
-	
-	/**
-	 * Queues a a key event built externally to be dispatched in the game thread. The event will be transformed in its respective key up and key down events, to act as a one off key event
-	 * @param Event to queue
-	 */
-	void QueueOneOffExternalKeyEvent(const FDeferredEvent& InExternalKeyEvent);
-	
-	/**
-	 * Evaluates a character to see if it is a printable character (ASCII or Unicode)
-	 * @return true if the character is printable
-	 */
-	bool IsPrintableChar(TCHAR InChar) const;
-
-protected:
-	
-	FKeyboardAvailabilityChanged AvailabilityChangedDelegate;
-	
-	FKeyCodeTranslator KeyCodeToPlatformCodeTranslator;
-	FKeyCodeTranslator KeyCodeToCharTranslator;
-
-	TSharedRef<FGenericApplicationMessageHandler> MessageHandler;
-	
 	struct FInternalModifierKeyState
 	{
 		bool bIsLeftShiftDown = false;
@@ -147,7 +124,7 @@ protected:
 	void HandleKeyboardDisconnected();
 	void HandleKeyChangeEvent(GCKeyCode keyCode, BOOL pressed);
 	
-	EKeyEvent GetUpdatedKeyStateForKey(const FDeferredEvent& InDeferredEvent);
+	EKeyEvent GetUpdatedKeyStateForKey(GCKeyCode KeyCode);
 	
 	void UpdateModifierKeyState(bool& OutModifierKeyFlag, GCKeyCode ModifierKeyCode, const FDeferredEvent& InKeyEvent);
 	
diff --git a/Engine/Source/Runtime/ApplicationCore/Private/IOS/IOSApplication.cpp b/Engine/Source/Runtime/ApplicationCore/Private/IOS/IOSApplication.cpp
--- a/Engine/Source/Runtime/ApplicationCore/Private/IOS/IOSApplication.cpp
+++ b/Engine/Source/Runtime/ApplicationCore/Private/IOS/IOSApplication.cpp
@@ -66,24 +66,21 @@ FIOSApplication::~FIOSApplication()
 
 void FIOSApplication::HandlePhysicalKeyboardConnectionChanged()
 {
-	IOSViewController* ViewController = [[IOSAppDelegate GetDelegate] IOSController];
+	FIOSView *View = [[IOSAppDelegate GetDelegate] IOSView];
 	
-	if (!ensure(ViewController))
+	if (!ensure(View))
 	{
 		return;
 	}
 
+	// TODO: When the dedicated view for external keyboard is implemented, we need to make it first responder here
+
 	if (InputInterface->IsAnyPhysicalKeyboardConnected())
 	{
-		ViewController.controllerUserInteractionEnabled = false;
-		[ViewController becomeFirstResponder];
 		UE_LOG(LogIOS, Log, TEXT("[%hs] External keyboard connected."), __func__);
 	}
 	else
 	{
-		ViewController.controllerUserInteractionEnabled = true;
-		[ViewController resignFirstResponder];
-
 		UE_LOG(LogIOS, Log, TEXT("[%hs] All external keyboards disconnected"), __func__);
 	}
 }
diff --git a/Engine/Source/Runtime/ApplicationCore/Private/IOS/IOSInputInterface.cpp b/Engine/Source/Runtime/ApplicationCore/Private/IOS/IOSInputInterface.cpp
--- a/Engine/Source/Runtime/ApplicationCore/Private/IOS/IOSInputInterface.cpp
+++ b/Engine/Source/Runtime/ApplicationCore/Private/IOS/IOSInputInterface.cpp
@@ -142,6 +142,7 @@ constexpr EIOSEventType operator+(EIOSEventType type, int Index) { return (EIOSE
 // protects the input stack used on 2 threads
 static FCriticalSection CriticalSection;
 static TArray<TouchInput> TouchInputStack;
+static TArray<int32> KeyInputStack;
 
 FName FIOSInputInterface::InputClassName_DefaultMobileTouch;
 FString FIOSInputInterface::HardwareDeviceIdentifier_DefaultMobileTouch;
@@ -163,8 +164,7 @@ FIOSInputInterface::FIOSInputInterface( const TSharedRef< FGenericApplicationMes
 	
 	KeyboardController->SetGCKeyCodeToCharTranslator(FAppleKeyboardController::FKeyCodeTranslator::CreateLambda([](uint32 KeyCode)
 	{
-		// Currently character events and injected to the controller using UIKit Events as these have the layout and character translation we need.
-		return 0;
+		return TranslateGCKeyCodeToASCII(KeyCode);
 	}));
 
 	KeyboardController->SetGCKeyCodeToPlatformCodeTranslator(FAppleKeyboardController::FKeyCodeTranslator::CreateLambda([](uint32 KeyCode)
@@ -295,7 +295,7 @@ void ModifyVectorByOrientation(FVector& Vec, bool bIsRotation)
 }
 #endif
 
-void FIOSInputInterface::ProcessTouchesAndKeys(uint32 ControllerId, const TArray<TouchInput>& InTouchInputStack)
+void FIOSInputInterface::ProcessTouchesAndKeys(uint32 ControllerId, const TArray<TouchInput>& InTouchInputStack, const TArray<int32>& InKeyInputStack)
 {
     for(int i = 0; i < InTouchInputStack.Num(); ++i)
     {
@@ -328,6 +328,16 @@ void FIOSInputInterface::ProcessTouchesAndKeys(uint32 ControllerId, const TArray
             MessageHandler->OnTouchFirstMove(Touch.Position, Touch.Force, Touch.Handle, ControllerId);
         }
     }
+    
+    // these come in pairs
+    for(int32 KeyIndex = 0; KeyIndex < InKeyInputStack.Num(); KeyIndex+=2)
+    {
+        int32 KeyCode = InKeyInputStack[KeyIndex];
+        int32 CharCode = InKeyInputStack[KeyIndex + 1];
+        MessageHandler->OnKeyDown(KeyCode, CharCode, false);
+        MessageHandler->OnKeyChar((TCHAR)CharCode,  false);
+        MessageHandler->OnKeyUp  (KeyCode, CharCode, false);
+    }
 }
 
 void FIOSInputInterface::SendControllerEvents()
@@ -336,14 +346,16 @@ void FIOSInputInterface::SendControllerEvents()
 	MouseController->SendControllerEvents();
 
     TArray<TouchInput> LocalTouchInputStack;
+    TArray<int32> LocalKeyInputStack;
     {
         FScopeLock Lock(&CriticalSection);
         Exchange(LocalTouchInputStack, TouchInputStack);
+        Exchange(LocalKeyInputStack, KeyInputStack);
     }
     
 #if !PLATFORM_TVOS
     // on ios, touches always go go player 0
-    ProcessTouchesAndKeys(0, LocalTouchInputStack);
+    ProcessTouchesAndKeys(0, LocalTouchInputStack, LocalKeyInputStack);
 #endif
 
     
@@ -454,6 +466,15 @@ void FIOSInputInterface::QueueTouchInput(const TArray<TouchInput>& InTouchEvents
     TouchInputStack.Append(InTouchEvents);
 }
 
+void FIOSInputInterface::QueueKeyInput(int32 Key, int32 Char)
+{
+    FScopeLock Lock(&CriticalSection);
+    
+    // put the key and char into the array
+    KeyInputStack.Add(Key);
+    KeyInputStack.Add(Char);
+}
+
 void FIOSInputInterface::EnableMotionData(bool bEnable)
 {
     bPauseMotion = !bEnable;
diff --git a/Engine/Source/Runtime/ApplicationCore/Private/IOS/IOSPlatformApplicationMisc.cpp b/Engine/Source/Runtime/ApplicationCore/Private/IOS/IOSPlatformApplicationMisc.cpp
--- a/Engine/Source/Runtime/ApplicationCore/Private/IOS/IOSPlatformApplicationMisc.cpp
+++ b/Engine/Source/Runtime/ApplicationCore/Private/IOS/IOSPlatformApplicationMisc.cpp
@@ -360,15 +360,9 @@ EScreenPhysicalAccuracy FIOSPlatformApplicationMisc::ComputePhysicalScreenDimens
 
 bool FIOSPlatformApplicationMisc::RequiresVirtualKeyboard()
 {
-	{
-		FIOSApplication* IOSApplication = [IOSAppDelegate GetDelegate].IOSApplication;
-		FIOSInputInterface* InputInterface = IOSApplication ? IOSApplication->GetIOSInputInterface() : nullptr;
-		if (ensure(InputInterface))
-		{
-			return !InputInterface->IsAnyPhysicalKeyboardConnected();
-		}
-	}
-	
-	// Falback to showing the virtual keyboard
-	return true;
+#if !PLATFORM_TVOS
+    return PLATFORM_HAS_TOUCH_MAIN_SCREEN;
+#else
+    return true;
+#endif
 }
diff --git a/Engine/Source/Runtime/ApplicationCore/Private/IOS/IOSView.cpp b/Engine/Source/Runtime/ApplicationCore/Private/IOS/IOSView.cpp
--- a/Engine/Source/Runtime/ApplicationCore/Private/IOS/IOSView.cpp
+++ b/Engine/Source/Runtime/ApplicationCore/Private/IOS/IOSView.cpp
@@ -664,9 +664,8 @@ self.accessibilityElements = @[Window.accessibilityContainer];
 		[CachedMarkedText release];
 		CachedMarkedText = nil;
 	}
-	
-	TSharedPtr<FAppleKeyboardController> KeyboardController = ([self GetKeyboardController]).Pin();
-	if (!KeyboardController)
+
+	if ([self IsAnyPhysicalKeyboardConnected])
 	{
 		return;
 	}
@@ -676,14 +675,12 @@ self.accessibilityElements = @[Window.accessibilityContainer];
 	{
 		int32 Char = [theText characterAtIndex:CharIndex];
 
-		FAppleKeyboardController::FDeferredEvent KeyEvent;
-		
-		KeyEvent.CharCode = Char;
+		// FPlatformMisc::LowLevelOutputDebugStringf(TEXT("sending key '%c' to game\n"), Char);
 
 		if (Char == '\n')
 		{
 			// send the enter keypress
-			KeyEvent.TranslatedKeyCode = KEYCODE_ENTER;
+			FIOSInputInterface::QueueKeyInput(KEYCODE_ENTER, Char);
 			
 			// hide the keyboard
 			if (!bIsUsingIntegratedKeyboard)
@@ -693,10 +690,8 @@ self.accessibilityElements = @[Window.accessibilityContainer];
 		}
 		else
 		{
-			KeyEvent.TranslatedKeyCode = Char;
+			FIOSInputInterface::QueueKeyInput(Char, Char);
 		}
-
-		KeyboardController->QueueOneOffExternalKeyEvent(KeyEvent);
 	}
 }
 
@@ -708,17 +703,12 @@ self.accessibilityElements = @[Window.accessibilityContainer];
 	}
 
 
-	TSharedPtr<FAppleKeyboardController> KeyboardController = ([self GetKeyboardController]).Pin();
-	if (!KeyboardController)
+	if ([self IsAnyPhysicalKeyboardConnected])
 	{
 		return;
 	}
-	
-	FAppleKeyboardController::FDeferredEvent KeyEvent;
-	KeyEvent.CharCode = '\b';
-	KeyEvent.TranslatedKeyCode = KEYCODE_BACKSPACE;
-	
-	KeyboardController->QueueOneOffExternalKeyEvent(KeyEvent);
+
+	FIOSInputInterface::QueueKeyInput(KEYCODE_BACKSPACE, '\b');
 }
 
 -(void)ActivateKeyboard:(bool)bInSendEscapeOnClose
@@ -790,21 +780,17 @@ self.accessibilityElements = @[Window.accessibilityContainer];
 	return FParse::Param(FCommandLine::Get(), TEXT("NewKeyboard")) || bUseIntegratedKeyboard;
 }
 
--(TWeakPtr<FAppleKeyboardController>)GetKeyboardController
+-(bool)IsAnyPhysicalKeyboardConnected;
 {
-	FIOSApplication* IOSApplication = [IOSAppDelegate GetDelegate].IOSApplication;
-	FIOSInputInterface* InputInterface = IOSApplication ? IOSApplication->GetIOSInputInterface() : nullptr;
-	
-	return InputInterface ? InputInterface->GetKeyboardController() : nullptr;
+	FIOSInputInterface* InputInterface = [IOSAppDelegate GetDelegate].IOSApplication->GetIOSInputInterface();
+
+	return InputInterface ? InputInterface->IsAnyPhysicalKeyboardConnected() : false;
 }
 
 -(BOOL)becomeFirstResponder
 {
-	TSharedPtr<FAppleKeyboardController> KeyboardController = ([self GetKeyboardController]).Pin();
-	bool bIsAnyPhysicalKeyboardConnected = KeyboardController && KeyboardController->IsAnyKeyboardConnected();
-
 	volatile int32 ShowCount = KeyboardShowCount;
-	if (ShowCount >= 1 && !bIsAnyPhysicalKeyboardConnected)
+	if (ShowCount >= 1)
 	{
 		return [super becomeFirstResponder];
 	}
@@ -818,16 +804,8 @@ self.accessibilityElements = @[Window.accessibilityContainer];
 {
 	if (bSendEscapeOnClose)
 	{
-		TSharedPtr<FAppleKeyboardController> KeyboardController = ([self GetKeyboardController]).Pin();
-		if (KeyboardController)
-		{
-			FAppleKeyboardController::FDeferredEvent KeyEvent;
-			KeyEvent.CharCode = 0;
-			KeyEvent.TranslatedKeyCode = KEYCODE_ESCAPE;
-
-			// tell the console to close itself
-			KeyboardController->QueueOneOffExternalKeyEvent(KeyEvent);
-		}
+		// tell the console to close itself
+		FIOSInputInterface::QueueKeyInput(KEYCODE_ESCAPE, 0);
 	}
 	
 	return [super resignFirstResponder];
@@ -1313,92 +1291,4 @@ self.accessibilityElements = @[Window.accessibilityContainer];
  } 
  #endif
 
-- (BOOL)canBecomeFirstResponder
-{
-	return true;
-}
-
--(TWeakPtr<FAppleKeyboardController>)GetKeyboardController
-{
-	FIOSApplication* IOSApplication = [IOSAppDelegate GetDelegate].IOSApplication;
-	FIOSInputInterface* InputInterface = IOSApplication ? IOSApplication->GetIOSInputInterface() : nullptr;
-	return InputInterface ? InputInterface->GetKeyboardController() : nullptr;
-}
-
-- (bool)IsAnyPhysicalKeyboardConnected
-{
-	TSharedPtr<FAppleKeyboardController> KeyboardController = ([self GetKeyboardController]).Pin();
-	return KeyboardController ? KeyboardController->IsAnyKeyboardConnected() : false;
-}
-
-- (bool)ProcessKeyPresses:(NSSet<UIPress *> *)keypresses eventType:(FAppleKeyboardController::EKeyEvent)eventType
-{
-	TSharedPtr<FAppleKeyboardController> KeyboardController = ([self GetKeyboardController]).Pin();
-	if (KeyboardController && KeyboardController->IsAnyKeyboardConnected())
-	{
-		for (UIPress* Press in keypresses.allObjects)
-		{
-			if (UIKey* PressedKey = Press ? Press.key : nil)
-			{
-				// Sequences like backspace or new lines are handled by the controller internally using GC Key events
-				// So it is safe (and desired) ignore them here
-
-				unichar Char = [PressedKey.characters length] == 1 ? [PressedKey.characters characterAtIndex:0] : 0;
-				if (KeyboardController->IsPrintableChar((TCHAR) Char))
-				{
-					FAppleKeyboardController::FDeferredEvent KeyEvent;
-					KeyEvent.KeyEventType = eventType;
-					KeyEvent.TranslatedKeyCode = Char;
-					KeyEvent.CharCode = Char;
-					
-					KeyboardController->QueueExternalKeyChangeEvent(KeyEvent);
-				}
-			}
-		}
-		
-		return true;
-	}
-	
-	return false;
-}
-
-- (void)pressesBegan:(NSSet<UIPress *> *)presses withEvent:(nullable UIPressesEvent *)event
-{
-	if ([self ProcessKeyPresses:presses eventType:FAppleKeyboardController::EKeyEvent::KeyDown])
-	{
-		return;
-	}
-	
-	[super pressesBegan:presses withEvent:event];
-}
-
-- (void)pressesChanged:(NSSet<UIPress *> *)presses withEvent:(nullable UIPressesEvent *)event
-{
-	if ([self IsAnyPhysicalKeyboardConnected])
-	{
-		return;
-	}
-	[super pressesChanged:presses withEvent:event];
-}
-
-- (void)pressesEnded:(NSSet<UIPress *> *)presses withEvent:(nullable UIPressesEvent *)event
-{
-	if ([self ProcessKeyPresses:presses eventType:FAppleKeyboardController::EKeyEvent::KeyUp])
-	{
-		return;
-	}
-
-	[super pressesEnded:presses withEvent:event];
-}
-
-- (void)pressesCancelled:(NSSet<UIPress *> *)presses withEvent:(nullable UIPressesEvent *)event
-{
-	if ([self ProcessKeyPresses:presses eventType:FAppleKeyboardController::EKeyEvent::KeyUp])
-	{
-		return;
-	}
-
-	[super pressesCancelled:presses withEvent:event];
-}
-
 @end
diff --git a/Engine/Source/Runtime/ApplicationCore/Public/IOS/IOSInputInterface.h b/Engine/Source/Runtime/ApplicationCore/Public/IOS/IOSInputInterface.h
--- a/Engine/Source/Runtime/ApplicationCore/Public/IOS/IOSInputInterface.h
+++ b/Engine/Source/Runtime/ApplicationCore/Public/IOS/IOSInputInterface.h
@@ -124,6 +124,7 @@ public:
 	virtual void SetForceFeedbackChannelValues(int32 ControllerId, const FForceFeedbackValues &values) override;
 
 	static APPLICATIONCORE_API void QueueTouchInput(const TArray<TouchInput>& InTouchEvents);
+	static void QueueKeyInput(int32 Key, int32 Char);
 
 	void SetGamepadsAllowed(bool bAllowed) { bAllowControllers = bAllowed; }
 	void SetGamepadsBlockDeviceFeedback(bool bBlock) { bControllersBlockDeviceFeedback = bBlock; }
@@ -177,7 +178,7 @@ private:
 	void CalibrateMotion(uint32 PlayerIndex);
 
 private:
-	void ProcessTouchesAndKeys(uint32 ControllerId, const TArray<TouchInput>& InTouchInputStack);
+	void ProcessTouchesAndKeys(uint32 ControllerId, const TArray<TouchInput>& InTouchInputStack, const TArray<int32>& InKeyInputStack);
 	
 	// can the remote be rotated to landscape
 	bool bAllowRemoteRotation;
diff --git a/Engine/Source/Runtime/ApplicationCore/Public/IOS/IOSView.h b/Engine/Source/Runtime/ApplicationCore/Public/IOS/IOSView.h
--- a/Engine/Source/Runtime/ApplicationCore/Public/IOS/IOSView.h
+++ b/Engine/Source/Runtime/ApplicationCore/Public/IOS/IOSView.h
@@ -3,11 +3,9 @@
 #pragma once
 
 #ifndef SWIFT_IMPORT
-#include "AppleKeyboardController.h"
 #include "CoreTypes.h"
 #include "Delegates/Delegate.h"
 #include "IOS/IOSInputInterface.h"
-#include "Templates/SharedPointer.h"
 #endif
 
 #import <UIKit/UIKit.h>
@@ -118,7 +116,8 @@ APPLICATIONCORE_API
 -(void)ActivateKeyboard:(bool)bInSendEscapeOnClose keyboardConfig:(FKeyboardConfig)KeyboardConfig;
 -(void)DeactivateKeyboard;
 -(bool)ShouldUseIntegratedKeyboard;
--(TWeakPtr<FAppleKeyboardController>) GetKeyboardController;
+
+-(bool)IsAnyPhysicalKeyboardConnected;
 
 // callable from outside to fake locations
 -(void)HandleTouchAtLoc:(CGPoint)Loc PrevLoc:(CGPoint)PrevLoc TouchIndex:(int)TouchIndex Force:(float)Force Type:(TouchType)Type TouchesArray:(TArray<TouchInput>&)TouchesArray;
@@ -135,10 +134,13 @@ APPLICATIONCORE_API
 /**
  * A view controller subclass that handles loading our IOS view as well as autorotation
  */
+#if PLATFORM_TVOS
 #import <GameController/GameController.h>
-// We use a GCEvent view controller to dynamically control how Game Controller API events flow in the responder chain, based on the presence of a physical keyboard or not
+// if TVOS doesn't use the GCEventViewController, it will background the app when the user presses Menu/Pause
 @interface IOSViewController : GCEventViewController
-
+#else
+@interface IOSViewController : UIViewController
+#endif
 {
 }
 
@@ -149,15 +151,4 @@ APPLICATIONCORE_API
 - (UIInterfaceOrientationMask)supportedInterfaceOrientations;
 #endif
 
-- (TWeakPtr<FAppleKeyboardController>)GetKeyboardController;
-- (bool)IsAnyPhysicalKeyboardConnected;
-
-- (bool)ProcessKeyPresses:(NSSet<UIPress *> *)keypresses eventType:(FAppleKeyboardController::EKeyEvent)eventType;
-
-- (void)pressesBegan:(NSSet<UIPress *> *)presses withEvent:(nullable UIPressesEvent *)event;
-- (void)pressesChanged:(NSSet<UIPress *> *)presses withEvent:(nullable UIPressesEvent *)event;
-- (void)pressesEnded:(NSSet<UIPress *> *)presses withEvent:(nullable UIPressesEvent *)event;
-- (void)pressesCancelled:(NSSet<UIPress *> *)presses withEvent:(nullable UIPressesEvent *)event;
-- (bool)canBecomeFirstResponder;
-
 @end
diff --git a/Engine/Source/Runtime/Launch/Private/IOS/LaunchIOS.cpp b/Engine/Source/Runtime/Launch/Private/IOS/LaunchIOS.cpp
--- a/Engine/Source/Runtime/Launch/Private/IOS/LaunchIOS.cpp
+++ b/Engine/Source/Runtime/Launch/Private/IOS/LaunchIOS.cpp
@@ -320,10 +320,6 @@ void FAppEntry::PreInit(IOSAppDelegate* AppDelegate, UIApplication* Application)
 	// make a controller object
 	IOSViewController* IOSController = [[IOSViewController alloc] init];
 	
-	// IOSViewController is a GCViewController so we can have proper GC Controller support
-	// But default we want all inputs to be forwaded to UI Kit. This will only change when we have physical keyborad or mouse connected
-	IOSController.controllerUserInteractionEnabled = true;
-	
 #if PLATFORM_TVOS
 	// @todo tvos: This may need to be exposed to the game so that when you click Menu it will background the app
 	// this is basically the same way Android handles the Back button (maybe we should pass Menu button as back... maybe)
@@ -386,8 +382,6 @@ static void MainThreadInit()
 #else
 	AppDelegate.IOSView = [[FIOSView alloc] initWithFrame:FullResolutionRect];
 	AppDelegate.IOSView.clearsContextBeforeDrawing = NO;
-	
-	
 #if !PLATFORM_TVOS
 	AppDelegate.IOSView.multipleTouchEnabled = YES;
 #endif
