Fix crash in FAndroidInputInterface::GetAlternateKeyEventForMouse

#rnx

#Synced-CL 49075355

diff --git a/Engine/Source/Runtime/ApplicationCore/Private/Android/AndroidInputInterface.cpp b/Engine/Source/Runtime/ApplicationCore/Private/Android/AndroidInputInterface.cpp
--- a/Engine/Source/Runtime/ApplicationCore/Private/Android/AndroidInputInterface.cpp
+++ b/Engine/Source/Runtime/ApplicationCore/Private/Android/AndroidInputInterface.cpp
@@ -1875,17 +1875,20 @@ int32 FAndroidInputInterface::GetAlternateKeyEventForMouse(int32 deviceId, int32
 
 	if (buttonId == 0)
 	{
-		FAndroidInputDeviceInfo* Info;
-		FAndroidGamepadDeviceMapping* data;
-		if (GetInputDeviceByDeviceId(deviceId, &Info, &data) &&
-			Info->DeviceState == MappingState::Valid &&
-			data != nullptr && data->DeviceConfig.IsValid())
+		FAndroidInputDeviceInfo* Info = nullptr;
+		FAndroidGamepadDeviceMapping* data = nullptr;
+		if (GetInputDeviceByDeviceId(deviceId, &Info, &data))
 		{
-			DeviceConfig_SPtr DeviceConfigPtr = data->DeviceConfig.Pin();
-			check(DeviceConfigPtr.IsValid());
-			if (DeviceConfigPtr->ControllerClass == ControllerClassType::PlaystationWireless)
+			ensure(Info); // must always be set if GetInputDeviceByDeviceId returned true
+			if (Info->DeviceState == MappingState::Valid &&
+				data != nullptr && data->DeviceConfig.IsValid())
 			{
-				return 3002;
+				DeviceConfig_SPtr DeviceConfigPtr = data->DeviceConfig.Pin();
+				check(DeviceConfigPtr.IsValid());
+				if (DeviceConfigPtr->ControllerClass == ControllerClassType::PlaystationWireless)
+				{
+					return 3002;
+				}
 			}
 		}
 	}
