Clean up Windows DPI handling

We require at least Windows 10, so Windows 8.1 functionality is safe to call directly.

#Synced-CL 49199121

diff --git a/Engine/Source/Runtime/ApplicationCore/ApplicationCore.Build.cs b/Engine/Source/Runtime/ApplicationCore/ApplicationCore.Build.cs
--- a/Engine/Source/Runtime/ApplicationCore/ApplicationCore.Build.cs
+++ b/Engine/Source/Runtime/ApplicationCore/ApplicationCore.Build.cs
@@ -43,6 +43,9 @@ public class ApplicationCore : ModuleRules
 
 			// Uses DXGI to query GPU hardware prior to RHI startup
 			PublicSystemLibraries.Add("DXGI.lib");
+
+			// DPI functions
+			PublicSystemLibraries.Add("Shcore.lib");
 		}
 		else if (Target.Platform == UnrealTargetPlatform.Mac)
 		{
diff --git a/Engine/Source/Runtime/ApplicationCore/Private/Windows/WindowsPlatformApplicationMisc.cpp b/Engine/Source/Runtime/ApplicationCore/Private/Windows/WindowsPlatformApplicationMisc.cpp
--- a/Engine/Source/Runtime/ApplicationCore/Private/Windows/WindowsPlatformApplicationMisc.cpp
+++ b/Engine/Source/Runtime/ApplicationCore/Private/Windows/WindowsPlatformApplicationMisc.cpp
@@ -31,14 +31,12 @@ THIRD_PARTY_INCLUDES_START
 #include "dxgi1_3.h"
 #include "dxgi1_4.h"
 #include "dxgi1_6.h"
+#include "shellscalingapi.h"
 #include "Windows/HideWindowsPlatformTypes.h"
 THIRD_PARTY_INCLUDES_END
 
 #define LOCTEXT_NAMESPACE "WindowsPlatformApplicationMisc"
 
-typedef HRESULT(STDAPICALLTYPE *GetDpiForMonitorProc)(HMONITOR Monitor, int32 DPIType, uint32 *DPIX, uint32 *DPIY);
-APPLICATIONCORE_API GetDpiForMonitorProc GetDpiForMonitor;
-
 void FWindowsPlatformApplicationMisc::PreInit()
 {
 	FApp::SetHasFocusFunction(&FWindowsPlatformApplicationMisc::IsThisApplicationForeground);
@@ -292,54 +290,20 @@ void FWindowsPlatformApplicationMisc::SetHighDPIMode()
 		bool bAllowHighDpiWhenUnattended = false;
 		GConfig->GetBool(TEXT("/Script/Engine.UserInterfaceSettings"), TEXT("bAllowHighDpiWhenUnattended"), bAllowHighDpiWhenUnattended, GEngineIni);
 
-		if (void* ShCoreDll = FPlatformProcess::GetDllHandle(TEXT("shcore.dll")))
+		if (!IsRunningCommandlet() && (!FApp::IsUnattended() || bAllowHighDpiWhenUnattended))
 		{
-			typedef enum _PROCESS_DPI_AWARENESS {
-				PROCESS_DPI_UNAWARE = 0,
-				PROCESS_SYSTEM_DPI_AWARE = 1,
-				PROCESS_PER_MONITOR_DPI_AWARE = 2
-			} PROCESS_DPI_AWARENESS;
-
-			typedef HRESULT(STDAPICALLTYPE *SetProcessDpiAwarenessProc)(PROCESS_DPI_AWARENESS Value);
-			SetProcessDpiAwarenessProc SetProcessDpiAwareness = (SetProcessDpiAwarenessProc)FPlatformProcess::GetDllExport(ShCoreDll, TEXT("SetProcessDpiAwareness"));
-			GetDpiForMonitor = (GetDpiForMonitorProc)FPlatformProcess::GetDllExport(ShCoreDll, TEXT("GetDpiForMonitor"));
+			PROCESS_DPI_AWARENESS CurrentAwareness = PROCESS_DPI_UNAWARE;
 
-			typedef HRESULT(STDAPICALLTYPE *GetProcessDpiAwarenessProc)(HANDLE hProcess, PROCESS_DPI_AWARENESS* Value);
-			GetProcessDpiAwarenessProc GetProcessDpiAwareness = (GetProcessDpiAwarenessProc)FPlatformProcess::GetDllExport(ShCoreDll, TEXT("GetProcessDpiAwareness"));
+			GetProcessDpiAwareness(nullptr, &CurrentAwareness);
 
-			if (SetProcessDpiAwareness && GetProcessDpiAwareness && !IsRunningCommandlet() && (!FApp::IsUnattended() || bAllowHighDpiWhenUnattended))
+			if (CurrentAwareness != PROCESS_PER_MONITOR_DPI_AWARE)
 			{
-				PROCESS_DPI_AWARENESS CurrentAwareness = PROCESS_DPI_UNAWARE;
-
-				GetProcessDpiAwareness(nullptr, &CurrentAwareness);
-
-				if (CurrentAwareness != PROCESS_PER_MONITOR_DPI_AWARE)
+				UE_LOG(LogInit, Log, TEXT("Setting process to per monitor DPI aware"));
+				HRESULT Hr = SetProcessDpiAwareness(PROCESS_PER_MONITOR_DPI_AWARE); // PROCESS_PER_MONITOR_DPI_AWARE_VALUE
+																					// We dont care about this warning if we are in any kind of headless mode
+				if (Hr != S_OK)
 				{
-					UE_LOG(LogInit, Log, TEXT("Setting process to per monitor DPI aware"));
-					HRESULT Hr = SetProcessDpiAwareness(PROCESS_PER_MONITOR_DPI_AWARE); // PROCESS_PER_MONITOR_DPI_AWARE_VALUE
-																						// We dont care about this warning if we are in any kind of headless mode
-					if (Hr != S_OK)
-					{
-						UE_LOG(LogInit, Warning, TEXT("SetProcessDpiAwareness failed.  Error code %x"), Hr);
-					}
-				}
-			}
-
-			FPlatformProcess::FreeDllHandle(ShCoreDll);
-		}
-		else if (void* User32Dll = GetModuleHandle(L"user32.dll"))
-		{
-			typedef BOOL(WINAPI *SetProcessDpiAwareProc)(void);
-			SetProcessDpiAwareProc SetProcessDpiAware = (SetProcessDpiAwareProc)FPlatformProcess::GetDllExport(User32Dll, TEXT("SetProcessDPIAware"));
-
-			if (SetProcessDpiAware && !IsRunningCommandlet() && (!FApp::IsUnattended() || bAllowHighDpiWhenUnattended))
-			{
-				UE_LOG(LogInit, Log, TEXT("Setting process to DPI aware"));
-
-				BOOL Result = SetProcessDpiAware();
-				if (Result == 0)
-				{
-					UE_LOG(LogInit, Warning, TEXT("SetProcessDpiAware failed"));
+					UE_LOG(LogInit, Warning, TEXT("SetProcessDpiAwareness failed.  Error code %x"), Hr);
 				}
 			}
 		}
@@ -382,31 +346,22 @@ int32 FWindowsPlatformApplicationMisc::GetMonitorDPI(const FMonitorInfo& Monitor
 
 	if (IsHighDPIAwarenessEnabled())
 	{
-		if (GetDpiForMonitor)
+		RECT MonitorDim;
+		MonitorDim.left = MonitorInfo.DisplayRect.Left;
+		MonitorDim.top = MonitorInfo.DisplayRect.Top;
+		MonitorDim.right = MonitorInfo.DisplayRect.Right;
+		MonitorDim.bottom = MonitorInfo.DisplayRect.Bottom;
+
+		HMONITOR Monitor = MonitorFromRect(&MonitorDim, MONITOR_DEFAULTTONEAREST);
+		if (Monitor)
 		{
-			RECT MonitorDim;
-			MonitorDim.left = MonitorInfo.DisplayRect.Left;
-			MonitorDim.top = MonitorInfo.DisplayRect.Top;
-			MonitorDim.right = MonitorInfo.DisplayRect.Right;
-			MonitorDim.bottom = MonitorInfo.DisplayRect.Bottom;
-
-			HMONITOR Monitor = MonitorFromRect(&MonitorDim, MONITOR_DEFAULTTONEAREST);
-			if (Monitor)
+			uint32 DPIX = 0;
+			uint32 DPIY = 0;
+			if (SUCCEEDED(GetDpiForMonitor(Monitor, MDT_EFFECTIVE_DPI, &DPIX, &DPIY)))
 			{
-				uint32 DPIX = 0;
-				uint32 DPIY = 0;
-				if (SUCCEEDED(GetDpiForMonitor(Monitor, 0 /*MDT_EFFECTIVE_DPI*/, &DPIX, &DPIY)))
-				{
-					DisplayDPI = DPIX;
-				}
+				DisplayDPI = DPIX;
 			}
 		}
-		else
-		{
-			HDC Context = GetDC(nullptr);
-			DisplayDPI = GetDeviceCaps(Context, LOGPIXELSX);
-			ReleaseDC(nullptr, Context);
-		}
 	}
 
 	return DisplayDPI;
@@ -446,27 +401,17 @@ float FWindowsPlatformApplicationMisc::GetDPIScaleFactorAtPoint(float X, float Y
 
 	if (IsHighDPIAwarenessEnabled())
 	{
-		if (GetDpiForMonitor)
+		POINT Position = { static_cast<LONG>(X), static_cast<LONG>(Y) };
+		HMONITOR Monitor = MonitorFromPoint(Position, MONITOR_DEFAULTTONEAREST);
+		if (Monitor)
 		{
-			POINT Position = { static_cast<LONG>(X), static_cast<LONG>(Y) };
-			HMONITOR Monitor = MonitorFromPoint(Position, MONITOR_DEFAULTTONEAREST);
-			if (Monitor)
+			uint32 DPIX = 0;
+			uint32 DPIY = 0;
+			if (SUCCEEDED(GetDpiForMonitor(Monitor, MDT_EFFECTIVE_DPI, &DPIX, &DPIY)))
 			{
-				uint32 DPIX = 0;
-				uint32 DPIY = 0;
-				if (SUCCEEDED(GetDpiForMonitor(Monitor, 0 /*MDT_EFFECTIVE_DPI*/, &DPIX, &DPIY)))
-				{
-					Scale = (float)DPIX / USER_DEFAULT_SCREEN_DPI;
-				}
+				Scale = (float)DPIX / USER_DEFAULT_SCREEN_DPI;
 			}
 		}
-		else
-		{
-			HDC Context = GetDC(nullptr);
-			int32 DPI = GetDeviceCaps(Context, LOGPIXELSX);
-			Scale = (float)DPI / USER_DEFAULT_SCREEN_DPI;
-			ReleaseDC(nullptr, Context);
-		}
 	}
 
 	return Scale;
