[iOS/macOS] Setting a name to the game thread, so it can be identified in Instruments.

#Synced-CL 49198593

diff --git a/Engine/Source/Runtime/ApplicationCore/Private/IOS/IOSAppDelegate.cpp b/Engine/Source/Runtime/ApplicationCore/Private/IOS/IOSAppDelegate.cpp
--- a/Engine/Source/Runtime/ApplicationCore/Private/IOS/IOSAppDelegate.cpp
+++ b/Engine/Source/Runtime/ApplicationCore/Private/IOS/IOSAppDelegate.cpp
@@ -1469,6 +1469,7 @@ static FAutoConsoleVariableRef CVarGEnableThermalsReport(
 	NSThread* GameThread = [[NSThread alloc] initWithTarget:self selector:@selector(MainAppThread:) object:self.launchOptions];
 	[GameThread setStackSize:GAME_THREAD_STACK_SIZE];
 	[GameThread start];
+	[GameThread setName:@"GameThread"];
 
 	// this can be slow (1/3 of a second!), so don't make the game thread stall loading for it
 	// check to see if we are using the network file system, if so, disable the idle timer
diff --git a/Engine/Source/Runtime/Core/Private/Mac/CocoaThread.cpp b/Engine/Source/Runtime/Core/Private/Mac/CocoaThread.cpp
--- a/Engine/Source/Runtime/Core/Private/Mac/CocoaThread.cpp
+++ b/Engine/Source/Runtime/Core/Private/Mac/CocoaThread.cpp
@@ -501,6 +501,7 @@ void RunGameThread(id Target, SEL Selector)
 	FCocoaGameThread* GameThread = [[FCocoaGameThread alloc] initWithTarget:Target selector:Selector object:nil];
 	[GameThread setStackSize:GAME_THREAD_STACK_SIZE];
 	[GameThread start];
+	[GameThread setName:@"GameThread"];
 #else
 	[Target performSelector:Selector withObject:nil];
 	
