[iOS] Fix CIS issue

Adding a missing UI Kit header

#Synced-CL 49336016

diff --git a/Engine/Source/Runtime/ApplicationCore/Public/IOS/IOSInputUtils.h b/Engine/Source/Runtime/ApplicationCore/Public/IOS/IOSInputUtils.h
--- a/Engine/Source/Runtime/ApplicationCore/Public/IOS/IOSInputUtils.h
+++ b/Engine/Source/Runtime/ApplicationCore/Public/IOS/IOSInputUtils.h
@@ -4,6 +4,7 @@
 
 #include "HAL/Platform.h"
 
+#import <UIKit/UIKit.h>
 #import <UIKit/UIKeyConstants.h>
 
 namespace IOS
