Replace MONITOR_DEFAULTTOPRIMARY with MONITOR_DEFAULTTONEAREST

DXGI uses MONITOR_DEFAULTTONEAREST, so it doesn't make sense to mismatch.

#Synced-CL 49357083

diff --git a/Engine/Source/Runtime/ApplicationCore/Private/Windows/WindowsWindow.cpp b/Engine/Source/Runtime/ApplicationCore/Private/Windows/WindowsWindow.cpp
--- a/Engine/Source/Runtime/ApplicationCore/Private/Windows/WindowsWindow.cpp
+++ b/Engine/Source/Runtime/ApplicationCore/Private/Windows/WindowsWindow.cpp
@@ -561,10 +561,8 @@ void FWindowsWindow::ReshapeWindow( int32 NewX, int32 NewY, int32 NewWidth, int3
 
 bool FWindowsWindow::GetFullScreenInfo( int32& X, int32& Y, int32& Width, int32& Height ) const
 {
-	bool bTrueFullscreen = ( WindowMode == EWindowMode::Fullscreen );
-
 	// Grab current monitor data for sizing
-	HMONITOR Monitor = MonitorFromWindow( HWnd, bTrueFullscreen ? MONITOR_DEFAULTTOPRIMARY : MONITOR_DEFAULTTONEAREST );
+	HMONITOR Monitor = MonitorFromWindow( HWnd, MONITOR_DEFAULTTONEAREST );
 	MONITORINFO MonitorInfo;
 	MonitorInfo.cbSize = sizeof(MONITORINFO);
 	GetMonitorInfo( Monitor, &MonitorInfo );
@@ -829,7 +827,7 @@ void FWindowsWindow::SetWindowMode( EWindowMode::Type NewWindowMode )
 			GetClientRect(HWnd, &ClientRect);
 
 			// Grab current monitor data for sizing
-			HMONITOR Monitor = MonitorFromWindow( HWnd, bTrueFullscreen ? MONITOR_DEFAULTTOPRIMARY : MONITOR_DEFAULTTONEAREST );
+			HMONITOR Monitor = MonitorFromWindow( HWnd, MONITOR_DEFAULTTONEAREST );
 			MONITORINFO MonitorInfo;
 			MonitorInfo.cbSize = sizeof(MONITORINFO);
 			GetMonitorInfo( Monitor, &MonitorInfo );
@@ -919,8 +917,7 @@ bool FWindowsWindow::GetRestoredDimensions(int32& X, int32& Y, int32& Width, int
 		if ((WindowExStyle & WS_EX_TOOLWINDOW) == 0)
 		{
 			// For windows without WS_EX_TOOLWINDOW window style WindowPlacement.rcNormalPosition is in workspace coordinates, so we need to convert the position to screen coordinates
-			const bool bTrueFullscreen = (WindowMode == EWindowMode::Fullscreen);
-			HMONITOR Monitor = MonitorFromWindow(HWnd, bTrueFullscreen ? MONITOR_DEFAULTTOPRIMARY : MONITOR_DEFAULTTONEAREST);
+			HMONITOR Monitor = MonitorFromWindow(HWnd, MONITOR_DEFAULTTONEAREST);
 			MONITORINFO MonitorInfo;
 			MonitorInfo.cbSize = sizeof(MONITORINFO);
 			GetMonitorInfo(Monitor, &MonitorInfo);
diff --git a/Engine/Source/Runtime/OpenGLDrv/Private/Windows/WindowsOpenGLPlatform.cpp b/Engine/Source/Runtime/OpenGLDrv/Private/Windows/WindowsOpenGLPlatform.cpp
--- a/Engine/Source/Runtime/OpenGLDrv/Private/Windows/WindowsOpenGLPlatform.cpp
+++ b/Engine/Source/Runtime/OpenGLDrv/Private/Windows/WindowsOpenGLPlatform.cpp
@@ -565,7 +565,7 @@ void PlatformResizeGLContext(FPlatformOpenGLDevice* Device, FPlatformOpenGLConte
 		if (bFullscreen)
 		{
 			// Get the monitor info from the window handle.
-			HMONITOR hMonitor = MonitorFromWindow(Context->WindowHandle, MONITOR_DEFAULTTOPRIMARY);
+			HMONITOR hMonitor = MonitorFromWindow(Context->WindowHandle, MONITOR_DEFAULTTONEAREST);
 			MONITORINFOEX MonitorInfo;
 			memset(&MonitorInfo, 0, sizeof(MONITORINFOEX));
 			MonitorInfo.cbSize = sizeof(MONITORINFOEX);
