Add double click on title bar behavior to maximize/restore the window

#rnx
#Synced-CL 49371807

diff --git a/Engine/Source/Runtime/ApplicationCore/Private/Linux/LinuxApplication.cpp b/Engine/Source/Runtime/ApplicationCore/Private/Linux/LinuxApplication.cpp
--- a/Engine/Source/Runtime/ApplicationCore/Private/Linux/LinuxApplication.cpp
+++ b/Engine/Source/Runtime/ApplicationCore/Private/Linux/LinuxApplication.cpp
@@ -108,8 +108,14 @@ FLinuxApplication* FLinuxApplication::CreateLinuxApplication()
 
 FLinuxApplication::FLinuxApplication() 
 	:	GenericApplication( MakeShareable( new FLinuxCursor() ) )
+	,	SDLDoubleClickTime(0)
+	,	SDLDoubleClickRadius(0)	
 	,	bIsMouseCursorLocked(false)
 	,	bIsMouseCaptureEnabled(false)
+	,	LastTitleBarClickX(0.0f)
+	,	LastTitleBarClickY(0.0f)
+	,	LastTitleBarClickTimestampNS(0ULL)
+	,	LastTitleBarClickWindow(nullptr)
 	,	bHasLoadedInputPlugins(false)
 	,	bInsideOwnWindow(false)
 	,	bIsDragWindowButtonPressed(false)
@@ -901,6 +907,8 @@ void FLinuxApplication::ProcessDeferredMessage( SDL_Event Event )
 
 	case SDL_EVENT_WINDOW_HIT_TEST:
 		{
+			UE_LOG(LogLinuxWindowEvent, Verbose, TEXT("SDL_EVENT_WINDOW_HIT_TEST received"));
+
 			// The user clicked into the hit test area (Titlebar for example). Is the whole Application active?
 			// If not, first activate (ignore tooltips).
 			if (!bActivateApp && !CurrentEventWindow->IsTooltipWindow())
@@ -927,6 +935,50 @@ void FLinuxApplication::ProcessDeferredMessage( SDL_Event Event )
 			}
 
 			bFirstFrameOfWindowMove = true;
+
+			// Using the SDL_EVENT_WINDOW_HIT_TEST event as a replacement to a mouse click event.  Which SDL doesn't send for clicks 
+			// in the draggable hit test area (title bars in UE)
+			// If we see two of these close enough together, assume it's a double click on the titlebar area and maximize/restore the window
+
+			uint64 EventTimeNS;
+			float MouseX;
+			float MouseY;
+
+			// Get the event time from the event
+			const SDL_WindowEvent &WindowEvent = Event.window;
+			EventTimeNS = WindowEvent.timestamp;
+			// The event doesn't pass mouse coordinates, it always passes 0,0 so we are going to get them from the system
+			SDL_GetGlobalMouseState(&MouseX, &MouseY);
+
+			if (LastTitleBarClickTimestampNS > 0 && LastTitleBarClickWindow.IsValid() &&
+				LastTitleBarClickWindow == CurrentEventWindow &&
+				((EventTimeNS - LastTitleBarClickTimestampNS) <= (GetSDLDoubleClickTime() * 1000000ULL)) &&
+				(FMath::Abs(MouseX - LastTitleBarClickX) <= (float)GetSDLDoubleClickRadius()) &&
+				(FMath::Abs(MouseY - LastTitleBarClickY) <= (float)GetSDLDoubleClickRadius()))
+			{
+				// double clicked the title bar we need to maximize or restore the window
+				UE_LOG(LogLinuxWindowEvent, Verbose, TEXT("MOUSE Double Click on Titlebar X:%f Y:%f"), MouseX, MouseY);
+				uint32 flag = SDL_GetWindowFlags(CurrentEventWindow->GetHWnd());
+				if (flag & SDL_WINDOW_MAXIMIZED)
+				{
+					SDL_RestoreWindow(CurrentEventWindow->GetHWnd());
+				}
+				else
+				{
+					SDL_MaximizeWindow(CurrentEventWindow->GetHWnd());
+				}
+				// reset the last click
+				LastTitleBarClickTimestampNS = 0;
+				LastTitleBarClickWindow = nullptr;
+			}
+			else
+			{
+				// remember this click and see if another one happens quickly
+				LastTitleBarClickTimestampNS = EventTimeNS;
+				LastTitleBarClickX = MouseX;
+				LastTitleBarClickY = MouseY;
+				LastTitleBarClickWindow = CurrentEventWindow;
+			}
 		}
 		break;
 
@@ -2137,3 +2189,44 @@ void FLinuxApplication::OnInputDeviceModuleRegistered(const FName& Type, IModula
 	}
 }
 
+// These functions are duplicating how SDL3 works since it doesn't provide accessors to these values
+const uint32 FLinuxApplication::GetSDLDoubleClickTime()
+{
+	if (SDLDoubleClickTime == 0)
+	{
+		const char *hint = SDL_GetHint(SDL_HINT_MOUSE_DOUBLE_CLICK_TIME);
+		if (hint && *hint) 
+		{
+			int i = SDL_atoi(hint);
+
+			SDLDoubleClickTime = SDL_max(0, i);
+		}
+
+		if (SDLDoubleClickTime == 0)
+		{
+			SDLDoubleClickTime = 500;   // Linux defaults to 500 ms
+		}
+	}
+	return SDLDoubleClickTime;
+}
+
+const uint32 FLinuxApplication::GetSDLDoubleClickRadius()
+{
+	if (SDLDoubleClickRadius == 0)
+	{
+		const char *hint = SDL_GetHint(SDL_HINT_MOUSE_DOUBLE_CLICK_RADIUS);
+		if (hint && *hint) 
+		{
+			int i = SDL_atoi(hint);
+
+			SDLDoubleClickRadius = SDL_max(0, i);
+		}
+
+		if (SDLDoubleClickRadius == 0)
+		{
+			SDLDoubleClickRadius = 32;   // Linux defaults to 32 pixels
+		}
+	}
+
+	return SDLDoubleClickRadius;
+}
diff --git a/Engine/Source/Runtime/ApplicationCore/Public/Linux/LinuxApplication.h b/Engine/Source/Runtime/ApplicationCore/Public/Linux/LinuxApplication.h
--- a/Engine/Source/Runtime/ApplicationCore/Public/Linux/LinuxApplication.h
+++ b/Engine/Source/Runtime/ApplicationCore/Public/Linux/LinuxApplication.h
@@ -235,6 +235,12 @@ private:
 
 	void OnInputDeviceModuleRegistered(const FName& Type, IModularFeature* ModularFeature);
 
+	uint32 SDLDoubleClickTime;
+	uint32 SDLDoubleClickRadius;
+
+	const uint32 GetSDLDoubleClickTime();
+	const uint32 GetSDLDoubleClickRadius();
+
 	/** Stores context information about a currently active touch. */
 	struct FTouchContext
 	{
@@ -321,6 +327,12 @@ private:
 	/** True after every SDL_WINDOWEVENT_HIT_TEST until a following SDL_WINDOWEVENT_MOVED */
 	bool bFirstFrameOfWindowMove = false;
 
+	/** Used to detect a double click on the title bar area of the window */
+	float LastTitleBarClickX;
+	float LastTitleBarClickY;
+	uint64 LastTitleBarClickTimestampNS;
+	TSharedPtr< FLinuxWindow > LastTitleBarClickWindow;
+
 	/** Window that we think has been activated last. */
 	TSharedPtr< FLinuxWindow > CurrentlyActiveWindow;
 	TSharedPtr< FLinuxWindow > CurrentFocusWindow;
