Improve support for multiple monitors on Windows

The main change is to extend UGameUserSettings with members to track the DisplayID and DisplayIndex of the main display. The DisplayID and DisplayIndex can be used to determine the MonitorIndex using FDisplayMetrics::GetClosestMonitorFromIDAndIndex. MonitorIndex is plumbed around the engine to index into the FDisplayMetrics::MonitorInfo array. In the common case, DisplayIndex and MonitorIndex will be the same, but we have a preference for DisplayID if displays get shuffled around.

r.SetRes.MonitorIndex was added to complement r.SetRes for displaying on a specific monitor.

The code will try to preserve legacy behavior when MonitorIndex == -1, which usually means using the primary display members of FDisplayMetrics. If the user moves or resizes the window or uses a fullscreen toggle shortcut (F11 or Alt+Enter), then display properties will be written into GameUserSettings and cause MonitorIndex to become valid.

Add FMonitorInfo::FriendlyName for display purposes because FMonitorInfo::Name is not very friendly. Also add FMonitorInfo::NativeHandle, which is useful for calling RHIGetAvailableResolutionsForDisplay.

Modify FWindowsWindow and FWindowsApplication to prevent shenanigans if we are not in a modal loop for moving or sizing that was initiated by the user.

Simplify FMonitorInfo::DPI on Windows to use effective DPI logic, undoing changelist 4534956. If raw or angular DPI become necessary, we should supply them separately.

Fix FSlateApplication::LocateWidgetInWindow and FSlateApplication::TransformPointerEvent to allow exclusive fullscreen to compute mouse position correctly on secondary monitors.

Add FSlateApplicationBase::GetCachedDisplayMetricsByRef to avoid copies when using FSlateApplicationBase::GetCachedDisplayMetrics, but callers need to be careful.

Prevent remote desktop from entering exclusive fullscreen through F11 toggle by checking new function FViewport::IsExclusiveFullscreenSupported.

Fix settings screen not refreshing settings after video revert through new multicast delegate FOnGameUserSettingsVideoRevert.

#Synced-CL 49751596

diff --git a/Engine/Source/Runtime/ApplicationCore/Private/GenericPlatform/GenericApplication.cpp b/Engine/Source/Runtime/ApplicationCore/Private/GenericPlatform/GenericApplication.cpp
--- a/Engine/Source/Runtime/ApplicationCore/Private/GenericPlatform/GenericApplication.cpp
+++ b/Engine/Source/Runtime/ApplicationCore/Private/GenericPlatform/GenericApplication.cpp
@@ -137,6 +137,51 @@ FPlatformRect FDisplayMetrics::GetMonitorWorkAreaFromPoint(const FVector2D& Poin
 	return FPlatformRect(0, 0, 0, 0);
 }
 
+int32 FDisplayMetrics::GetClosestMonitorFromIDAndIndex(const FString& MonitorID, int32 MonitorIndex) const
+{
+	if (!MonitorID.IsEmpty())
+	{
+		// Match monitor by ID and index
+		for (int32 Index = 0, Num = MonitorInfo.Num(); Index < Num; ++Index)
+		{
+			const FMonitorInfo& Monitor = MonitorInfo[Index];
+			if (Monitor.ID == MonitorID && Index == MonitorIndex)
+			{
+				return Index;
+			}
+		}
+
+		// Match monitor by just ID
+		for (int32 Index = 0, Num = MonitorInfo.Num(); Index < Num; ++Index)
+		{
+			const FMonitorInfo& Monitor = MonitorInfo[Index];
+			if (Monitor.ID == MonitorID)
+			{
+				return Index;
+			}
+		}
+	}
+
+	// Match monitor by just index
+	if (MonitorInfo.IsValidIndex(MonitorIndex))
+	{
+		return MonitorIndex;
+	}
+
+	// Grab the primary monitor
+	for (int32 Index = 0, Num = MonitorInfo.Num(); Index < Num; ++Index)
+	{
+		const FMonitorInfo& Monitor = MonitorInfo[Index];
+		if (Monitor.bIsPrimary)
+		{
+			return Index;
+		}
+	}
+
+	// First monitor if there is one
+	return MonitorInfo.IsEmpty() ? -1 : 0;
+}
+
 float FDisplayMetrics::GetDebugTitleSafeZoneRatio()
 {
 	return GDebugSafeZoneRatio;
diff --git a/Engine/Source/Runtime/ApplicationCore/Private/Windows/WindowsApplication.cpp b/Engine/Source/Runtime/ApplicationCore/Private/Windows/WindowsApplication.cpp
--- a/Engine/Source/Runtime/ApplicationCore/Private/Windows/WindowsApplication.cpp
+++ b/Engine/Source/Runtime/ApplicationCore/Private/Windows/WindowsApplication.cpp
@@ -1209,7 +1209,6 @@ FWindowsApplication::FWindowsApplication( const HINSTANCE HInstance, const HICON
 		TEXT( "Slate.DeferWindowsMessageProcessing" ),
 		bAllowedToDeferMessageProcessing,
 		TEXT( "Whether windows message processing is deferred until tick or if they are processed immediately" ) )
-	, bInModalSizeLoop( false )
 #if WITH_ACCESSIBILITY && UE_WINDOWS_USING_UIA
 	, UIAManager(new FWindowsUIAManager(*this))
 #endif
@@ -1808,19 +1807,19 @@ FPlatformRect FWindowsApplication::GetWorkAreaFromOS( const FPlatformRect& Curre
  * @param OutHeight - Reference to output variable for monitor native height
  * @returns 'true' if data was extracted successfully, 'false' otherwise
  **/
-static bool GetMonitorSizeFromEDID(const HKEY hDevRegKey, int32& OutWidth, int32& OutHeight, int32& OutDPI)
+static bool GetMonitorSizeFromEDID(const HKEY hDevRegKey, int32& OutWidth, int32& OutHeight)
 {	
 	static const uint32 NameSize = 512;
 	static TCHAR ValueName[NameSize];
 
 	DWORD Type;
-	DWORD ActualValueNameLength = NameSize;
 
 	BYTE EDIDData[1024];
-	DWORD EDIDSize = sizeof(EDIDData);
 
 	for (LONG i = 0, RetValue = ERROR_SUCCESS; RetValue != ERROR_NO_MORE_ITEMS; ++i)
 	{
+		DWORD ActualValueNameLength = NameSize;
+		DWORD EDIDSize = sizeof(EDIDData);
 		RetValue = RegEnumValue ( hDevRegKey, 
 			i, 
 			&ValueName[0],
@@ -1828,7 +1827,7 @@ static bool GetMonitorSizeFromEDID(const HKEY hDevRegKey, int32& OutWidth, int32
 			EDIDData,
 			&EDIDSize);
 
-		if (RetValue != ERROR_SUCCESS || (FCString::Strcmp(ValueName, TEXT("EDID")) != 0))
+		if (RetValue != ERROR_SUCCESS || (FCString::Strcmp(ValueName, TEXT("EDID")) != 0) || (EDIDSize < 62))
 		{
 			continue;
 		}
@@ -1840,41 +1839,19 @@ static bool GetMonitorSizeFromEDID(const HKEY hDevRegKey, int32& OutWidth, int32
 		OutWidth = ((EDIDData[DetailTimingDescriptorStartIndex+4] >> 4) << 8) | EDIDData[DetailTimingDescriptorStartIndex+2];
 		OutHeight = ((EDIDData[DetailTimingDescriptorStartIndex+7] >> 4) << 8) | EDIDData[DetailTimingDescriptorStartIndex+5];
 
-		const int32 HorizontalSizeOffset = 21;
-		const int32 VerticalSizeOffset = 22;
-		const float CmToInch = 0.393701f;
-
-		if (EDIDData[HorizontalSizeOffset] > 0 && EDIDData[VerticalSizeOffset] > 0)
-		{
-			float PhysicalWidth = CmToInch * (float)EDIDData[HorizontalSizeOffset];
-			float PhysicalHeight = CmToInch * (float)EDIDData[VerticalSizeOffset];
-
-			int32 HDpi = FMath::TruncToInt((float)OutWidth / PhysicalWidth);
-			int32 VDpi = FMath::TruncToInt((float)OutHeight / PhysicalHeight);
-
-			OutDPI = (HDpi + VDpi) / 2;
-		}
-		else
-		{
-			OutDPI = 0;
-		}
-
 		return true; // valid EDID found
 	}
 
 	return false; // EDID not found
 }
 
-/**
- * Locate registry record for the given display device ID and extract native size information
- * @param TargetDevID - Name of taret device
- * @praam OutWidth - Reference to output variable for monitor native width
- * @praam OutHeight - Reference to output variable for monitor native height
- * @returns TRUE if data was extracted successfully, FALSE otherwise
- **/
-inline bool GetSizeForDevID(const FString& TargetDevID, int32& Width, int32& Height, int32& DPI)
+static void GetMonitorInfoSetupAPI(const FString& HardwareID, TOptional<FIntPoint>& NativeResolution, FString& FriendlyName)
 {
-	static const GUID ClassMonitorGuid = {0x4d36e96e, 0xe325, 0x11ce, {0xbf, 0xc1, 0x08, 0x00, 0x2b, 0xe1, 0x03, 0x18}};
+	static constexpr GUID ClassMonitorGuid = {0x4d36e96e, 0xe325, 0x11ce, {0xbf, 0xc1, 0x08, 0x00, 0x2b, 0xe1, 0x03, 0x18}};
+
+	// Could include <Devpkey.h> instead of copying here, but INITGUID behavior seems to diverge between editor and client, and would need a solve.
+	static constexpr DEVPROPKEY DEVPKEY_Device_HardwareIds = {{0xa45c254e, 0xdf1c, 0x4efd, {0x80, 0x20, 0x67, 0xd1, 0x46, 0xa8, 0x50, 0xe0}}, 3};
+	static constexpr DEVPROPKEY DEVPKEY_Device_FriendlyName = {{0xa45c254e, 0xdf1c, 0x4efd, {0x80, 0x20, 0x67, 0xd1, 0x46, 0xa8, 0x50, 0xe0}}, 14};
 
 	HDEVINFO DevInfo = SetupDiGetClassDevsEx(
 		&ClassMonitorGuid, //class GUID
@@ -1885,69 +1862,194 @@ inline bool GetSizeForDevID(const FString& TargetDevID, int32& Width, int32& Hei
 		NULL,
 		NULL);
 
-	if (NULL == DevInfo)
+	if (DevInfo != INVALID_HANDLE_VALUE)
 	{
-		return false;
-	}
-
-	bool bRes = false;
-
-	for (ULONG MonitorIndex = 0; ERROR_NO_MORE_ITEMS != GetLastError(); ++MonitorIndex)
-	{ 
-		SP_DEVINFO_DATA DevInfoData;
-		ZeroMemory(&DevInfoData, sizeof(DevInfoData));
+		bool bFound = false;
+		SP_DEVINFO_DATA DevInfoData{};
 		DevInfoData.cbSize = sizeof(DevInfoData);
-
-		if (SetupDiEnumDeviceInfo(DevInfo, MonitorIndex, &DevInfoData) == TRUE)
+		for (ULONG MonitorIndex = 0; SetupDiEnumDeviceInfo(DevInfo, MonitorIndex, &DevInfoData); ++MonitorIndex)
 		{
-			TCHAR Buffer[MAX_DEVICE_ID_LEN];
-			if (CM_Get_Device_ID(DevInfoData.DevInst, Buffer, MAX_DEVICE_ID_LEN, 0) == CR_SUCCESS)
+			DEVPROPTYPE PropertyType;
+			DWORD RequiredSize;
+			if (!SetupDiGetDeviceProperty(DevInfo, &DevInfoData, &DEVPKEY_Device_HardwareIds, &PropertyType, nullptr, 0, &RequiredSize, 0) && (GetLastError() == ERROR_INSUFFICIENT_BUFFER))
 			{
-				FString DevID(Buffer);
-				DevID.MidInline(8, DevID.Find(TEXT("\\"), ESearchCase::CaseSensitive, ESearchDir::FromStart, 9) - 8, EAllowShrinking::No);
-				if (DevID == TargetDevID)
+				TArray<WCHAR> HardwareIDsBuffer;
+				HardwareIDsBuffer.SetNum((RequiredSize + 1) / 2);
+				if (SetupDiGetDeviceProperty(DevInfo, &DevInfoData, &DEVPKEY_Device_HardwareIds, &PropertyType, (PBYTE)HardwareIDsBuffer.GetData(), RequiredSize, nullptr, 0))
 				{
-					HKEY hDevRegKey = SetupDiOpenDevRegKey(DevInfo, &DevInfoData, DICS_FLAG_GLOBAL, 0, DIREG_DEV, KEY_READ);
+					int32 Cursor = 0;
+					while ((Cursor < HardwareIDsBuffer.Num()) && HardwareIDsBuffer[Cursor])
+					{
+						if (HardwareID == &HardwareIDsBuffer[Cursor])
+						{
+							bFound = true;
+							break;
+						}
+
+						// Move cursor to null terminator
+						do
+						{
+							++Cursor;
+						}
+						while ((Cursor < HardwareIDsBuffer.Num()) && HardwareIDsBuffer[Cursor]);
 
-					if (hDevRegKey && hDevRegKey != INVALID_HANDLE_VALUE)
+						// Move cursor to next string
+						++Cursor;
+					}
+
+					if (bFound)
 					{
-						bRes = GetMonitorSizeFromEDID(hDevRegKey, Width, Height, DPI);
-						RegCloseKey(hDevRegKey);
+						HKEY hDevRegKey = SetupDiOpenDevRegKey(DevInfo, &DevInfoData, DICS_FLAG_GLOBAL, 0, DIREG_DEV, KEY_READ);
+						if (hDevRegKey && hDevRegKey != INVALID_HANDLE_VALUE)
+						{
+							int32 NativeWidth, NativeHeight;
+							if (GetMonitorSizeFromEDID(hDevRegKey, NativeWidth, NativeHeight))
+							{
+								NativeResolution = FIntPoint(NativeWidth, NativeHeight);
+							}
+
+							if (!SetupDiGetDeviceProperty(DevInfo, &DevInfoData, &DEVPKEY_Device_FriendlyName, &PropertyType, nullptr, 0, &RequiredSize, 0) && (GetLastError() == ERROR_INSUFFICIENT_BUFFER))
+							{
+								TArray<WCHAR> FriendlyNameBuffer;
+								FriendlyNameBuffer.SetNum((RequiredSize + 1) / 2);
+								if (SetupDiGetDeviceProperty(DevInfo, &DevInfoData, &DEVPKEY_Device_FriendlyName, &PropertyType, (PBYTE)FriendlyNameBuffer.GetData(), RequiredSize, nullptr, 0))
+								{
+									FriendlyName = FriendlyNameBuffer.GetData();
+								}
+							}
+
+							RegCloseKey(hDevRegKey);
+						}
+
 						break;
 					}
 				}
 			}
+
+			// Just in case SetupDiEnumDeviceInfo could modify this value.
+			DevInfoData.cbSize = sizeof(DevInfoData);
+		}
+		ensure(bFound || (GetLastError() == ERROR_NO_MORE_ITEMS));
+
+		SetupDiDestroyDeviceInfoList(DevInfo);
+	}
+}
+
+static void GetFriendlyNameDisplayDevices(const WCHAR DeviceName[CCHDEVICENAME], FString& FriendlyName)
+{
+	UINT32 PathCount;
+	TArray<DISPLAYCONFIG_PATH_INFO> Paths;
+	LONG Result;
+
+	do
+	{
+		constexpr UINT32 Flags = QDC_ONLY_ACTIVE_PATHS | QDC_VIRTUAL_MODE_AWARE;
+		UINT32 ModeCount;
+		Result = GetDisplayConfigBufferSizes(Flags, &PathCount, &ModeCount);
+		if (Result != ERROR_SUCCESS)
+		{
+			return;
+		}
+
+		TArray<DISPLAYCONFIG_MODE_INFO> Modes;
+		Paths.SetNum(PathCount);
+		Modes.SetNum(ModeCount);
+		Result = QueryDisplayConfig(Flags, &PathCount, Paths.GetData(), &ModeCount, Modes.GetData(), NULL);
+	} while (Result == ERROR_INSUFFICIENT_BUFFER);
+
+	if (Result == ERROR_SUCCESS)
+	{
+		for (size_t PathIndex = 0; PathIndex < PathCount; ++PathIndex)
+		{
+			const DISPLAYCONFIG_PATH_INFO& Path = Paths[PathIndex];
+			DISPLAYCONFIG_SOURCE_DEVICE_NAME SourceName;
+			SourceName.header.type = DISPLAYCONFIG_DEVICE_INFO_GET_SOURCE_NAME;
+			SourceName.header.size = sizeof(SourceName);
+			SourceName.header.adapterId = Path.sourceInfo.adapterId;
+			SourceName.header.id = Path.sourceInfo.id;
+			if ((DisplayConfigGetDeviceInfo(&SourceName.header) == ERROR_SUCCESS) &&
+				(FCString::Strcmp(DeviceName, SourceName.viewGdiDeviceName) == 0))
+			{
+				DISPLAYCONFIG_TARGET_DEVICE_NAME TargetName;
+				TargetName.header.type = DISPLAYCONFIG_DEVICE_INFO_GET_TARGET_NAME;
+				TargetName.header.size = sizeof(TargetName);
+				TargetName.header.adapterId = Path.targetInfo.adapterId;
+				TargetName.header.id = Path.targetInfo.id;
+				if (DisplayConfigGetDeviceInfo(&TargetName.header) == ERROR_SUCCESS)
+				{
+					FriendlyName = TargetName.monitorFriendlyDeviceName;
+				}
+				break;
+			}
 		}
 	}
+}
 
-	if (SetupDiDestroyDeviceInfoList(DevInfo) == FALSE)
+static bool GetMonitorDetails(HMONITOR Monitor, const FString& DeviceID, FString& FriendlyName, int32& NativeWidth, int32& NativeHeight, int32& DPI)
+{
+	MONITORINFOEX MonitorInfoEx;
+	MonitorInfoEx.cbSize = sizeof(MonitorInfoEx);
+	if (GetMonitorInfo(Monitor, &MonitorInfoEx))
 	{
-		bRes = false;
+		TOptional<FIntPoint> NativeResolution;
+		FString FriendlyNameFallback;
+		const int32 BackslashIndex = DeviceID.Find(TEXT("\\"), ESearchCase::CaseSensitive, ESearchDir::FromStart, 9);
+		GetMonitorInfoSetupAPI((BackslashIndex != INDEX_NONE) ? DeviceID.Left(BackslashIndex) : DeviceID, NativeResolution, FriendlyNameFallback);
+
+		if (NativeResolution.IsSet())
+		{
+			NativeWidth = NativeResolution->X;
+			NativeHeight = NativeResolution->Y;
+		}
+		else
+		{
+			NativeWidth = MonitorInfoEx.rcMonitor.right - MonitorInfoEx.rcMonitor.left;
+			NativeHeight = MonitorInfoEx.rcMonitor.bottom - MonitorInfoEx.rcMonitor.top;
+		}
+
+		GetFriendlyNameDisplayDevices(MonitorInfoEx.szDevice, FriendlyName);
+		if (FriendlyName.IsEmpty())
+		{
+			FriendlyName = FriendlyNameFallback;
+		}
+
+		DPI = FWindowsPlatformApplicationMisc::GetMonitorDPI(Monitor);
+
+		return true;
 	}
 
-	return bRes;
+	return false;
 }
 
+struct MonitorUserData
+{
+	FMonitorInfo* Info = nullptr;
+	HMONITOR NativeHandle = NULL;
+};
 static BOOL CALLBACK MonitorEnumProc(HMONITOR Monitor, HDC MonitorDC, LPRECT Rect, LPARAM UserData)
 {
 	MONITORINFOEX MonitorInfoEx;
 	MonitorInfoEx.cbSize = sizeof(MonitorInfoEx);	
-	GetMonitorInfo(Monitor, &MonitorInfoEx);
-
-	FMonitorInfo* Info = (FMonitorInfo*)UserData;
-	if (Info->Name == MonitorInfoEx.szDevice)
+	if (GetMonitorInfo(Monitor, &MonitorInfoEx))
+	{
+		MonitorUserData* Data = (MonitorUserData*)UserData;
+		FMonitorInfo* Info = Data->Info;
+		if (Info->Name == MonitorInfoEx.szDevice)
 		{
-		Info->DisplayRect.Bottom = MonitorInfoEx.rcMonitor.bottom;
-		Info->DisplayRect.Left = MonitorInfoEx.rcMonitor.left;
-		Info->DisplayRect.Right = MonitorInfoEx.rcMonitor.right;
-		Info->DisplayRect.Top = MonitorInfoEx.rcMonitor.top;
+			Info->DisplayRect.Bottom = MonitorInfoEx.rcMonitor.bottom;
+			Info->DisplayRect.Left = MonitorInfoEx.rcMonitor.left;
+			Info->DisplayRect.Right = MonitorInfoEx.rcMonitor.right;
+			Info->DisplayRect.Top = MonitorInfoEx.rcMonitor.top;
 
-		Info->WorkArea.Bottom = MonitorInfoEx.rcWork.bottom;
-		Info->WorkArea.Left = MonitorInfoEx.rcWork.left;
-		Info->WorkArea.Right = MonitorInfoEx.rcWork.right;
-		Info->WorkArea.Top = MonitorInfoEx.rcWork.top;
+			Info->WorkArea.Bottom = MonitorInfoEx.rcWork.bottom;
+			Info->WorkArea.Left = MonitorInfoEx.rcWork.left;
+			Info->WorkArea.Right = MonitorInfoEx.rcWork.right;
+			Info->WorkArea.Top = MonitorInfoEx.rcWork.top;
 
-		return FALSE;
+			Data->NativeHandle = Monitor;
+
+			return FALSE;
+		}
 	}
 
 	return TRUE;
@@ -2004,14 +2106,17 @@ static void GetMonitorsInfo(TArray<FMonitorInfo>& OutMonitorInfo)
 					!(Monitor.StateFlags & DISPLAY_DEVICE_MIRRORING_DRIVER))
 				{
 					FMonitorInfo Info;
-
 					Info.Name = DisplayDevice.DeviceName;
-					EnumDisplayMonitors(nullptr, nullptr, MonitorEnumProc, (LPARAM)&Info);
 
-					Info.ID = FString::Printf(TEXT("%s"), Monitor.DeviceID);
-					Info.Name = Info.ID.Mid (8, Info.ID.Find (TEXT("\\"), ESearchCase::CaseSensitive, ESearchDir::FromStart, 9) - 8);
+					MonitorUserData Data{};
+					Data.Info = &Info;
+					EnumDisplayMonitors(nullptr, nullptr, MonitorEnumProc, (LPARAM)&Data);
+
+					Info.ID = Monitor.DeviceID;
+					const int32 BackslashIndex = Info.ID.Find(TEXT("\\"), ESearchCase::CaseSensitive, ESearchDir::FromStart, 9);
+					Info.Name = (BackslashIndex != INDEX_NONE) ? Info.ID.Mid(8, BackslashIndex - 8) : Monitor.DeviceID;
 
-					if (GetSizeForDevID(Info.Name, Info.NativeWidth, Info.NativeHeight, Info.DPI))
+					if (GetMonitorDetails(Data.NativeHandle, Info.ID, Info.FriendlyName, Info.NativeWidth, Info.NativeHeight, Info.DPI))
 					{
 						// usually expected: 4:3=1.333, 16:9=1.777, 16:10=1.6, multi-monitor-wide: >2
 						bool bIsPortrait = Info.NativeWidth < Info.NativeHeight;
@@ -2021,30 +2126,16 @@ static void GetMonitorsInfo(TArray<FMonitorInfo>& OutMonitorInfo)
 							Swap(Info.NativeWidth, Info.NativeHeight);
 						}
 
-						Info.ID = Monitor.DeviceID;
 						Info.bIsPrimary = (DisplayDevice.StateFlags & DISPLAY_DEVICE_PRIMARY_DEVICE) > 0;
 
-						// sanity check for DPI values
-						if (Info.DPI < 96 || Info.DPI > 300)
-						{
-							// switch to default winapi value
-							Info.DPI = FWindowsPlatformApplicationMisc::GetMonitorDPI(Info);
-						}
-						else
-						{
-							// we also need to include the OS scaling value
-							const float CenterX = 0.5f * float(Info.WorkArea.Right + Info.WorkArea.Left);
-							const float CenterY = 0.5f * float(Info.WorkArea.Top + Info.WorkArea.Bottom);
-							const float DPIScaleFactor = FWindowsPlatformApplicationMisc::GetDPIScaleFactorAtPoint(CenterX, CenterY);
-							Info.DPI = (int32)((float)Info.DPI * DPIScaleFactor);
-						}
-
 						// The editor shouldn't care about lower level display resolutions. This is only necessary for fullscreen exclusive situations.
 						if (!GIsEditor)
 						{
 							Info.MaxResolution = GetMaxResolutionForDisplay(DisplayDevice, Info.NativeWidth, Info.NativeHeight);
 						}
 
+						Info.NativeHandle = Data.NativeHandle;
+
 						OutMonitorInfo.Add(Info);
 					}
 				}
@@ -2585,6 +2676,10 @@ int32 FWindowsApplication::ProcessMessage( HWND hwnd, uint32 msg, WPARAM wParam,
 				{
 					MessageHandler->OnWindowAction(CurrentNativeEventWindow, bWasMaximized ? EWindowAction::Maximize : EWindowAction::Restore);
 				}
+				if (CurrentNativeEventWindowPtr->IsInModalSizeLoop())
+				{
+					MessageHandler->OnRegularWindowMovedOrResized(CurrentNativeEventWindowPtr.ToSharedRef());
+				}
 
 				return 0;
 			}
@@ -2735,13 +2830,13 @@ int32 FWindowsApplication::ProcessMessage( HWND hwnd, uint32 msg, WPARAM wParam,
 			break;
 		case WM_ENTERSIZEMOVE:
 			{
-				bInModalSizeLoop = true;
+				CurrentNativeEventWindowPtr->EnterModalSizeLoop();
 				DeferMessage( CurrentNativeEventWindowPtr, hwnd, msg, wParam, lParam, 0, 0 );
 			}
 			break;
 		case WM_EXITSIZEMOVE:
 			{
-				bInModalSizeLoop = false;
+				CurrentNativeEventWindowPtr->ExitModalSizeLoop();
 				DeferMessage( CurrentNativeEventWindowPtr, hwnd, msg, wParam, lParam, 0, 0 );
 			}
 			break;
@@ -2758,9 +2853,13 @@ int32 FWindowsApplication::ProcessMessage( HWND hwnd, uint32 msg, WPARAM wParam,
 				if ( FWindowsApplication::MinimizedWindowPosition != NewPosition )
 				{
 					MessageHandler->OnMovedWindow( CurrentNativeEventWindow, NewX, NewY );
-
-					return 0;
 				}
+				if (CurrentNativeEventWindowPtr->IsInModalSizeLoop())
+				{
+					MessageHandler->OnRegularWindowMovedOrResized(CurrentNativeEventWindow);
+				}
+
+				return 0;
 			}
 			break;
 
@@ -2823,8 +2922,11 @@ int32 FWindowsApplication::ProcessMessage( HWND hwnd, uint32 msg, WPARAM wParam,
 		case WM_ACTIVATEAPP:
 			{
 				// When window activation changes we are not in a modal size loop
-				bInModalSizeLoop = false;
+				// NOTE: WM_EXITSIZEMOVE might be reliable these days, but keep this around for now.
+				CurrentNativeEventWindowPtr->ExitModalSizeLoop();
+
 				DeferMessage( CurrentNativeEventWindowPtr, hwnd, msg, wParam, lParam );
+				return 0;
 			}
 			break;
 
@@ -2844,9 +2946,10 @@ int32 FWindowsApplication::ProcessMessage( HWND hwnd, uint32 msg, WPARAM wParam,
 
 		case WM_PAINT:
 			{
-				if( bInModalSizeLoop && IsInGameThread() )
+				check(IsInGameThread());
+				if (CurrentNativeEventWindowPtr->IsInModalSizeLoop())
 				{
-					MessageHandler->OnOSPaint(CurrentNativeEventWindowPtr.ToSharedRef() );
+					MessageHandler->OnOSPaint(CurrentNativeEventWindowPtr.ToSharedRef());
 				}
 			}
 			break;
@@ -3937,8 +4040,11 @@ int32 FWindowsApplication::ProcessDeferredMessage( const FDeferredWindowsMessage
 					{
 						CurrentNativeEventWindowPtr->SetDPIScaleFactor(LOWORD(wParam) / static_cast<float>(USER_DEFAULT_SCREEN_DPI));
 
-						LPRECT NewRect = (LPRECT)lParam;
-						SetWindowPos(hwnd, nullptr, NewRect->left, NewRect->top, NewRect->right - NewRect->left, NewRect->bottom - NewRect->top, SWP_NOZORDER | SWP_NOACTIVATE);
+						if (CurrentNativeEventWindowPtr->IsInModalSizeLoop())
+						{
+							LPRECT NewRect = (LPRECT)lParam;
+							SetWindowPos(hwnd, nullptr, NewRect->left, NewRect->top, NewRect->right - NewRect->left, NewRect->bottom - NewRect->top, SWP_NOZORDER | SWP_NOACTIVATE);
+						}
 
 						MessageHandler->HandleDPIScaleChanged(CurrentNativeEventWindowPtr.ToSharedRef());
 					}
diff --git a/Engine/Source/Runtime/ApplicationCore/Private/Windows/WindowsPlatformApplicationMisc.cpp b/Engine/Source/Runtime/ApplicationCore/Private/Windows/WindowsPlatformApplicationMisc.cpp
--- a/Engine/Source/Runtime/ApplicationCore/Private/Windows/WindowsPlatformApplicationMisc.cpp
+++ b/Engine/Source/Runtime/ApplicationCore/Private/Windows/WindowsPlatformApplicationMisc.cpp
@@ -355,12 +355,25 @@ int32 FWindowsPlatformApplicationMisc::GetMonitorDPI(const FMonitorInfo& Monitor
 		HMONITOR Monitor = MonitorFromRect(&MonitorDim, MONITOR_DEFAULTTONEAREST);
 		if (Monitor)
 		{
-			uint32 DPIX = 0;
-			uint32 DPIY = 0;
-			if (SUCCEEDED(GetDpiForMonitor(Monitor, MDT_EFFECTIVE_DPI, &DPIX, &DPIY)))
-			{
-				DisplayDPI = DPIX;
-			}
+			DisplayDPI = GetMonitorDPI(Monitor);
+		}
+	}
+
+	return DisplayDPI;
+}
+
+int32 FWindowsPlatformApplicationMisc::GetMonitorDPI(void* NativeHandle)
+{
+	int32 DisplayDPI = USER_DEFAULT_SCREEN_DPI;
+
+	if (IsHighDPIAwarenessEnabled())
+	{
+		HMONITOR Monitor = static_cast<HMONITOR>(NativeHandle);
+		uint32 DPIX = 0;
+		uint32 DPIY = 0;
+		if (SUCCEEDED(GetDpiForMonitor(Monitor, MDT_EFFECTIVE_DPI, &DPIX, &DPIY)))
+		{
+			DisplayDPI = DPIX;
 		}
 	}
 
diff --git a/Engine/Source/Runtime/ApplicationCore/Public/GenericPlatform/GenericApplication.h b/Engine/Source/Runtime/ApplicationCore/Public/GenericPlatform/GenericApplication.h
--- a/Engine/Source/Runtime/ApplicationCore/Public/GenericPlatform/GenericApplication.h
+++ b/Engine/Source/Runtime/ApplicationCore/Public/GenericPlatform/GenericApplication.h
@@ -338,8 +338,15 @@ struct FPlatformRect
  */
 struct FMonitorInfo
 {
+	/** Legacy name that may not be as friendly as FriendlyName for surfacing to UI. */
 	FString Name;
+
+	/** Friendly name for surfacing to UI. */
+	FString FriendlyName;
+
+	/** String used for monitor identification, which is less likely to alias displays than Name or FriendlyName. */
 	FString ID;
+
 	int32 NativeWidth = 0;
 	int32 NativeHeight = 0;
 	FIntPoint MaxResolution = FIntPoint(ForceInitToZero);
@@ -347,6 +354,7 @@ struct FMonitorInfo
 	FPlatformRect WorkArea;
 	bool bIsPrimary = false;
 	int32 DPI = 0;
+	void* NativeHandle = nullptr;
 };
 
 
@@ -396,6 +404,9 @@ struct FDisplayMetrics
 	/** Gets the monitor work area from a position in the global display rect */
 	APPLICATIONCORE_API FPlatformRect GetMonitorWorkAreaFromPoint(const FVector2D& Point) const;
 
+	/** Gets the closest monitor from monitor id and index favoring id */
+	APPLICATIONCORE_API int32 GetClosestMonitorFromIDAndIndex(const FString& MonitorID, int32 MonitorIndex) const;
+
 	/** Logs out display metrics */
 	APPLICATIONCORE_API void PrintToLog() const;
 
diff --git a/Engine/Source/Runtime/ApplicationCore/Public/GenericPlatform/GenericApplicationMessageHandler.h b/Engine/Source/Runtime/ApplicationCore/Public/GenericPlatform/GenericApplicationMessageHandler.h
--- a/Engine/Source/Runtime/ApplicationCore/Public/GenericPlatform/GenericApplicationMessageHandler.h
+++ b/Engine/Source/Runtime/ApplicationCore/Public/GenericPlatform/GenericApplicationMessageHandler.h
@@ -425,7 +425,12 @@ public:
 	virtual void SignalSystemDPIChanged(const TSharedRef< FGenericWindow >& Window)
 	{
 
-	}	
+	}
+
+	virtual void OnRegularWindowMovedOrResized(const TSharedRef< FGenericWindow >& Window)
+	{
+
+	}
 
 	virtual void OnMovedWindow( const TSharedRef< FGenericWindow >& Window, const int32 X, const int32 Y )
 	{
diff --git a/Engine/Source/Runtime/ApplicationCore/Public/Windows/WindowsApplication.h b/Engine/Source/Runtime/ApplicationCore/Public/Windows/WindowsApplication.h
--- a/Engine/Source/Runtime/ApplicationCore/Public/Windows/WindowsApplication.h
+++ b/Engine/Source/Runtime/ApplicationCore/Public/Windows/WindowsApplication.h
@@ -591,9 +591,6 @@ private:
 	int32 bAllowedToDeferMessageProcessing;
 	
 	FAutoConsoleVariableRef CVarDeferMessageProcessing;
-	
-	/** True if we are in the middle of a windows modal size loop */
-	bool bInModalSizeLoop;
 
 	FDisplayMetrics InitialDisplayMetrics;
 
diff --git a/Engine/Source/Runtime/ApplicationCore/Public/Windows/WindowsPlatformApplicationMisc.h b/Engine/Source/Runtime/ApplicationCore/Public/Windows/WindowsPlatformApplicationMisc.h
--- a/Engine/Source/Runtime/ApplicationCore/Public/Windows/WindowsPlatformApplicationMisc.h
+++ b/Engine/Source/Runtime/ApplicationCore/Public/Windows/WindowsPlatformApplicationMisc.h
@@ -27,6 +27,7 @@ struct FWindowsPlatformApplicationMisc : public FGenericPlatformApplicationMisc
 	/** Windows platform only */
 	/** Function should retrieve the DPI value for the provided monitor information structure */
 	static APPLICATIONCORE_API int32 GetMonitorDPI(const FMonitorInfo& MonitorInfo);
+	static APPLICATIONCORE_API int32 GetMonitorDPI(void* NativeHandle);
 
 	struct FGPUInfo
 	{
diff --git a/Engine/Source/Runtime/ApplicationCore/Public/Windows/WindowsWindow.h b/Engine/Source/Runtime/ApplicationCore/Public/Windows/WindowsWindow.h
--- a/Engine/Source/Runtime/ApplicationCore/Public/Windows/WindowsWindow.h
+++ b/Engine/Source/Runtime/ApplicationCore/Public/Windows/WindowsWindow.h
@@ -80,6 +80,10 @@ public:
 	/** call with a true argument if this window need to do its custom size management in response to DPI variations */
 	APPLICATIONCORE_API virtual void SetManualManageDPIChanges(const bool bManualDPIChanges) override;
 
+	bool IsInModalSizeLoop() const { return bInModalSizeLoop; }
+	void EnterModalSizeLoop() { bInModalSizeLoop = true; }
+	void ExitModalSizeLoop() { bInModalSizeLoop = false; }
+
 	/** Called when our parent window is minimized (which will in turn cause us to become minimized). */
 	APPLICATIONCORE_API void OnParentWindowMinimized();
 
@@ -213,4 +217,7 @@ private:
 
 	/** when true the window is responsible for its own size adjustments in response to a DPI change */
 	bool bHandleManualDPIChanges = false;
+
+	/** when true the window in the middle of a windows modal size loop */
+	bool bInModalSizeLoop = false;
 };
diff --git a/Engine/Source/Runtime/D3D12RHI/Private/Windows/WindowsD3D12Device.cpp b/Engine/Source/Runtime/D3D12RHI/Private/Windows/WindowsD3D12Device.cpp
--- a/Engine/Source/Runtime/D3D12RHI/Private/Windows/WindowsD3D12Device.cpp
+++ b/Engine/Source/Runtime/D3D12RHI/Private/Windows/WindowsD3D12Device.cpp
@@ -2013,7 +2013,7 @@ bool FD3D12Device::GetRayTracingPipelineInfo(ID3D12StateObject* Pipeline, FD3D12
 
 #endif // D3D12_RHI_RAYTRACING
 
-static bool GetAvailableResolutionsForOutput(IDXGIOutput* Output, FScreenResolutionArray& Resolutions, bool bIgnoreRefreshRate)
+static void GetAvailableResolutionsForOutput(IDXGIOutput* Output, FScreenResolutionArray& Resolutions, bool bIgnoreRefreshRate)
 {
 	// TODO: GetDisplayModeList is a terribly SLOW call.  It can take up to a second per invocation.
 	//  We might want to work around some DXGI badness here.
@@ -2022,7 +2022,7 @@ static bool GetAvailableResolutionsForOutput(IDXGIOutput* Output, FScreenResolut
 	HRESULT HResult = Output->GetDisplayModeList(Format, 0, &NumModes, nullptr);
 	if (HResult == DXGI_ERROR_NOT_FOUND)
 	{
-		return false;
+		return;
 	}
 
 	if (HResult == DXGI_ERROR_NOT_CURRENTLY_AVAILABLE)
@@ -2030,13 +2030,13 @@ static bool GetAvailableResolutionsForOutput(IDXGIOutput* Output, FScreenResolut
 		UE_LOG(LogD3D12RHI, Warning,
 			TEXT("RHIGetAvailableResolutions() can not be used over a remote desktop configuration")
 		);
-		return false;
+		return;
 	}
 
 	if (FAILED(HResult))
 	{
 		UE_LOG(LogD3D12RHI, Warning, TEXT("RHIGetAvailableResolutions failed with unknown error (0x%x)."), HResult);
-		return false;
+		return;
 	}
 
 	if (NumModes > 0)
@@ -2105,8 +2105,6 @@ static bool GetAvailableResolutionsForOutput(IDXGIOutput* Output, FScreenResolut
 	{
 		UE_LOG(LogD3D12RHI, Warning, TEXT("No display modes found for the standard format DXGI_FORMAT_R8G8B8A8_UNORM!"));
 	}
-
-	return !Resolutions.IsEmpty();
 }
 
 /**
@@ -2145,7 +2143,11 @@ bool FD3D12DynamicRHI::RHIGetAvailableResolutions(FScreenResolutionArray& Resolu
 
 	TRefCountPtr<IDXGIOutput> Output;
 	HResult = Adapter->EnumOutputs(0, Output.GetInitReference());
-	return SUCCEEDED(HResult) ? GetAvailableResolutionsForOutput(Output, Resolutions, bIgnoreRefreshRate) : false;
+	if (SUCCEEDED(HResult))
+	{
+		GetAvailableResolutionsForOutput(Output, Resolutions, bIgnoreRefreshRate);
+	}
+	return !Resolutions.IsEmpty();
 }
 
 bool FD3D12DynamicRHI::RHIGetAvailableResolutionsForDisplay(FScreenResolutionArray& Resolutions, bool bIgnoreRefreshRate, void* NativeMonitorHandle)
@@ -2183,7 +2185,33 @@ bool FD3D12DynamicRHI::RHIGetAvailableResolutionsForDisplay(FScreenResolutionArr
 			check(SUCCEEDED(OutputDescRes));
 			if (MonitorHandle == OutputDesc.Monitor)
 			{
-				return GetAvailableResolutionsForOutput(Output, Resolutions, bIgnoreRefreshRate);
+				GetAvailableResolutionsForOutput(Output, Resolutions, bIgnoreRefreshRate);
+
+				// Only allow the desktop resolution if SetFullscreenState is expected to fail to allow borderless to seemlessly take over.
+				// Undocumented but assume exclusive only works if the the display is primary, i.e. top-left at (0, 0), or the display is connected to the chosen adapter.
+				if (OutputDesc.DesktopCoordinates.left != 0 || OutputDesc.DesktopCoordinates.top != 0)
+				{
+					const FD3D12Adapter* const ChosenAdapter = ChosenAdapters[0].Get();
+					if (ChosenAdapter)
+					{
+						const FD3D12AdapterDesc& ChosenDesc = ChosenAdapter->GetDesc();
+						if (ChosenDesc.IsValid())
+						{
+							DXGI_ADAPTER_DESC AdapterDesc;
+							if (SUCCEEDED(Adapter->GetDesc(&AdapterDesc)) &&
+								FMemory::Memcmp(&ChosenAdapter->GetDesc().Desc.AdapterLuid, &AdapterDesc.AdapterLuid, sizeof(LUID)))
+							{
+								Resolutions.Reset(1);
+								FScreenResolutionRHI& ScreenResolution = Resolutions.Emplace_GetRef();
+								ScreenResolution.Width = OutputDesc.DesktopCoordinates.right - OutputDesc.DesktopCoordinates.left;
+								ScreenResolution.Height = OutputDesc.DesktopCoordinates.bottom - OutputDesc.DesktopCoordinates.top;
+								ScreenResolution.RefreshRate = 0;
+							}
+						}
+					}
+				}
+
+				return !Resolutions.IsEmpty();
 			}
 		}
 	}
diff --git a/Engine/Source/Runtime/Engine/Classes/Engine/GameEngine.h b/Engine/Source/Runtime/Engine/Classes/Engine/GameEngine.h
--- a/Engine/Source/Runtime/Engine/Classes/Engine/GameEngine.h
+++ b/Engine/Source/Runtime/Engine/Classes/Engine/GameEngine.h
@@ -83,8 +83,9 @@ public:
 	 * @param ResolutionX	[in/out] Width of the game window, in pixels
 	 * @param ResolutionY	[in/out] Height of the game window, in pixels
 	 * @param WindowMode	[in/out] What window mode the game should be in
+	 * @param MonitorIndex	[in] Index of the entry in FDisplayMetrics::MonitorInfo
 	 */
-	static ENGINE_API void ConditionallyOverrideSettings( int32& ResolutionX, int32& ResolutionY, EWindowMode::Type& WindowMode );
+	static ENGINE_API void ConditionallyOverrideSettings( int32& ResolutionX, int32& ResolutionY, EWindowMode::Type& WindowMode, int32 MonitorIndex = -1 );
 	
 	/**
 	 * Determines the resolution of the game window, ensuring that the requested size is never bigger than the available desktop size
@@ -93,8 +94,9 @@ public:
 	 * @param ResolutionY	[in/out] Height of the game window, in pixels
 	 * @param WindowMode	[in/out] What window mode the game should be in
 	 * @param bUseWorkArea	[in] Should we find a resolution that fits within the desktop work area for the windowed mode instead of monitor's full resolution
+	 * @param MonitorIndex	[in] Index of the entry in FDisplayMetrics::MonitorInfo
 	 */
-	static ENGINE_API void DetermineGameWindowResolution( int32& ResolutionX, int32& ResolutionY, EWindowMode::Type& WindowMode, bool bUseWorkAreaForWindowed = false );
+	static ENGINE_API void DetermineGameWindowResolution( int32& ResolutionX, int32& ResolutionY, EWindowMode::Type& WindowMode, bool bUseWorkAreaForWindowed = false, int32 MonitorIndex = -1 );
 
 	/**
 	 * Changes the game window to use the game viewport instead of any loading screen
diff --git a/Engine/Source/Runtime/Engine/Classes/Engine/GameViewportClient.h b/Engine/Source/Runtime/Engine/Classes/Engine/GameViewportClient.h
--- a/Engine/Source/Runtime/Engine/Classes/Engine/GameViewportClient.h
+++ b/Engine/Source/Runtime/Engine/Classes/Engine/GameViewportClient.h
@@ -846,6 +846,12 @@ public:
 	}
 #endif
 
+	/** Accessor for delegate called when the window moves to another display */
+	FOnWindowDisplayChanged& OnWindowDisplayChanged()
+	{
+		return WindowDisplayChangedDelegate;
+	}
+
 	/** Accessor for delegate called when the engine toggles fullscreen */
 	FOnToggleFullscreen& OnToggleFullscreen()
 	{
@@ -934,6 +940,9 @@ private:
 	/** Delegate handler for when all stats are disabled in a viewport */
 	ENGINE_API void HandleViewportStatDisableAll(const bool bInAnyViewport);
 
+	/** Delegate handler for when a window is moved or resized */
+	ENGINE_API void HandleRegularWindowMovedOrResized(TSharedRef<SWindow> InWindow);
+
 	/** Delegate handler for when a window DPI changes and we might need to adjust the scenes resolution */
 	ENGINE_API void HandleWindowDPIScaleChanged(TSharedRef<SWindow> InWindow);
 
@@ -1060,6 +1069,9 @@ private:
 	/** Delegate called when ticking the game viewport */
 	FOnGameViewportTick TickDelegate;
 
+	/** Delegate called when the window moves to another display */
+	FOnWindowDisplayChanged WindowDisplayChangedDelegate;
+
 	/** Delegate called when the engine toggles fullscreen */
 	FOnToggleFullscreen ToggleFullscreenDelegate;
 
diff --git a/Engine/Source/Runtime/Engine/Classes/Engine/GameViewportDelegates.h b/Engine/Source/Runtime/Engine/Classes/Engine/GameViewportDelegates.h
--- a/Engine/Source/Runtime/Engine/Classes/Engine/GameViewportDelegates.h
+++ b/Engine/Source/Runtime/Engine/Classes/Engine/GameViewportDelegates.h
@@ -68,8 +68,15 @@ DECLARE_MULTICAST_DELEGATE_OneParam(FOnGameViewportTick, float);
 DECLARE_DELEGATE_RetVal_ThreeParams(bool, FOnGameViewportInputKey, FKey, FModifierKeysState, EInputEvent);
 
 /**
-* Delegate type used by UGameViewportClient for when engine in toggling fullscreen
-*
-* @see UGameViewportClient
-*/
+ * Delegate type used by UGameViewportClient for when the window moves to another display
+ *
+ * @see UGameViewportClient
+ */
+DECLARE_MULTICAST_DELEGATE(FOnWindowDisplayChanged);
+
+/**
+ * Delegate type used by UGameViewportClient for when engine in toggling fullscreen
+ *
+ * @see UGameViewportClient
+ */
 DECLARE_MULTICAST_DELEGATE_OneParam(FOnToggleFullscreen, bool);
diff --git a/Engine/Source/Runtime/Engine/Classes/GameFramework/GameUserSettings.h b/Engine/Source/Runtime/Engine/Classes/GameFramework/GameUserSettings.h
--- a/Engine/Source/Runtime/Engine/Classes/GameFramework/GameUserSettings.h
+++ b/Engine/Source/Runtime/Engine/Classes/GameFramework/GameUserSettings.h
@@ -86,6 +86,30 @@ public:
 	UFUNCTION(BlueprintPure, Category = Settings)
 	ENGINE_API EWindowMode::Type GetPreferredFullscreenMode() const;
 
+	/** Returns the user setting for game window display ID. See UGameUserSettings::DisplayID. */
+	UFUNCTION(BlueprintPure, Category=Settings)
+	ENGINE_API FString GetDisplayID() const;
+
+	/** Returns the last confirmed user setting for game window display ID. */
+	UFUNCTION(BlueprintPure, Category=Settings)
+	ENGINE_API FString GetLastConfirmedDisplayID() const;
+
+	/** Returns the user setting for game window display index. See UGameUserSettings::DisplayIndex. */
+	UFUNCTION(BlueprintPure, Category=Settings)
+	ENGINE_API int32 GetDisplayIndex() const;
+
+	/** Returns the last confirmed user setting for game window display index. */
+	UFUNCTION(BlueprintPure, Category=Settings)
+	ENGINE_API int32 GetLastConfirmedDisplayIndex() const;
+
+	/** Returns the closest monitor index based on display ID and display index. */
+	UFUNCTION(BlueprintPure, Category=Settings)
+	ENGINE_API int32 GetClosestMonitorIndex() const;
+
+	/** Sets the user setting for the game window display. */
+	UFUNCTION(BlueprintCallable, Category=Settings)
+	ENGINE_API void SetDisplayProperties(const FString& ID, int32 Index);
+
 	/** Sets the user setting for vsync. See UGameUserSettings::bUseVSync. */
 	UFUNCTION(BlueprintCallable, Category=Settings)
 	ENGINE_API void SetVSyncEnabled(bool bEnable);
@@ -344,7 +368,7 @@ public:
 	static ENGINE_API void LoadConfigIni(bool bForceReload = false);
 
 	/** Request a change to the specified resolution and window mode. Optionally apply cmd line overrides. */
-	static ENGINE_API void RequestResolutionChange(int32 InResolutionX, int32 InResolutionY, EWindowMode::Type InWindowMode, bool bInDoOverrides = true);
+	static ENGINE_API void RequestResolutionChange(int32 InResolutionX, int32 InResolutionY, EWindowMode::Type InWindowMode, bool bInDoOverrides = true, int32 MonitorIndex = -1);
 
 	/** Returns the game local machine settings (resolution, windowing mode, scalability settings, etc...) */
 	UFUNCTION(BlueprintCallable, Category=Settings)
@@ -432,6 +456,22 @@ protected:
 	UPROPERTY(config)
 	int32 PreferredFullscreenMode;
 
+	/** String that attempts to uniquely identify a display even if FDisplayMetrics::MonitorInfo is reordered. */
+	UPROPERTY(config)
+	FString DisplayID;
+
+	/** Index of a display in FDisplayMetrics::MonitorInfo. */
+	UPROPERTY(config)
+	int32 DisplayIndex;
+
+	/** Last user confirmed display by ID. */
+	UPROPERTY(config)
+	FString LastUserConfirmedDisplayID;
+
+	/** Last user confirmed display by index. */
+	UPROPERTY(config)
+	int32 LastUserConfirmedDisplayIndex;
+
 	/** All settings will be wiped and set to default if the serialized version differs from UE_GAMEUSERSETTINGS_VERSION. */
 	UPROPERTY(config)
 	uint32 Version;
@@ -562,6 +602,12 @@ public:
 	DECLARE_DELEGATE_RetVal_OneParam(bool, FUpdateCloudDataFromGameUserSettings, const FString&);
 	FUpdateCloudDataFromGameUserSettings OnUpdateCloudDataFromGameUserSettings;
 
+	/**
+	 * We call this to notify UI to refresh relevant settings.
+	 */
+	DECLARE_MULTICAST_DELEGATE(FOnGameUserSettingsVideoRevert);
+	FOnGameUserSettingsVideoRevert OnGameUserSettingsVideoRevert;
+
 protected:
 	/**
 	 * Check if the current version of the game user settings is valid. Sub-classes can override this to provide game-specific versioning as necessary.
diff --git a/Engine/Source/Runtime/Engine/Private/GameEngine.cpp b/Engine/Source/Runtime/Engine/Private/GameEngine.cpp
--- a/Engine/Source/Runtime/Engine/Private/GameEngine.cpp
+++ b/Engine/Source/Runtime/Engine/Private/GameEngine.cpp
@@ -131,9 +131,9 @@ static void RunSynthBenchmark(const TArray<FString>& Args)
 }
 
 /** Helper function to generate a set of windowed resolutions which are convenient for the current primary display size */
-void GenerateConvenientWindowedResolutions(const struct FDisplayMetrics& InDisplayMetrics, TArray<FIntPoint>& OutResolutions)
+static void GenerateConvenientWindowedResolutions(int32 WorkAreaWidth, int32 WorkAreaHeight, TArray<FIntPoint>& OutResolutions)
 {
-	bool bInPortraitMode = InDisplayMetrics.PrimaryDisplayWidth < InDisplayMetrics.PrimaryDisplayHeight;
+	bool bInPortraitMode = WorkAreaWidth < WorkAreaHeight;
 
 	// Generate windowed resolutions as scaled versions of primary monitor size
 	static const float Scales[] = { 3.0f / 6.0f, 4.0f / 6.0f, 4.5f / 6.0f, 5.0f / 6.0f };
@@ -153,16 +153,16 @@ void GenerateConvenientWindowedResolutions(const struct FDisplayMetrics& InDispl
 
 			if (bInPortraitMode)
 			{
-				TargetHeight = FMath::RoundToFloat(InDisplayMetrics.PrimaryDisplayHeight * Scales[ScaleIndex]);
+				TargetHeight = FMath::RoundToFloat(WorkAreaHeight * Scales[ScaleIndex]);
 				TargetWidth = TargetHeight * Aspect;
 			}
 			else
 			{
-				TargetWidth = FMath::RoundToFloat(InDisplayMetrics.PrimaryDisplayWidth * Scales[ScaleIndex]);
+				TargetWidth = FMath::RoundToFloat(WorkAreaWidth * Scales[ScaleIndex]);
 				TargetHeight = TargetWidth * Aspect;
 			}
 
-			if (TargetWidth < InDisplayMetrics.PrimaryDisplayWidth && TargetHeight < InDisplayMetrics.PrimaryDisplayHeight && TargetWidth >= MinWidth && TargetHeight >= MinHeight)
+			if (TargetWidth < WorkAreaWidth && TargetHeight < WorkAreaHeight && TargetWidth >= MinWidth && TargetHeight >= MinHeight)
 			{
 				OutResolutions.Add(FIntPoint(TargetWidth, TargetHeight));
 			}
@@ -172,7 +172,7 @@ void GenerateConvenientWindowedResolutions(const struct FDisplayMetrics& InDispl
 	// if no convenient resolutions have been found, add a minimum one
 	if (OutResolutions.Num() == 0)
 	{
-		if (InDisplayMetrics.PrimaryDisplayHeight > MinHeight && InDisplayMetrics.PrimaryDisplayWidth > MinWidth)
+		if (WorkAreaHeight > MinHeight && WorkAreaWidth > MinWidth)
 		{
 			//Add the minimum size if it fit
 			OutResolutions.Add(FIntPoint(MinWidth, MinHeight));
@@ -181,13 +181,18 @@ void GenerateConvenientWindowedResolutions(const struct FDisplayMetrics& InDispl
 		{
 			//Force a resolution even if its smaller then the minimum height and width to avoid a bigger window then the desktop
 			// LWC_TODO: revisit. Seems like the round should be done after the multiply, otherwise the int was going to float and no rounding at all occured here.
-			float TargetWidth = FMath::RoundToFloat((float)InDisplayMetrics.PrimaryDisplayWidth) * Scales[NumScales - 1];
-			float TargetHeight = FMath::RoundToFloat((float)InDisplayMetrics.PrimaryDisplayHeight) * Scales[NumScales - 1];
+			float TargetWidth = FMath::RoundToFloat((float)WorkAreaWidth) * Scales[NumScales - 1];
+			float TargetHeight = FMath::RoundToFloat((float)WorkAreaHeight) * Scales[NumScales - 1];
 			OutResolutions.Add(FIntPoint(TargetWidth, TargetHeight));
 		}
 	}
 }
 
+void GenerateConvenientWindowedResolutions(const struct FDisplayMetrics& InDisplayMetrics, TArray<FIntPoint>& OutResolutions)
+{
+	GenerateConvenientWindowedResolutions(InDisplayMetrics.PrimaryDisplayWidth, InDisplayMetrics.PrimaryDisplayHeight, OutResolutions);
+}
+
 static FAutoConsoleCommand GSynthBenchmarkCmd(
 	TEXT("SynthBenchmark"),
 	TEXT("Run simple benchmark to get some metrics to find reasonable game settings automatically\n")
@@ -315,7 +320,7 @@ void UGameEngine::EnableGameWindowSettingsOverride(bool bEnabled)
 	bGameWindowSettingsOverrideEnabled = bEnabled;
 }
 
-void UGameEngine::ConditionallyOverrideSettings(int32& ResolutionX, int32& ResolutionY, EWindowMode::Type& WindowMode)
+void UGameEngine::ConditionallyOverrideSettings(int32& ResolutionX, int32& ResolutionY, EWindowMode::Type& WindowMode, int32 MonitorIndex)
 {
 	if (bGameWindowSettingsOverrideEnabled)
 	{
@@ -333,10 +338,10 @@ void UGameEngine::ConditionallyOverrideSettings(int32& ResolutionX, int32& Resol
 		}
 	}
 
-	DetermineGameWindowResolution(ResolutionX, ResolutionY, WindowMode);
+	DetermineGameWindowResolution(ResolutionX, ResolutionY, WindowMode, false, MonitorIndex);
 }
 
-void UGameEngine::DetermineGameWindowResolution( int32& ResolutionX, int32& ResolutionY, EWindowMode::Type& WindowMode, bool bUseWorkAreaForWindowed )
+void UGameEngine::DetermineGameWindowResolution( int32& ResolutionX, int32& ResolutionY, EWindowMode::Type& WindowMode, bool bUseWorkAreaForWindowed, int32 MonitorIndex )
 {
 	FString ResolutionStr;;
 	if (bGameWindowSettingsOverrideEnabled && FParse::Value(FCommandLine::Get(), TEXT("Res="), ResolutionStr))
@@ -376,7 +381,7 @@ void UGameEngine::DetermineGameWindowResolution( int32& ResolutionX, int32& Reso
 	FDisplayMetrics DisplayMetrics;
 	if (FSlateApplication::IsInitialized())
 	{
-		FSlateApplication::Get().GetInitialDisplayMetrics(DisplayMetrics);
+		FSlateApplication::Get().GetCachedDisplayMetrics(DisplayMetrics);
 	}
 	else
 	{
@@ -387,7 +392,33 @@ void UGameEngine::DetermineGameWindowResolution( int32& ResolutionX, int32& Reso
 	// Use PrimaryDisplayWidth/Height in windowed mode
 	int32 MaxResolutionX = bUseWorkAreaForWindowed && WindowMode == EWindowMode::Windowed ? DisplayMetrics.PrimaryDisplayWorkAreaRect.Right - DisplayMetrics.PrimaryDisplayWorkAreaRect.Left : DisplayMetrics.PrimaryDisplayWidth;
 	int32 MaxResolutionY = bUseWorkAreaForWindowed && WindowMode == EWindowMode::Windowed ? DisplayMetrics.PrimaryDisplayWorkAreaRect.Bottom - DisplayMetrics.PrimaryDisplayWorkAreaRect.Top : DisplayMetrics.PrimaryDisplayHeight;
-	if (WindowMode == EWindowMode::Fullscreen && DisplayMetrics.MonitorInfo.Num() > 0)
+	float DisplayAspect = (float)DisplayMetrics.PrimaryDisplayWidth / (float)DisplayMetrics.PrimaryDisplayHeight;
+	if (DisplayMetrics.MonitorInfo.IsValidIndex(MonitorIndex))
+	{
+		// Fall back to the monitor's native width/height if there was no max width/height found.
+		const FMonitorInfo& Info = DisplayMetrics.MonitorInfo[MonitorIndex];
+		int32 MonitorWidth = Info.MaxResolution.X;
+		int32 MonitorHeight = Info.MaxResolution.Y;
+		if (MonitorWidth == 0 || MonitorHeight == 0)
+		{
+			MonitorWidth = Info.NativeWidth;
+			MonitorHeight = Info.NativeHeight;
+		}
+
+		DisplayAspect = (float)MonitorWidth / (float)MonitorHeight;
+
+		if (bUseWorkAreaForWindowed && WindowMode == EWindowMode::Windowed)
+		{
+			MaxResolutionX = (Info.WorkArea.Right - Info.WorkArea.Left) - (Info.DisplayRect.Right - Info.DisplayRect.Left) + MonitorWidth;
+			MaxResolutionY = (Info.WorkArea.Bottom - Info.WorkArea.Top) - (Info.DisplayRect.Bottom - Info.DisplayRect.Top) + MonitorHeight;
+		}
+		else
+		{
+			MaxResolutionX = MonitorWidth;
+			MaxResolutionY = MonitorHeight;
+		}
+	}
+	else if (WindowMode == EWindowMode::Fullscreen && !DisplayMetrics.MonitorInfo.IsEmpty())
 	{
 		// In fullscreen, PrimaryDisplayWidth/Height is equal to your current resolution, so we will use your max native resolution instead
 		// Since we have info for at least one monitor, default to that if the primary can not be found
@@ -428,7 +459,7 @@ void UGameEngine::DetermineGameWindowResolution( int32& ResolutionX, int32& Reso
 		if (WindowMode == EWindowMode::Windowed)
 		{
 			TArray<FIntPoint> WindowedResolutions;
-			GenerateConvenientWindowedResolutions(DisplayMetrics, WindowedResolutions);
+			GenerateConvenientWindowedResolutions(MaxResolutionX, MaxResolutionY, WindowedResolutions);
 
 			if (WindowedResolutions.Num() > 0)
 			{
@@ -437,7 +468,6 @@ void UGameEngine::DetermineGameWindowResolution( int32& ResolutionX, int32& Reso
 				ResolutionY = WindowedResolutions[WindowedResolutions.Num() - 1].Y;
 
 				// Attempt to find the largest one with the same aspect ratio
-				float DisplayAspect = (float)DisplayMetrics.PrimaryDisplayWidth / (float)DisplayMetrics.PrimaryDisplayHeight;
 				for (int32 i = WindowedResolutions.Num() - 1; i >= 0; --i)
 				{
 					float Aspect = (float)WindowedResolutions[i].X / (float)WindowedResolutions[i].Y;
@@ -461,7 +491,6 @@ void UGameEngine::DetermineGameWindowResolution( int32& ResolutionX, int32& Reso
 		FSystemResolution::RequestResolutionChange(ResolutionX, ResolutionY, EWindowMode::Fullscreen);
 	}
 
-
 	if (bGameWindowSettingsOverrideEnabled && FParse::Param(FCommandLine::Get(), TEXT("Portrait")))
 	{
 		Swap(ResolutionX, ResolutionY);
@@ -474,6 +503,7 @@ TSharedRef<SWindow> UGameEngine::CreateGameWindow()
 	int32 ResX = GSystemResolution.ResX;
 	int32 ResY = GSystemResolution.ResY;
 	EWindowMode::Type WindowMode = GSystemResolution.WindowMode;
+	const int32 MonitorIndex = GSystemResolution.MonitorIndex;
 	
 #if WITH_EDITOR
 	/**************************************************************/
@@ -483,7 +513,7 @@ TSharedRef<SWindow> UGameEngine::CreateGameWindow()
 #endif
 	
 	{
-		ConditionallyOverrideSettings(ResX, ResY, WindowMode);
+		ConditionallyOverrideSettings(ResX, ResY, WindowMode, MonitorIndex);
 
 		// If the current settings have been overridden, apply them back into the system
 		if (ResX != GSystemResolution.ResX || ResY != GSystemResolution.ResY || WindowMode != GSystemResolution.WindowMode)
@@ -634,6 +664,43 @@ TSharedRef<SWindow> UGameEngine::CreateGameWindow()
 		.SetBackgroundBrush(FSlateNoResource())
 		.SetChildBackgroundBrush(FSlateNoResource());
 
+	if (MonitorIndex != -1)
+	{
+		FDisplayMetrics DisplayMetrics;
+		if (FSlateApplication::IsInitialized())
+		{
+			FSlateApplication::Get().GetInitialDisplayMetrics(DisplayMetrics);
+		}
+		else
+		{
+			FDisplayMetrics::RebuildDisplayMetrics(DisplayMetrics);
+		}
+
+		if (DisplayMetrics.MonitorInfo.IsValidIndex(MonitorIndex))
+		{
+			AutoCenterType = EAutoCenter::None;
+			WinX = DisplayMetrics.MonitorInfo[MonitorIndex].WorkArea.Left;
+			WinY = DisplayMetrics.MonitorInfo[MonitorIndex].WorkArea.Top;
+
+			if (WindowMode == EWindowMode::Windowed)
+			{
+				const FPlatformRect& WorkArea = DisplayMetrics.MonitorInfo[MonitorIndex].WorkArea;
+				const int32 WorkWidth = WorkArea.Right - WorkArea.Left;
+				const int32 WorkHeight = WorkArea.Bottom - WorkArea.Top;
+
+				if (ResX < WorkWidth)
+				{
+					WinX += (WorkWidth - ResX) / 2;
+				}
+
+				if (ResY < WorkHeight)
+				{
+					WinY += (WorkHeight - ResY) / 2;
+				}
+			}
+		}
+	}
+
 	TSharedRef<SWindow> Window = SNew(SWindow)
 	.Type(EWindowType::GameWindow)
 	.Style(bUseBorderlessWindow ? &BorderlessStyle : &FCoreStyle::Get().GetWidgetStyle<FWindowStyle>("Window"))
@@ -728,6 +795,7 @@ void UGameEngine::SwitchGameWindowToUseGameViewport()
 			// If Scene Viewport is not valid, the window was closed.
 			if (SceneViewport.IsValid() && FSlateApplication::IsInitialized())
 			{
+				SceneViewport->SetMonitorIndexForResize(GSystemResolution.MonitorIndex);
 				SceneViewport->ResizeFrame((uint32)GSystemResolution.ResX, (uint32)GSystemResolution.ResY, GSystemResolution.WindowMode);
 
 				// Registration of the game viewport to that messages are correctly received.
diff --git a/Engine/Source/Runtime/Engine/Private/GameUserSettings.cpp b/Engine/Source/Runtime/Engine/Private/GameUserSettings.cpp
--- a/Engine/Source/Runtime/Engine/Private/GameUserSettings.cpp
+++ b/Engine/Source/Runtime/Engine/Private/GameUserSettings.cpp
@@ -64,7 +64,7 @@ FIntPoint UGameUserSettings::GetDesktopResolution() const
 	FDisplayMetrics DisplayMetrics;
 	if (FSlateApplication::IsInitialized())
 	{
-		FSlateApplication::Get().GetInitialDisplayMetrics(DisplayMetrics);
+		FSlateApplication::Get().GetCachedDisplayMetrics(DisplayMetrics);
 	}
 	else
 	{
@@ -79,7 +79,16 @@ FIntPoint UGameUserSettings::GetDesktopResolution() const
 			DisplayMetrics.PrimaryDisplayHeight = 0;
 		}
 	}
-	return FIntPoint(DisplayMetrics.PrimaryDisplayWidth, DisplayMetrics.PrimaryDisplayHeight);
+
+	FIntPoint Resolution = FIntPoint(DisplayMetrics.PrimaryDisplayWidth, DisplayMetrics.PrimaryDisplayHeight);
+	const int32 MonitorIndex = DisplayMetrics.GetClosestMonitorFromIDAndIndex(DisplayID, DisplayIndex);
+	if (DisplayMetrics.MonitorInfo.IsValidIndex(MonitorIndex))
+	{
+		const FPlatformRect& DisplayRect = DisplayMetrics.MonitorInfo[MonitorIndex].DisplayRect;
+		Resolution.X = DisplayRect.Right - DisplayRect.Left;
+		Resolution.Y = DisplayRect.Bottom - DisplayRect.Top;
+	}
+	return Resolution;
 }
 
 void UGameUserSettings::SetScreenResolution(FIntPoint Resolution)
@@ -141,6 +150,51 @@ EWindowMode::Type UGameUserSettings::GetPreferredFullscreenMode() const
 	return PreferredFullscreenMode == 0 ? EWindowMode::Fullscreen : EWindowMode::WindowedFullscreen;
 }
 
+FString UGameUserSettings::GetDisplayID() const
+{
+	return DisplayID;
+}
+
+FString UGameUserSettings::GetLastConfirmedDisplayID() const
+{
+	return LastUserConfirmedDisplayID;
+}
+
+int32 UGameUserSettings::GetDisplayIndex() const
+{
+	return DisplayIndex;
+}
+
+int32 UGameUserSettings::GetLastConfirmedDisplayIndex() const
+{
+	return LastUserConfirmedDisplayIndex;
+}
+
+static int32 GetClosestMonitorFromIDAndIndex(const FString& DisplayID, int32 DisplayIndex)
+{
+	int32 MonitorIndex = -1;
+
+	if (!DisplayID.IsEmpty() || (DisplayIndex != -1))
+	{
+		FDisplayMetrics DisplayMetrics;
+		FDisplayMetrics::RebuildDisplayMetrics(DisplayMetrics);
+		MonitorIndex = DisplayMetrics.GetClosestMonitorFromIDAndIndex(DisplayID, DisplayIndex);
+	}
+
+	return MonitorIndex;
+}
+
+int32 UGameUserSettings::GetClosestMonitorIndex() const
+{
+	return GetClosestMonitorFromIDAndIndex(DisplayID, DisplayIndex);
+}
+
+void UGameUserSettings::SetDisplayProperties(const FString& ID, int32 Index)
+{
+	DisplayID = ID;
+	DisplayIndex = Index;
+}
+
 void UGameUserSettings::SetVSyncEnabled(bool bEnable)
 {
 	bUseVSync = bEnable;
@@ -214,6 +268,8 @@ void UGameUserSettings::ConfirmVideoMode()
 	LastConfirmedFullscreenMode = FullscreenMode;
 	LastUserConfirmedResolutionSizeX = ResolutionSizeX;
 	LastUserConfirmedResolutionSizeY = ResolutionSizeY;
+	LastUserConfirmedDisplayID = DisplayID;
+	LastUserConfirmedDisplayIndex = DisplayIndex;
 }
 
 void UGameUserSettings::RevertVideoMode()
@@ -221,6 +277,10 @@ void UGameUserSettings::RevertVideoMode()
 	FullscreenMode = LastConfirmedFullscreenMode;
 	ResolutionSizeX = LastUserConfirmedResolutionSizeX;
 	ResolutionSizeY = LastUserConfirmedResolutionSizeY;
+	DisplayID = LastUserConfirmedDisplayID;
+	DisplayIndex = LastUserConfirmedDisplayIndex;
+
+	OnGameUserSettingsVideoRevert.Broadcast();
 }
 
 void UGameUserSettings::SetToDefaults()
@@ -253,6 +313,11 @@ void UGameUserSettings::SetToDefaults()
 	static const auto CVar = IConsoleManager::Get().FindTConsoleVariableDataInt(TEXT("r.FullScreenMode"));
 	PreferredFullscreenMode = CVar->GetValueOnGameThread();
 
+	DisplayID.Empty();
+	DisplayIndex = -1;
+	LastUserConfirmedDisplayID = DisplayID;
+	LastUserConfirmedDisplayIndex = DisplayIndex;
+
 	ScalabilityQuality.SetDefaults();
 
 	if (FApp::CanEverRender())
@@ -407,6 +472,8 @@ void UGameUserSettings::ValidateSettings()
 		LastConfirmedFullscreenMode = FullscreenMode;
 		LastUserConfirmedResolutionSizeX = ResolutionSizeX;
 		LastUserConfirmedResolutionSizeY = ResolutionSizeY;
+		LastUserConfirmedDisplayID = DisplayID;
+		LastUserConfirmedDisplayIndex = DisplayIndex;
 	}
 
 	const int32 ScreenWidth = (FullscreenMode == EWindowMode::WindowedFullscreen) ? GetDesktopResolution().X : ResolutionSizeX;
@@ -506,7 +573,8 @@ void UGameUserSettings::ApplyResolutionSettings(bool bCheckForCommandLineOverrid
 	EWindowMode::Type NewFullscreenMode = GetFullscreenMode();
 
 	// Request a resolution change
-	RequestResolutionChange(ResolutionSizeX, ResolutionSizeY, NewFullscreenMode, bCheckForCommandLineOverrides);
+	const int32 MonitorIndex = GetClosestMonitorIndex();
+	RequestResolutionChange(ResolutionSizeX, ResolutionSizeY, NewFullscreenMode, bCheckForCommandLineOverrides, MonitorIndex);
 
 	if (NewFullscreenMode == EWindowMode::Fullscreen || NewFullscreenMode == EWindowMode::WindowedFullscreen)
 	{
@@ -583,7 +651,7 @@ void UGameUserSettings::LoadSettings(bool bForceReload/*=false*/)
 	}
 }
 
-void UGameUserSettings::RequestResolutionChange(int32 InResolutionX, int32 InResolutionY, EWindowMode::Type InWindowMode, bool bInDoOverrides /* = true */)
+void UGameUserSettings::RequestResolutionChange(int32 InResolutionX, int32 InResolutionY, EWindowMode::Type InWindowMode, bool bInDoOverrides /* = true */, int32 MonitorIndex /* = -1 */)
 {
 	if (FPlatformProperties::HasFixedResolution())
 	{
@@ -592,9 +660,10 @@ void UGameUserSettings::RequestResolutionChange(int32 InResolutionX, int32 InRes
 
 	if (bInDoOverrides)
 	{
-		UGameEngine::ConditionallyOverrideSettings(InResolutionX, InResolutionY, InWindowMode);
+		UGameEngine::ConditionallyOverrideSettings(InResolutionX, InResolutionY, InWindowMode, MonitorIndex);
 	}
 
+	FSystemResolution::SetMonitorIndex(MonitorIndex);
 	FSystemResolution::RequestResolutionChange(InResolutionX, InResolutionY, InWindowMode);
 }
 
@@ -688,6 +757,8 @@ void UGameUserSettings::PreloadResolutionSettings(bool bAllowCmdLineOverrides /*
 	EWindowMode::Type WindowMode = GetDefaultWindowMode();
 	bool bUseDesktopResolution = false;
 	bool bUseHDR = FPlatformMisc::UseHDRByDefault();
+	FString DisplayID;
+	int32 DisplayIndex = -1;
 
 	int32 Version = 0;
 	if (GConfig->GetInt(*GameUserSettingsCategory, TEXT("Version"), Version, GGameUserSettingsIni) && Version == UE_GAMEUSERSETTINGS_VERSION)
@@ -701,15 +772,26 @@ void UGameUserSettings::PreloadResolutionSettings(bool bAllowCmdLineOverrides /*
 		GConfig->GetInt(*GameUserSettingsCategory, TEXT("ResolutionSizeX"), ResolutionX, GGameUserSettingsIni);
 		GConfig->GetInt(*GameUserSettingsCategory, TEXT("ResolutionSizeY"), ResolutionY, GGameUserSettingsIni);
 
+		GConfig->GetString(*GameUserSettingsCategory, TEXT("DisplayID"), DisplayID, GGameUserSettingsIni);
+		GConfig->GetInt(*GameUserSettingsCategory, TEXT("DisplayIndex"), DisplayIndex, GGameUserSettingsIni);
+
 #if PLATFORM_DESKTOP
 		if (bUseDesktopResolution && ResolutionX == 0 && ResolutionY == 0 && WindowMode != EWindowMode::Windowed)
 		{
-			// Grab display metrics so we can get the primary display output size.
+			// Grab display metrics so we can get the display output size.
 			FDisplayMetrics DisplayMetrics;
 			FDisplayMetrics::RebuildDisplayMetrics(DisplayMetrics);
-
-			ResolutionX = DisplayMetrics.PrimaryDisplayWidth;
-			ResolutionY = DisplayMetrics.PrimaryDisplayHeight;
+			const int32 MonitorIndex = DisplayMetrics.GetClosestMonitorFromIDAndIndex(DisplayID, DisplayIndex);
+			if (DisplayMetrics.MonitorInfo.IsValidIndex(MonitorIndex))
+			{
+				ResolutionX = DisplayMetrics.MonitorInfo[MonitorIndex].NativeWidth;
+				ResolutionY = DisplayMetrics.MonitorInfo[MonitorIndex].NativeHeight;
+			}
+			else
+			{
+				ResolutionX = DisplayMetrics.PrimaryDisplayWidth;
+				ResolutionY = DisplayMetrics.PrimaryDisplayHeight;
+			}
 		}
 #endif
 
@@ -728,7 +810,8 @@ void UGameUserSettings::PreloadResolutionSettings(bool bAllowCmdLineOverrides /*
 	}
 #endif
 
-	RequestResolutionChange(ResolutionX, ResolutionY, WindowMode, bAllowCmdLineOverrides);
+	const int32 MonitorIndex = GetClosestMonitorFromIDAndIndex(DisplayID, DisplayIndex);
+	RequestResolutionChange(ResolutionX, ResolutionY, WindowMode, bAllowCmdLineOverrides, MonitorIndex);
 
 	IConsoleManager::Get().CallAllConsoleVariableSinks();
 }
diff --git a/Engine/Source/Runtime/Engine/Private/GameViewportClient.cpp b/Engine/Source/Runtime/Engine/Private/GameViewportClient.cpp
--- a/Engine/Source/Runtime/Engine/Private/GameViewportClient.cpp
+++ b/Engine/Source/Runtime/Engine/Private/GameViewportClient.cpp
@@ -355,12 +355,16 @@ UGameViewportClient::UGameViewportClient(const FObjectInitializer& ObjectInitial
 		FCoreDelegates::StatDisabled.AddUObject(this, &UGameViewportClient::HandleViewportStatDisabled);
 		FCoreDelegates::StatDisableAll.AddUObject(this, &UGameViewportClient::HandleViewportStatDisableAll);
 
-#if WITH_EDITOR
-		if (GIsEditor)
+		if (FSlateApplication::IsInitialized())
 		{
-			FSlateApplication::Get().OnWindowDPIScaleChanged().AddUObject(this, &UGameViewportClient::HandleWindowDPIScaleChanged);
-		}
+			FSlateApplication::Get().OnRegularWindowMovedOrResized().AddUObject(this, &UGameViewportClient::HandleRegularWindowMovedOrResized);
+#if WITH_EDITOR
+			if (GIsEditor)
+			{
+				FSlateApplication::Get().OnWindowDPIScaleChanged().AddUObject(this, &UGameViewportClient::HandleWindowDPIScaleChanged);
+			}
 #endif
+		}
 
 		AudioDeviceDestroyedHandle = FAudioDeviceManagerDelegates::OnAudioDeviceDestroyed.AddLambda([this](const Audio::FDeviceId InDeviceId)
 		{
@@ -407,12 +411,16 @@ UGameViewportClient::~UGameViewportClient()
 	FCoreDelegates::StatDisabled.RemoveAll(this);
 	FCoreDelegates::StatDisableAll.RemoveAll(this);
 
-#if WITH_EDITOR
-	if (GIsEditor && FSlateApplication::IsInitialized())
+	if (FSlateApplication::IsInitialized())
 	{
-		FSlateApplication::Get().OnWindowDPIScaleChanged().RemoveAll(this);
-	}
+		FSlateApplication::Get().OnRegularWindowMovedOrResized().RemoveAll(this);
+#if WITH_EDITOR
+		if (GIsEditor)
+		{
+			FSlateApplication::Get().OnWindowDPIScaleChanged().RemoveAll(this);
+		}
 #endif
+	}
 
 	if (StatHitchesData)
 	{
@@ -3966,25 +3974,45 @@ bool UGameViewportClient::HandleToggleFullscreenCommand()
 
 	static auto CVar = IConsoleManager::Get().FindTConsoleVariableDataInt(TEXT("r.FullScreenMode"));
 	check(CVar);
-	auto FullScreenMode = CVar->GetValueOnGameThread() == 0 ? EWindowMode::Fullscreen : EWindowMode::WindowedFullscreen;
-	FullScreenMode = Viewport->IsFullscreen() ? EWindowMode::Windowed : FullScreenMode;
+	EWindowMode::Type FullScreenMode = Viewport->IsFullscreen()
+		? EWindowMode::Windowed
+		: (((CVar->GetValueOnGameThread() == 0) && Viewport->IsExclusiveFullscreenSupported())
+			? EWindowMode::Fullscreen
+			: EWindowMode::WindowedFullscreen);
 
 	int32 ResolutionX = GSystemResolution.ResX;
 	int32 ResolutionY = GSystemResolution.ResY;
+	const int32 MonitorIndex = Viewport->GetCurrentMonitorIndex();
 	bool bNewModeApplied = false;
 
 	// Make sure the user's settings are updated after pressing Alt+Enter to toggle fullscreen.  Note
 	// that we don't need to "apply" the setting change, as we already did that above directly.
-	UGameEngine* GameEngine = Cast<UGameEngine>( GEngine );
-	if( GameEngine )
+	UGameEngine* const GameEngine = Cast<UGameEngine>( GEngine );
+	if (GameEngine != nullptr)
 	{
-		UGameUserSettings* UserSettings = GameEngine->GetGameUserSettings();
-		if( UserSettings != nullptr )
+		UGameUserSettings* const UserSettings = GameEngine->GetGameUserSettings();
+		if (UserSettings != nullptr)
 		{
-			// Ensure that our desired screen size will fit on the display
 			ResolutionX = UserSettings->GetScreenResolution().X;
 			ResolutionY = UserSettings->GetScreenResolution().Y;
-			UGameEngine::DetermineGameWindowResolution(ResolutionX, ResolutionY, FullScreenMode, true);
+
+			const FDisplayMetrics& DisplayMetrics = FSlateApplication::Get().GetCachedDisplayMetricsByRef();
+			if (DisplayMetrics.MonitorInfo.IsValidIndex(MonitorIndex))
+			{
+				const FMonitorInfo& Info = DisplayMetrics.MonitorInfo[MonitorIndex];
+				ResolutionX = Info.MaxResolution.X;
+				ResolutionY = Info.MaxResolution.Y;
+				if ((ResolutionX == 0) || (ResolutionY == 0))
+				{
+					ResolutionX = Info.NativeWidth;
+					ResolutionY = Info.NativeHeight;
+				}
+
+				UserSettings->SetDisplayProperties(Info.ID, MonitorIndex);
+			}
+
+			// Ensure that our desired screen size will fit on the display
+			UGameEngine::DetermineGameWindowResolution(ResolutionX, ResolutionY, FullScreenMode, true, MonitorIndex);
 
 			UserSettings->SetScreenResolution(FIntPoint(ResolutionX, ResolutionY));
 			UserSettings->SetFullscreenMode(FullScreenMode);
@@ -3996,6 +4024,7 @@ bool UGameViewportClient::HandleToggleFullscreenCommand()
 
 	if (!bNewModeApplied)
 	{
+		FSystemResolution::SetMonitorIndex(MonitorIndex);
 		FSystemResolution::RequestResolutionChange(ResolutionX, ResolutionY, FullScreenMode);
 	}
 
@@ -4444,6 +4473,33 @@ void UGameViewportClient::HandleViewportStatDisableAll(const bool bInAnyViewport
 	}
 }
 
+void UGameViewportClient::HandleRegularWindowMovedOrResized(TSharedRef<SWindow> InWindow)
+{
+	if ((InWindow == Window) && Viewport)
+	{
+		UGameEngine* const GameEngine = Cast<UGameEngine>(GEngine);
+		if (GameEngine)
+		{
+			UGameUserSettings* const UserSettings = GameEngine->GetGameUserSettings();
+			if (UserSettings != nullptr)
+			{
+				const FDisplayMetrics& DisplayMetrics = FSlateApplication::Get().GetCachedDisplayMetricsByRef();
+				const int32 MonitorIndex = Viewport->GetCurrentMonitorIndex();
+				if (DisplayMetrics.MonitorInfo.IsValidIndex(MonitorIndex))
+				{
+					const FString& DisplayID = DisplayMetrics.MonitorInfo[MonitorIndex].ID;
+					if (DisplayID != UserSettings->GetDisplayID())
+					{
+						UserSettings->SetDisplayProperties(DisplayID, MonitorIndex);
+
+						WindowDisplayChangedDelegate.Broadcast();
+					}
+				}
+			}
+		}
+	}
+}
+
 void UGameViewportClient::HandleWindowDPIScaleChanged(TSharedRef<SWindow> InWindow)
 {
 #if WITH_EDITOR
diff --git a/Engine/Source/Runtime/Engine/Private/Slate/SceneViewport.cpp b/Engine/Source/Runtime/Engine/Private/Slate/SceneViewport.cpp
--- a/Engine/Source/Runtime/Engine/Private/Slate/SceneViewport.cpp
+++ b/Engine/Source/Runtime/Engine/Private/Slate/SceneViewport.cpp
@@ -680,10 +680,11 @@ FReply FSceneViewport::AcquireFocusAndCapture(FIntPoint MousePosition, EFocusCau
 			TSharedPtr<SWindow> Window = FSlateApplication::Get().FindWidgetWindow(ViewportWidgetRef);
 			if (FSlateApplication::Get().GetTransformFullscreenMouseInput() && !GIsEditor && Window.IsValid() && Window->GetWindowMode() == EWindowMode::Fullscreen)
 			{
-				FDisplayMetrics CachedDisplayMetrics;
-				FSlateApplication::Get().GetCachedDisplayMetrics(CachedDisplayMetrics);
+				const FDisplayMetrics& CachedDisplayMetrics = FSlateApplication::Get().GetCachedDisplayMetricsByRef();
 				FVector2D WindowSize = Window->GetSizeInScreen();
-				FVector2D DisplaySize = { (float)CachedDisplayMetrics.PrimaryDisplayWidth, (float)CachedDisplayMetrics.PrimaryDisplayHeight };
+				FVector2D DisplaySize = CachedDisplayMetrics.MonitorInfo.IsEmpty()
+					? FVector2D((float)CachedDisplayMetrics.PrimaryDisplayWidth, (float)CachedDisplayMetrics.PrimaryDisplayHeight)
+					: Window->GetFullScreenInfo().GetSize();
 				FVector2D CorrectionScale = DisplaySize / WindowSize;
 				MousePosBeforeHiddenDueToCapture = { (int32)(MousePosition.X * CorrectionScale.X), (int32)(MousePosition.Y * CorrectionScale.Y) };
 			}
@@ -1461,6 +1462,29 @@ void FSceneViewport::PaintDebugCanvas(const FGeometry& AllottedGeometry, FSlateW
 	}
 }
 
+static void EnsureValidFullscreenResolution(const FMonitorInfo& Monitor, uint32& NewWindowSizeX, uint32& NewWindowSizeY)
+{
+	FScreenResolutionArray ResArray;
+	RHIGetAvailableResolutionsForDisplay(ResArray, true, Monitor.NativeHandle);
+	if (ResArray.IsEmpty())
+	{
+		ResArray.Emplace(Monitor.DisplayRect.Right - Monitor.DisplayRect.Left, Monitor.DisplayRect.Bottom - Monitor.DisplayRect.Top);
+	}
+	bool bFound = false;
+	for (FScreenResolutionRHI ScreenResolution : ResArray)
+	{
+		if ((bFound = ScreenResolution.Width == NewWindowSizeX && ScreenResolution.Height == NewWindowSizeY))
+		{
+			break;
+		}
+	}
+	if (!bFound && !ResArray.IsEmpty())
+	{
+		NewWindowSizeX = ResArray[0].Width;
+		NewWindowSizeY = ResArray[0].Height;
+	}
+}
+
 void FSceneViewport::ResizeFrame(uint32 NewWindowSizeX, uint32 NewWindowSizeY, EWindowMode::Type NewWindowMode)
 {
 	// Resizing the window directly is only supported in the game
@@ -1476,6 +1500,45 @@ void FSceneViewport::ResizeFrame(uint32 NewWindowSizeX, uint32 NewWindowSizeY, E
 			const FVector2D OldWindowSize = WindowToResize->GetClientSizeInScreen();
 			const EWindowMode::Type OldWindowMode = WindowToResize->GetWindowMode();
 
+			// Check if switching to exclusive fullscreen mode on a different monitor.
+			static bool bAllowRecursion = true;
+			if (FPlatformProperties::SupportsWindowedMode() && (NewWindowMode == EWindowMode::Fullscreen) && bAllowRecursion)
+			{
+				int32 NewLeft = 0;
+				int32 NewTop = 0;
+				const FDisplayMetrics& DisplayMetrics = FSlateApplication::Get().GetCachedDisplayMetricsByRef();
+				if (DisplayMetrics.MonitorInfo.IsValidIndex(MonitorIndexForResize))
+				{
+					const FMonitorInfo& Monitor = DisplayMetrics.MonitorInfo[MonitorIndexForResize];
+					const FPlatformRect& DisplayRect = DisplayMetrics.MonitorInfo[MonitorIndexForResize].DisplayRect;
+					NewLeft = DisplayRect.Left;
+					NewTop = DisplayRect.Top;
+					EnsureValidFullscreenResolution(Monitor, NewWindowSizeX, NewWindowSizeY);
+				}
+				else
+				{
+					for (const FMonitorInfo& Monitor : DisplayMetrics.MonitorInfo)
+					{
+						if (Monitor.bIsPrimary)
+						{
+							EnsureValidFullscreenResolution(Monitor, NewWindowSizeX, NewWindowSizeY);
+							break;
+						}
+					}
+				}
+
+				const FSlateRect& OldFullScreenInfo = WindowToResize->GetFullScreenInfo();
+				if ((OldFullScreenInfo.Left != NewLeft) || (OldFullScreenInfo.Top != NewTop))
+				{
+					bAllowRecursion = false;
+					// Avoid switching to exclusive fullscreen on the wrong monitor.
+					ResizeFrame(NewWindowSizeX, NewWindowSizeY, EWindowMode::WindowedFullscreen);
+					ResizeFrame(NewWindowSizeX, NewWindowSizeY, EWindowMode::Fullscreen);
+					bAllowRecursion = true;
+					return;
+				}
+			}
+
 			// Set the new window mode first to ensure that the work area size is correct (fullscreen windows can affect this)
 			if (NewWindowMode != OldWindowMode)
 			{
@@ -1489,7 +1552,18 @@ void FSceneViewport::ResizeFrame(uint32 NewWindowSizeX, uint32 NewWindowSizeY, E
 			// Only adjust window size if not in off-screen rendering mode, because off-screen rendering skips rendering to screen and uses custom size.
 			if (!FSlateApplication::Get().IsRenderingOffScreen())
 			{
-				const FSlateRect BestWorkArea = FSlateApplication::Get().GetWorkArea(FSlateRect::FromPointAndExtent(OldWindowPos, OldWindowSize));
+				const FDisplayMetrics& DisplayMetrics = FSlateApplication::Get().GetCachedDisplayMetricsByRef();
+				const bool bValidMonitorIndex = DisplayMetrics.MonitorInfo.IsValidIndex(MonitorIndexForResize);
+				FSlateRect BestWorkArea;
+				if (bValidMonitorIndex)
+				{
+					const FPlatformRect& WorkArea = DisplayMetrics.MonitorInfo[MonitorIndexForResize].WorkArea;
+					BestWorkArea = FSlateRect(WorkArea.Left, WorkArea.Top, WorkArea.Right, WorkArea.Bottom);
+				}
+				else
+				{
+					BestWorkArea = FSlateApplication::Get().GetWorkArea(FSlateRect::FromPointAndExtent(OldWindowPos, OldWindowSize));
+				}
 
 				// A switch to window mode should position the window to be in the center of the work-area (we don't do this if we were already in window mode to allow the user to move the window)
 				// Fullscreen modes should position the window to the top-left of the monitor.
@@ -1498,7 +1572,7 @@ void FSceneViewport::ResizeFrame(uint32 NewWindowSizeX, uint32 NewWindowSizeY, E
 				// scaling so we actual render to the resolution we've been asked for.
 				if (NewWindowMode == EWindowMode::Windowed)
 				{
-					if (OldWindowMode == EWindowMode::Windowed && NewWindowSize == OldWindowSize)
+					if (!bValidMonitorIndex && (OldWindowMode == EWindowMode::Windowed) && (NewWindowSize == OldWindowSize))
 					{
 						// Leave the window position alone!
 						NewWindowPos.Reset();
@@ -1512,53 +1586,46 @@ void FSceneViewport::ResizeFrame(uint32 NewWindowSizeX, uint32 NewWindowSizeY, E
 
 						if (NewWindowSize.X < BestWorkAreaSize.X)
 						{
-							CenteredWindowPos.X += FMath::Max(0.0f, (BestWorkAreaSize.X - NewWindowSize.X) * 0.5f);
+							CenteredWindowPos.X += (BestWorkAreaSize.X - NewWindowSize.X) * 0.5f;
 						}
 
 						if (NewWindowSize.Y < BestWorkAreaSize.Y)
 						{
-							CenteredWindowPos.Y += FMath::Max(0.0f, (BestWorkAreaSize.Y - NewWindowSize.Y) * 0.5f);
+							CenteredWindowPos.Y += (BestWorkAreaSize.Y - NewWindowSize.Y) * 0.5f;
 						}
 
 						NewWindowPos = CenteredWindowPos;
 					}
 				}
-				else
+				else if (!DisplayMetrics.MonitorInfo.IsEmpty())
 				{
-					FDisplayMetrics DisplayMetrics;
-					FSlateApplication::Get().GetCachedDisplayMetrics(DisplayMetrics);
-
-					if (DisplayMetrics.MonitorInfo.Num() > 0)
+					FPlatformRect DisplayRect;
+					if (bValidMonitorIndex)
 					{
-						// Try to find the monitor that the viewport belongs to based on BestWorkArea.
-						// For widowed fullscreen and fullscreen modes it should be top left position of one of monitors.
-						FPlatformRect DisplayRect = DisplayMetrics.MonitorInfo[0].DisplayRect;
-						for (int32 Index = 1; Index < DisplayMetrics.MonitorInfo.Num(); ++Index)
-						{
-							const FMonitorInfo& MonitorInfo = DisplayMetrics.MonitorInfo[Index];
-							if (BestWorkArea.GetTopLeft() == FVector2D(MonitorInfo.WorkArea.Left, MonitorInfo.WorkArea.Top))
-							{
-								DisplayRect = DisplayMetrics.MonitorInfo[Index].DisplayRect;
-							}
-						}
-
-						NewWindowPos = FVector2D(DisplayRect.Left, DisplayRect.Top);
-
-						if (NewWindowMode == EWindowMode::WindowedFullscreen)
-						{
-							NewWindowSize.X = DisplayRect.Right - DisplayRect.Left;
-							NewWindowSize.Y = DisplayRect.Bottom - DisplayRect.Top;
-						}
+						DisplayRect = DisplayMetrics.MonitorInfo[MonitorIndexForResize].DisplayRect;
 					}
 					else
 					{
-						NewWindowPos = FVector2D(0.0f, 0.0f);
+						const FSlateRect& FullScreenInfo = WindowToResize->GetFullScreenInfo();
+						DisplayRect = FPlatformRect(FullScreenInfo.Left, FullScreenInfo.Top, FullScreenInfo.Right, FullScreenInfo.Bottom);
+					}
 
-						if (NewWindowMode == EWindowMode::WindowedFullscreen)
-						{
-							NewWindowSize.X = DisplayMetrics.PrimaryDisplayWidth;
-							NewWindowSize.Y = DisplayMetrics.PrimaryDisplayHeight;
-						}
+					NewWindowPos = FVector2D(DisplayRect.Left, DisplayRect.Top);
+
+					if (NewWindowMode == EWindowMode::WindowedFullscreen)
+					{
+						NewWindowSize.X = DisplayRect.Right - DisplayRect.Left;
+						NewWindowSize.Y = DisplayRect.Bottom - DisplayRect.Top;
+					}
+				}
+				else
+				{
+					NewWindowPos = FVector2D(0.0f, 0.0f);
+
+					if (NewWindowMode == EWindowMode::WindowedFullscreen)
+					{
+						NewWindowSize.X = DisplayMetrics.PrimaryDisplayWidth;
+						NewWindowSize.Y = DisplayMetrics.PrimaryDisplayHeight;
 					}
 				}
 
@@ -1620,7 +1687,6 @@ void FSceneViewport::ResizeFrame(uint32 NewWindowSizeX, uint32 NewWindowSizeY, E
 				ResizeViewport(NewViewportSize.X, NewViewportSize.Y, NewWindowMode);
 			}
 
-
 			if(NewBackbufferSize != BackBufferSize)
 			{
 				FSlateApplicationBase::Get().GetRenderer()->UpdateFullscreenState(WindowToResize.ToSharedRef(), NewBackbufferSize.X, NewBackbufferSize.Y);
@@ -1795,6 +1861,48 @@ void FSceneViewport::EnqueueEndRenderFrame(const bool bLockToVsync, const bool b
 	}
 }
 
+int32 FSceneViewport::GetCurrentMonitorIndex() const
+{
+	int32 MonitorIndex = -1;
+
+	const FDisplayMetrics& DisplayMetrics = FSlateApplication::Get().GetCachedDisplayMetricsByRef();
+	const int32 Num = DisplayMetrics.MonitorInfo.Num();
+	TSharedPtr<SWidget> PinnedViewport = ViewportWidget.Pin();
+	if (PinnedViewport.IsValid())
+	{
+		TSharedPtr<SWindow> Window = FSlateApplication::Get().FindWidgetWindow(PinnedViewport.ToSharedRef());
+		if (Window.IsValid())
+		{
+			const FSlateRect FullScreenInfo = Window->GetFullScreenInfo();
+			for (int32 Index = 0; Index < Num; ++Index)
+			{
+				const FPlatformRect& DisplayRect = DisplayMetrics.MonitorInfo[Index].DisplayRect;
+				if ((DisplayRect.Left == FullScreenInfo.Left) && (DisplayRect.Top == FullScreenInfo.Top))
+				{
+					MonitorIndex = Index;
+					break;
+				}
+			}
+		}
+	}
+
+	return MonitorIndex;
+}
+
+bool FSceneViewport::IsExclusiveFullscreenSupported() const
+{
+	bool bSupported = true;
+
+	TSharedPtr<SWidget> PinnedViewport = ViewportWidget.Pin();
+	if (PinnedViewport.IsValid())
+	{
+		TSharedPtr<SWindow> Window = FSlateApplication::Get().FindWidgetWindow(PinnedViewport.ToSharedRef());
+		bSupported = !Window.IsValid() || Window->GetNativeWindow()->IsFullscreenSupported();
+	}
+
+	return bSupported;
+}
+
 const FTextureRHIRef& FSceneViewport::GetRenderTargetTexture() const
 {
 	if (IsInRenderingThread())
diff --git a/Engine/Source/Runtime/Engine/Private/UnrealEngine.cpp b/Engine/Source/Runtime/Engine/Private/UnrealEngine.cpp
--- a/Engine/Source/Runtime/Engine/Private/UnrealEngine.cpp
+++ b/Engine/Source/Runtime/Engine/Private/UnrealEngine.cpp
@@ -479,6 +479,13 @@ static FAutoConsoleVariable CVarSystemResolution(
 	TEXT("     1920x1080wf for windowed fullscreen\n")
 );
 
+static FAutoConsoleVariable CVarMonitorIndex(
+	TEXT("r.SetRes.MonitorIndex"),
+	-1,
+	TEXT("Set the monitor to be used by r.SetRes. Has no effect in the editor.\n")
+	TEXT("-1 for legacy behavior")
+);
+
 #if !UE_BUILD_SHIPPING
 static TAutoConsoleVariable<float> CVarSetOverrideFPS(
 	TEXT("t.OverrideFPS"),
@@ -1097,7 +1104,8 @@ void SystemResolutionSinkCallback()
 		return;
 	}
 
-	auto ResString = CVarSystemResolution->GetString();
+	const FString ResString = CVarSystemResolution->GetString();
+	const int32 MonitorIndex = CVarMonitorIndex->GetInt();
 
 	uint32 ResX, ResY;
 	int32 WindowModeInt = GSystemResolution.WindowMode;
@@ -1109,11 +1117,13 @@ void SystemResolutionSinkCallback()
 		if( GSystemResolution.ResX != ResX ||
 			GSystemResolution.ResY != ResY ||
 			GSystemResolution.WindowMode != WindowMode ||
+			GSystemResolution.MonitorIndex != MonitorIndex ||
 			GSystemResolution.bForceRefresh)
 		{
 			GSystemResolution.ResX = ResX;
 			GSystemResolution.ResY = ResY;
 			GSystemResolution.WindowMode = WindowMode;
+			GSystemResolution.MonitorIndex = MonitorIndex;
 			GSystemResolution.bForceRefresh = false;
 
 			// tell anyone listening about the change
@@ -1121,7 +1131,8 @@ void SystemResolutionSinkCallback()
 
 			if (GEngine && GEngine->GameViewport && GEngine->GameViewport->ViewportFrame)
 			{
-				FPlatformMisc::LowLevelOutputDebugStringf(TEXT("Resizing viewport due to setres change, %d x %d\n"), ResX, ResY);
+				FPlatformMisc::LowLevelOutputDebugStringf(TEXT("Resizing viewport due to setres change, %d x %d, MonitorIndex=%d\n"), ResX, ResY, MonitorIndex);
+				GEngine->GameViewport->ViewportFrame->SetMonitorIndexForResize(MonitorIndex);
 				GEngine->GameViewport->ViewportFrame->ResizeFrame(ResX, ResY, WindowMode);
 			}
 		}
@@ -18599,6 +18610,11 @@ bool AllowHighQualityLightmaps(const FStaticFeatureLevel FeatureLevel)
 		&& !IsMobilePlatform(GShaderPlatformForFeatureLevel[FeatureLevel]);
 }
 
+void FSystemResolution::SetMonitorIndex(int32 MonitorIndex)
+{
+	CVarMonitorIndex->Set(MonitorIndex);
+}
+
 // Helper function for changing system resolution via the r.setres console command
 void FSystemResolution::RequestResolutionChange(int32 InResX, int32 InResY, EWindowMode::Type InWindowMode)
 {
diff --git a/Engine/Source/Runtime/Engine/Public/Slate/SceneViewport.h b/Engine/Source/Runtime/Engine/Public/Slate/SceneViewport.h
--- a/Engine/Source/Runtime/Engine/Public/Slate/SceneViewport.h
+++ b/Engine/Source/Runtime/Engine/Public/Slate/SceneViewport.h
@@ -56,6 +56,8 @@ public:
 	ENGINE_API virtual FCanvas* GetDebugCanvas() override;
 	ENGINE_API virtual float GetDisplayGamma() const override;
 	ENGINE_API virtual void EnqueueEndRenderFrame(const bool bLockToVsync, const bool bShouldPresent) override;
+	ENGINE_API virtual int32 GetCurrentMonitorIndex() const override;
+	ENGINE_API virtual bool IsExclusiveFullscreenSupported() const override;
 
 	/** Gets the proper RenderTarget based on the current thread*/
 	ENGINE_API virtual const FTextureRHIRef& GetRenderTargetTexture() const;
@@ -175,6 +177,8 @@ public:
 	 */
 	virtual bool RequiresVsync() const override { return bRequiresVsync; }
 
+	virtual void SetMonitorIndexForResize(int32 MonitorIndex) override { MonitorIndexForResize = MonitorIndex; }
+
 	/**
 	 * Called to resize the actual window where this viewport resides
 	 *
@@ -433,6 +437,8 @@ private:
 	bool bIsResizing;
 	/** Do not resize the RenderTarget when resizing */
 	bool bForceViewportSize;
+	/** The monitor that we want to display the viewport on after ResizeFrame */
+	int32 MonitorIndexForResize = -1;
 	/** Delegate that is fired off in ResizeFrame after ResizeViewport */
 	FOnSceneViewportResize OnSceneViewportResizeDel;
 	/** Whether the PIE viewport is currently in simulate in editor mode */
diff --git a/Engine/Source/Runtime/Engine/Public/UnrealClient.h b/Engine/Source/Runtime/Engine/Public/UnrealClient.h
--- a/Engine/Source/Runtime/Engine/Public/UnrealClient.h
+++ b/Engine/Source/Runtime/Engine/Public/UnrealClient.h
@@ -186,6 +186,7 @@ class FViewportFrame
 public:
 
 	virtual class FViewport* GetViewport() = 0;
+	virtual void SetMonitorIndexForResize(int32 MonitorIndex) = 0;
 	virtual void ResizeFrame(uint32 NewSizeX,uint32 NewSizeY,EWindowMode::Type NewWindowMode) = 0;
 };
 
@@ -470,6 +471,8 @@ public:
 	virtual bool IsPenActive() { return false; }
 	virtual void SetMouse(int32 x, int32 y) = 0;
 	virtual bool IsFullscreen()	const { return WindowMode == EWindowMode::Fullscreen || WindowMode == EWindowMode::WindowedFullscreen; }
+	virtual int32 GetCurrentMonitorIndex() const { return -1; }
+	virtual bool IsExclusiveFullscreenSupported() const { return true; }
 	virtual bool IsExclusiveFullscreen() const { return WindowMode == EWindowMode::Fullscreen; }
 	virtual EWindowMode::Type GetWindowMode()	const { return WindowMode; }
 	virtual void ProcessInput( float DeltaTime ) = 0;
diff --git a/Engine/Source/Runtime/Engine/Public/UnrealEngine.h b/Engine/Source/Runtime/Engine/Public/UnrealEngine.h
--- a/Engine/Source/Runtime/Engine/Public/UnrealEngine.h
+++ b/Engine/Source/Runtime/Engine/Public/UnrealEngine.h
@@ -446,27 +446,22 @@ ENGINE_API const FCachedSystemScalabilityCVars& GetCachedScalabilityCVars();
 
 struct FSystemResolution
 {
-	int32 ResX;
-	int32 ResY;
-	EWindowMode::Type WindowMode;
-	bool bForceRefresh;
+	int32 ResX = 0;
+	int32 ResY = 0;
+	EWindowMode::Type WindowMode = EWindowMode::Windowed;
+	int32 MonitorIndex = -1;
+	bool bForceRefresh = false;
+
+	static ENGINE_API void SetMonitorIndex(int32 MonitorIndex);
 
 	// Helper function for changing system resolution via the r.setres console command
 	// This function will set r.setres, which will trigger a resolution change later on
 	// when the console variable sinks are called
 	static ENGINE_API void RequestResolutionChange(int32 InResX, int32 InResY, EWindowMode::Type InWindowMode);
 
-	FSystemResolution()
-		: ResX(0)
-		, ResY(0)
-		, WindowMode(EWindowMode::Windowed)
-		, bForceRefresh(false)
-	{
-
-	}
-
 	void ForceRefresh()
 	{
+		SetMonitorIndex(MonitorIndex);
 		RequestResolutionChange(ResX, ResY, WindowMode);
 		bForceRefresh = true;
 	}
diff --git a/Engine/Source/Runtime/Slate/Private/Framework/Application/SlateApplication.cpp b/Engine/Source/Runtime/Slate/Private/Framework/Application/SlateApplication.cpp
--- a/Engine/Source/Runtime/Slate/Private/Framework/Application/SlateApplication.cpp
+++ b/Engine/Source/Runtime/Slate/Private/Framework/Application/SlateApplication.cpp
@@ -1906,8 +1906,10 @@ FWidgetPath FSlateApplication::LocateWidgetInWindow(UE::Slate::FDeprecateVector2
 		{
 			// Screen space mapping scales everything. When window resolution doesn't match platform resolution, 
 			// this causes offset cursor hit-tests in fullscreen. Correct in slate since we are first window-aware slate processor.
-			FVector2f WindowSize = Window->GetSizeInScreen();
-			FVector2f DisplaySize = { (float)CachedDisplayMetrics.PrimaryDisplayWidth, (float)CachedDisplayMetrics.PrimaryDisplayHeight };
+			const FVector2f WindowSize = Window->GetSizeInScreen();
+			const FVector2f DisplaySize = CachedDisplayMetrics.MonitorInfo.IsEmpty()
+				? FVector2f(CachedDisplayMetrics.PrimaryDisplayWidth, CachedDisplayMetrics.PrimaryDisplayHeight)
+				: Window->GetFullScreenInfo().GetSize2f();
 
 			CursorPosition *= WindowSize / DisplaySize;
 		}
@@ -4664,8 +4666,10 @@ FPointerEvent FSlateApplication::TransformPointerEvent(const FPointerEvent& Poin
 		{
 			// Screen space mapping scales everything. When window resolution doesn't match platform resolution, 
 			// this causes offset cursor hit-tests in fullscreen. Correct in slate since we are first window-aware slate processor.
-			FVector2f WindowSize = Window->GetSizeInScreen();
-			FVector2f DisplaySize = { (float)CachedDisplayMetrics.PrimaryDisplayWidth, (float)CachedDisplayMetrics.PrimaryDisplayHeight };
+			const FVector2f WindowSize = Window->GetSizeInScreen();
+			const FVector2f DisplaySize = CachedDisplayMetrics.MonitorInfo.IsEmpty()
+				? FVector2f(CachedDisplayMetrics.PrimaryDisplayWidth, CachedDisplayMetrics.PrimaryDisplayHeight)
+				: Window->GetFullScreenInfo().GetSize2f();
 
 			TransformedPointerEvent = FPointerEvent(PointerEvent, PointerEvent.GetScreenSpacePosition() * WindowSize / DisplaySize, PointerEvent.GetLastScreenSpacePosition() * WindowSize / DisplaySize);
 		}
@@ -6958,6 +6962,16 @@ void FSlateApplication::HandleDPIScaleChanged(const TSharedRef<FGenericWindow>&
 #endif
 }
 
+void FSlateApplication::OnRegularWindowMovedOrResized(const TSharedRef<FGenericWindow>& PlatformWindow)
+{
+	TSharedPtr< SWindow > SlateWindow = FSlateWindowHelper::FindWindowByPlatformWindow(SlateWindows, PlatformWindow);
+
+	if (SlateWindow.IsValid() && SlateWindow->IsRegularWindow())
+	{
+		OnRegularWindowMovedOrResizedEvent.Broadcast(SlateWindow.ToSharedRef());
+	}
+}
+
 void FSlateApplication::OnMovedWindow( const TSharedRef< FGenericWindow >& PlatformWindow, const int32 X, const int32 Y )
 {
 	TSharedPtr< SWindow > Window = FSlateWindowHelper::FindWindowByPlatformWindow( SlateWindows, PlatformWindow );
diff --git a/Engine/Source/Runtime/Slate/Private/Framework/Application/SlateUser.cpp b/Engine/Source/Runtime/Slate/Private/Framework/Application/SlateUser.cpp
--- a/Engine/Source/Runtime/Slate/Private/Framework/Application/SlateUser.cpp
+++ b/Engine/Source/Runtime/Slate/Private/Framework/Application/SlateUser.cpp
@@ -1019,8 +1019,7 @@ void FSlateUser::LockCursorInternal(const FWidgetPath& WidgetPath)
 		// this causes offset cursor hit-tests in fullscreen. Correct when capturing mouse as viewport widget may be smaller than screen in pixels.
 		if (FSlateApplication::Get().GetTransformFullscreenMouseInput() && !GIsEditor && NativeWindow->GetWindowMode() == EWindowMode::Fullscreen)
 		{
-			FDisplayMetrics CachedDisplayMetrics;
-			FSlateApplication::Get().GetCachedDisplayMetrics(CachedDisplayMetrics);
+			const FDisplayMetrics& CachedDisplayMetrics = FSlateApplication::Get().GetCachedDisplayMetricsByRef();
 			FVector2f DisplaySize = { (float)CachedDisplayMetrics.PrimaryDisplayWidth, (float)CachedDisplayMetrics.PrimaryDisplayHeight };
 			FVector2f DisplayDistortion = SlateClipRect.GetSize() / DisplaySize;
 
diff --git a/Engine/Source/Runtime/Slate/Public/Framework/Application/SlateApplication.h b/Engine/Source/Runtime/Slate/Public/Framework/Application/SlateApplication.h
--- a/Engine/Source/Runtime/Slate/Public/Framework/Application/SlateApplication.h
+++ b/Engine/Source/Runtime/Slate/Public/Framework/Application/SlateApplication.h
@@ -710,6 +710,10 @@ public:
 	 */
 	SLATE_API void ClearKeyboardFocus(const EFocusCause ReasonFocusIsChanging = EFocusCause::SetDirectly);
 
+	/** Gets a delegate that is invoked when a regular window is moved or resized */
+	DECLARE_EVENT_OneParam(FSlateApplication, FOnRegularWindowMovedOrResized, TSharedRef<SWindow>);
+	FOnRegularWindowMovedOrResized& OnRegularWindowMovedOrResized() { return OnRegularWindowMovedOrResizedEvent; }
+
 #if WITH_EDITOR
 	/**
 	* Gets a delegate that is invoked before the input key get process by slate widgets bubble system.
@@ -1699,6 +1703,7 @@ public:
 	SLATE_API virtual void FinishedReshapingWindow( const TSharedRef< FGenericWindow >& PlatformWindow ) override;
 	SLATE_API virtual void SignalSystemDPIChanged(const TSharedRef<FGenericWindow>& Window) override;
 	SLATE_API virtual void HandleDPIScaleChanged(const TSharedRef<FGenericWindow>& Window) override;
+	SLATE_API virtual void OnRegularWindowMovedOrResized(const TSharedRef<FGenericWindow>& Window) override;
 	SLATE_API virtual void OnMovedWindow( const TSharedRef< FGenericWindow >& PlatformWindow, const int32 X, const int32 Y ) override;
 	SLATE_API virtual bool OnWindowActivationChanged( const TSharedRef< FGenericWindow >& PlatformWindow, const EWindowActivation ActivationType ) override;
 	SLATE_API virtual bool OnApplicationActivationChanged( const bool IsActive ) override;
@@ -2273,6 +2278,11 @@ private:
 	/** Allows applications finer control over how we map controllers to users. */
 	TSharedRef<ISlateInputManager> InputManager;
 
+	/**
+	 * Called when a regular window is moved or resized
+	 */
+	FOnRegularWindowMovedOrResized OnRegularWindowMovedOrResizedEvent;
+
 #if WITH_EDITOR
 	/**
 	 * Delegate that is invoked before the input key get process by slate widgets bubble system.
diff --git a/Engine/Source/Runtime/SlateCore/Public/Application/SlateApplicationBase.h b/Engine/Source/Runtime/SlateCore/Public/Application/SlateApplicationBase.h
--- a/Engine/Source/Runtime/SlateCore/Public/Application/SlateApplicationBase.h
+++ b/Engine/Source/Runtime/SlateCore/Public/Application/SlateApplicationBase.h
@@ -263,12 +263,19 @@ public:
 	SLATECORE_API void GetDisplayMetrics(FDisplayMetrics& OutDisplayMetrics);
 
 	/**
-	* Gets the application's cached display metrics.
-	*
-	* @param OutDisplayMetrics Will contain the display metrics.
-	*/
+	 * Gets the application's cached display metrics.
+	 *
+	 * @param OutDisplayMetrics Will contain the display metrics.
+	 */
 	SLATECORE_API void GetCachedDisplayMetrics(FDisplayMetrics& OutDisplayMetrics) const;
 
+	/**
+	 * Gets the application's cached display metrics without copy.
+	 *
+	 * @return The cached display metrics.
+	 */
+	const FDisplayMetrics& GetCachedDisplayMetricsByRef() const { return CachedDisplayMetrics; }
+
 	SLATECORE_API void GetSafeZoneSize(FMargin& SafeZone, const UE::Slate::FDeprecateVector2DParameter& OverrideSize);
 
 	/**
diff --git a/Engine/Source/Runtime/Windows/D3D11RHI/Private/Windows/WindowsD3D11Device.cpp b/Engine/Source/Runtime/Windows/D3D11RHI/Private/Windows/WindowsD3D11Device.cpp
--- a/Engine/Source/Runtime/Windows/D3D11RHI/Private/Windows/WindowsD3D11Device.cpp
+++ b/Engine/Source/Runtime/Windows/D3D11RHI/Private/Windows/WindowsD3D11Device.cpp
@@ -1704,7 +1704,7 @@ void FD3D11DynamicRHI::InitD3DDevice()
 	}
 }
 
-static bool GetAvailableResolutionsForOutput(IDXGIOutput* Output, ID3D11Device* Device, int32 CurrentOutput, FScreenResolutionArray& Resolutions, bool bIgnoreRefreshRate)
+static void GetAvailableResolutionsForOutput(IDXGIOutput* Output, ID3D11Device* Device, int32 CurrentOutput, FScreenResolutionArray& Resolutions, bool bIgnoreRefreshRate)
 {
 	// TODO: GetDisplayModeList is a terribly SLOW call.  It can take up to a second per invocation.
 	//  We might want to work around some DXGI badness here.
@@ -1726,12 +1726,12 @@ static bool GetAvailableResolutionsForOutput(IDXGIOutput* Output, ID3D11Device*
 			else if (HResult == DXGI_ERROR_NOT_CURRENTLY_AVAILABLE)
 			{
 				UE_LOG(LogD3D11RHI, Warning, TEXT("RHIGetAvailableResolutions does not return results when running under remote desktop."));
-				return false;
+				return;
 			}
 			else
 			{
 				UE_LOG(LogD3D11RHI, Warning, TEXT("RHIGetAvailableResolutions failed with unknown error (0x%x)."), HResult);
-				return false;
+				return;
 			}
 		}
 		else
@@ -1759,7 +1759,7 @@ static bool GetAvailableResolutionsForOutput(IDXGIOutput* Output, ID3D11Device*
 			OutputDesc.AttachedToDesktop,
 			*GetD3D11ErrorString(OutputDescRes, Device), OutputDescRes);
 
-		return false;
+		return;
 	}
 
 	if (NumModes > 0)
@@ -1829,8 +1829,6 @@ static bool GetAvailableResolutionsForOutput(IDXGIOutput* Output, ID3D11Device*
 	{
 		UE_LOG(LogD3D11RHI, Warning, TEXT("No display modes found for the standard format DXGI_FORMAT_R8G8B8A8_UNORM!"));
 	}
-
-	return !Resolutions.IsEmpty();
 }
 
 /**
@@ -1846,7 +1844,11 @@ bool FD3D11DynamicRHI::RHIGetAvailableResolutions(FScreenResolutionArray& Resolu
 	// TODO: The resolution information can be stale if the adapter's factory isn't current. Use RHIGetAvailableResolutionsForDisplay instead.
 	TRefCountPtr<IDXGIOutput> Output;
 	const HRESULT HResult = Adapter.DXGIAdapter->EnumOutputs(0, Output.GetInitReference());
-	return SUCCEEDED(HResult) ? GetAvailableResolutionsForOutput(Output, GetDevice(), 0, Resolutions, bIgnoreRefreshRate) : false;
+	if (SUCCEEDED(HResult))
+	{
+		GetAvailableResolutionsForOutput(Output, GetDevice(), 0, Resolutions, bIgnoreRefreshRate);
+	}
+	return !Resolutions.IsEmpty();
 }
 
 bool FD3D11DynamicRHI::RHIGetAvailableResolutionsForDisplay(FScreenResolutionArray& Resolutions, bool bIgnoreRefreshRate, void* NativeMonitorHandle)
@@ -1884,7 +1886,26 @@ bool FD3D11DynamicRHI::RHIGetAvailableResolutionsForDisplay(FScreenResolutionArr
 			check(SUCCEEDED(OutputDescRes));
 			if (MonitorHandle == OutputDesc.Monitor)
 			{
-				return GetAvailableResolutionsForOutput(Output, GetDevice(), CurrentOutput, Resolutions, bIgnoreRefreshRate);
+				GetAvailableResolutionsForOutput(Output, GetDevice(), CurrentOutput, Resolutions, bIgnoreRefreshRate);
+
+				// Only allow the desktop resolution if SetFullscreenState is expected to fail to allow borderless to seemlessly take over.
+				// Undocumented but assume exclusive only works if the the display is primary, i.e. top-left at (0, 0), or the display is connected to the chosen adapter.
+				if (OutputDesc.DesktopCoordinates.left != 0 || OutputDesc.DesktopCoordinates.top != 0)
+				{
+					const FD3D11Adapter& ChosenAdapter = GetAdapter();
+					DXGI_ADAPTER_DESC AdapterDesc;
+					if (SUCCEEDED(DXGIAdapter->GetDesc(&AdapterDesc)) &&
+						FMemory::Memcmp(&ChosenAdapter.DXGIAdapterDesc.AdapterLuid, &AdapterDesc.AdapterLuid, sizeof(LUID)))
+					{
+						Resolutions.Reset(1);
+						FScreenResolutionRHI& ScreenResolution = Resolutions.Emplace_GetRef();
+						ScreenResolution.Width = OutputDesc.DesktopCoordinates.right - OutputDesc.DesktopCoordinates.left;
+						ScreenResolution.Height = OutputDesc.DesktopCoordinates.bottom - OutputDesc.DesktopCoordinates.top;
+						ScreenResolution.RefreshRate = 0;
+					}
+				}
+
+				return !Resolutions.IsEmpty();
 			}
 		}
 	}
