Suppress V614 from PVS Studio for Win32 shenanigans where viewGdiDeviceName is initialized but it's not obvious

#Synced-CL 49751596

diff --git a/Engine/Source/Runtime/ApplicationCore/Private/Windows/WindowsApplication.cpp b/Engine/Source/Runtime/ApplicationCore/Private/Windows/WindowsApplication.cpp
--- a/Engine/Source/Runtime/ApplicationCore/Private/Windows/WindowsApplication.cpp
+++ b/Engine/Source/Runtime/ApplicationCore/Private/Windows/WindowsApplication.cpp
@@ -1968,7 +1968,7 @@ static void GetFriendlyNameDisplayDevices(const WCHAR DeviceName[CCHDEVICENAME],
 			SourceName.header.adapterId = Path.sourceInfo.adapterId;
 			SourceName.header.id = Path.sourceInfo.id;
 			if ((DisplayConfigGetDeviceInfo(&SourceName.header) == ERROR_SUCCESS) &&
-				(FCString::Strcmp(DeviceName, SourceName.viewGdiDeviceName) == 0))
+				(FCString::Strcmp(DeviceName, SourceName.viewGdiDeviceName) == 0)) //-V614
 			{
 				DISPLAYCONFIG_TARGET_DEVICE_NAME TargetName;
 				TargetName.header.type = DISPLAYCONFIG_DEVICE_INFO_GET_TARGET_NAME;
