Fixed non unity mac compile errors

diff --git a/Engine/Source/Runtime/ApplicationCore/Public/Mac/CocoaTextView.h b/Engine/Source/Runtime/ApplicationCore/Public/Mac/CocoaTextView.h
--- a/Engine/Source/Runtime/ApplicationCore/Public/Mac/CocoaTextView.h
+++ b/Engine/Source/Runtime/ApplicationCore/Public/Mac/CocoaTextView.h
@@ -2,6 +2,7 @@
 #pragma once
 
 #include "CoreMinimal.h"
+#include "Mac/MacSystemIncludes.h"
 #include "GenericPlatform/ITextInputMethodSystem.h"
 #if WITH_ACCESSIBILITY
 #include "Mac/Accessibility/CocoaAccessibilityView.h"
