Changing iOS default `preferredScreenEdgesDeferringSystemGestures` to `UIRectEdgeAll` to allow touches in all the screen area without delays from the OS.

#rnx
#Synced-CL 50583260

diff --git a/Engine/Source/Runtime/ApplicationCore/Private/IOS/IOSView.cpp b/Engine/Source/Runtime/ApplicationCore/Private/IOS/IOSView.cpp
--- a/Engine/Source/Runtime/ApplicationCore/Private/IOS/IOSView.cpp
+++ b/Engine/Source/Runtime/ApplicationCore/Private/IOS/IOSView.cpp
@@ -1149,6 +1149,8 @@ self.accessibilityElements = @[Window.accessibilityContainer];
 #if !PLATFORM_TVOS
 	self.view.multipleTouchEnabled = NO;
 #endif
+
+	[self setNeedsUpdateOfScreenEdgesDeferringSystemGestures];
 }
 
 /**
@@ -1277,7 +1279,7 @@ self.accessibilityElements = @[Window.accessibilityContainer];
 
 - (UIRectEdge)preferredScreenEdgesDeferringSystemGestures
 {
-	return UIRectEdgeBottom;
+	return UIRectEdgeAll;
 }
 
 - (void)viewWillTransitionToSize : (CGSize)size withTransitionCoordinator : (id<UIViewControllerTransitionCoordinator>)coordinator
