Make splash screens extensible with custom text and images.

#rnx

#Synced-CL 50656828

diff --git a/Engine/Source/Runtime/ApplicationCore/Private/GenericPlatform/GenericPlatformSplash.cpp b/Engine/Source/Runtime/ApplicationCore/Private/GenericPlatform/GenericPlatformSplash.cpp
--- a/Engine/Source/Runtime/ApplicationCore/Private/GenericPlatform/GenericPlatformSplash.cpp
+++ b/Engine/Source/Runtime/ApplicationCore/Private/GenericPlatform/GenericPlatformSplash.cpp
@@ -60,18 +60,29 @@ void FGenericPlatformSplash::SetSplashText(const SplashTextType::Type InType, co
 *
 * @return true if a splash screen was found
 */
-bool FGenericPlatformSplash::GetSplashPath(const TCHAR* SplashFilename, FString& OutPath, bool& OutIsCustom)
+bool FGenericPlatformSplash::GetSplashPath(const TCHAR* SplashFilename, FString& OutPath, bool& OutIsCustom,
+	const bool bAllowCustom)
 {
 	FString Filename = FString("Splash/") + SplashFilename;
 
-	if(GCustomSplashScreenFileName.IsEmpty())
+	if(GCustomSplashScreenFileName.IsEmpty() || !bAllowCustom)
 	{
+		// let's check if maybe the incoming path is already a fully qualified file name
+		if (FPaths::FileExists(SplashFilename))
+		{
+			OutPath = SplashFilename;
+			OutIsCustom = false;
+
+			// if this was found, then we're done
+			return true;
+		}
+
 		// first look in game's splash directory
 		OutPath = FPaths::ConvertRelativePathToFull(GetSplashFilename(FPaths::ProjectContentDir(), Filename));
 		OutIsCustom = true;
 
 		// if this was found, then we're done
-		if (IFileManager::Get().FileSize(*OutPath) != -1)
+		if (FPaths::FileExists(*OutPath))
 		{
 			return true;
 		}
@@ -81,7 +92,7 @@ bool FGenericPlatformSplash::GetSplashPath(const TCHAR* SplashFilename, FString&
 		OutIsCustom = false;
 
 		// if this was found, then we're done
-		if (IFileManager::Get().FileSize(*OutPath) != -1)
+		if (FPaths::FileExists(*OutPath))
 		{
 			return true;
 		}
@@ -92,7 +103,7 @@ bool FGenericPlatformSplash::GetSplashPath(const TCHAR* SplashFilename, FString&
 		OutPath = GCustomSplashScreenFileName;
 		OutIsCustom = true;
 		// if this was found, then we're done
-		if (IFileManager::Get().FileSize(*OutPath) != -1)
+		if (FPaths::FileExists(*OutPath))
 		{
 			return true;
 		}
diff --git a/Engine/Source/Runtime/ApplicationCore/Private/Windows/WindowsPlatformSplash.cpp b/Engine/Source/Runtime/ApplicationCore/Private/Windows/WindowsPlatformSplash.cpp
--- a/Engine/Source/Runtime/ApplicationCore/Private/Windows/WindowsPlatformSplash.cpp
+++ b/Engine/Source/Runtime/ApplicationCore/Private/Windows/WindowsPlatformSplash.cpp
@@ -48,13 +48,17 @@ THIRD_PARTY_INCLUDES_END
  */
 
 static HANDLE GSplashScreenThread = NULL;
-static HBITMAP GSplashScreenBitmap = NULL;
 static HWND GSplashScreenWnd = NULL; 
 static HWND GSplashScreenGuard = NULL; 
-static FString GSplashScreenFileName;
+static FString GSplashScreenFileNames[ SplashImageType::NumImageTypes ];
+static HBITMAP GSplashScreenBitmaps[ SplashImageType::NumImageTypes ];
+static FIntPoint GSplashScreenPos[ SplashImageType::NumImageTypes ];
 static FText GSplashScreenAppName;
 static FText GSplashScreenText[ SplashTextType::NumTextTypes ];
 static RECT GSplashScreenTextRects[ SplashTextType::NumTextTypes ];
+static ESplashTextSizeType GSplashScreenTextSizes[ SplashTextType::NumTextTypes ];
+static ESplashTextBrightnessType GSplashScreenTextBrightness[ SplashTextType::NumTextTypes ];
+static ESplashTextFlags GSplashScreenTextFlags[ SplashTextType::NumTextTypes ];
 static HFONT GSplashScreenSmallTextFontHandle = NULL;
 static HFONT GSplashScreenNormalTextFontHandle = NULL;
 static HFONT GSplashScreenTitleTextFontHandle = NULL;
@@ -75,20 +79,27 @@ LRESULT CALLBACK SplashScreenWindowProc(HWND hWnd, uint32 message, WPARAM wParam
 		case WM_PAINT:
 			{
 				// We can continue to receive messages for a period after we've torn down the bitmap resource
-				if (!GSplashScreenBitmap)
+				if (!GSplashScreenBitmaps[SplashImageType::Background])
 				{
 					return 0;
 				}
 
 				hdc = BeginPaint(hWnd, &ps);
 
-				// Draw splash bitmap
-				DrawState(hdc, DSS_NORMAL, NULL, (LPARAM)GSplashScreenBitmap, 0, 0, 0, 0, 0, DST_BITMAP);
-
 				{
-					// Take a critical section since another thread may be trying to set the splash text
+					// Take a critical section since another thread may be trying to set a splash text or image
 					FScopeLock ScopeLock( &GSplashScreenSynchronizationObject );
 
+					// Draw splash bitmap
+					for (int32 i = 0; i < SplashImageType::NumImageTypes; ++i)
+					{
+						if (HBITMAP BitmapHandle = GSplashScreenBitmaps[i])
+						{
+							const FIntPoint Pos = GSplashScreenPos[i];
+							DrawState(hdc, DSS_NORMAL, NULL, (LPARAM)BitmapHandle, 0, Pos.X, Pos.Y, 0, 0, DST_BITMAP);
+						}
+					}
+
 					// Draw splash text
 					for( int32 CurTypeIndex = 0; CurTypeIndex < SplashTextType::NumTextTypes; ++CurTypeIndex )
 					{
@@ -97,11 +108,11 @@ LRESULT CALLBACK SplashScreenWindowProc(HWND hWnd, uint32 message, WPARAM wParam
 
 						if( !SplashText.IsEmpty() )
 						{
-							if ( CurTypeIndex == SplashTextType::VersionInfo1 || CurTypeIndex == SplashTextType::StartupProgress )
+							if ( GSplashScreenTextSizes[CurTypeIndex] == ESplashTextSizeType::Normal )
 							{
 								SelectObject( hdc, GSplashScreenNormalTextFontHandle );
 							}
-							else if ( CurTypeIndex == SplashTextType::GameName )
+							else if ( GSplashScreenTextSizes[CurTypeIndex] == ESplashTextSizeType::Title )
 							{
 								SelectObject( hdc, GSplashScreenTitleTextFontHandle );
 							}
@@ -119,6 +130,9 @@ LRESULT CALLBACK SplashScreenWindowProc(HWND hWnd, uint32 message, WPARAM wParam
 							RECT ClientRect;
 							GetClientRect( hWnd, &ClientRect );
 
+							FString String = SplashText.ToString();
+							const ESplashTextFlags TextFlags = GSplashScreenTextFlags[CurTypeIndex];
+
 #if UE_WINDOWS_SPLASH_USE_TEXT_OUTLINE
 							// Draw background text passes
 							const int32 NumBGPasses = 8;
@@ -139,39 +153,48 @@ LRESULT CALLBACK SplashScreenWindowProc(HWND hWnd, uint32 message, WPARAM wParam
 								}
 
 								SetTextColor( hdc, 0x00000000 );
-								TextOut(
-									hdc,
+
+								RECT OffsetRect = RECT(
 									TextRect.left + BGXOffset,
 									TextRect.top + BGYOffset,
-									*SplashText.ToString(),
-									SplashText.ToString().Len() );
+									TextRect.right + BGXOffset,
+									TextRect.bottom + BGYOffset
+								);
+
+								if (EnumHasAllFlags(TextFlags, ESplashTextFlags::WordWrap))
+								{
+									DrawTextW(hdc, *String, String.Len(), &OffsetRect, DT_WORDBREAK | DT_LEFT | DT_TOP | DT_NOPREFIX);
+								}
+								else
+								{
+									TextOut(hdc, OffsetRect.left, OffsetRect.top, *String, String.Len());
+								}
 							}
 #endif // UE_WINDOWS_SPLASH_USE_TEXT_OUTLINE
 							
 							// Draw foreground text pass
-							if( CurTypeIndex == SplashTextType::StartupProgress )
+							if( GSplashScreenTextBrightness[CurTypeIndex] == ESplashTextBrightnessType::Low )
 							{
 								SetTextColor( hdc, RGB(160, 160, 160) );
 							}
-							else if( CurTypeIndex == SplashTextType::VersionInfo1 )
+							else if( GSplashScreenTextBrightness[CurTypeIndex] == ESplashTextBrightnessType::Medium )
 							{
-								SetTextColor( hdc, RGB(160, 160, 160) );
+								SetTextColor( hdc, RGB(180, 180, 180) );
 							}
-							else if ( CurTypeIndex == SplashTextType::GameName )
+							else
 							{
 								SetTextColor(hdc, RGB(255, 255, 255));
 							}
+
+							if (EnumHasAllFlags(TextFlags, ESplashTextFlags::WordWrap))
+							{
+								RECT DrawRect = TextRect;
+								DrawTextW(hdc, *String, String.Len(), &DrawRect, DT_WORDBREAK | DT_LEFT | DT_TOP);
+							}
 							else
 							{
-								SetTextColor( hdc, RGB(160, 160, 160) );
+								TextOut(hdc, TextRect.left, TextRect.top, *String, String.Len());
 							}
-
-							TextOut(
-								hdc,
-								TextRect.left,
-								TextRect.top,
-								*SplashText.ToString(),
-								SplashText.ToString().Len() );
 						}
 					}
 				}
@@ -210,8 +233,13 @@ LRESULT CALLBACK SplashScreenWindowProc(HWND hWnd, uint32 message, WPARAM wParam
  *  - colour space information or newer format extensions.
  * This code is largely taken from the WicViewerGDI sample provided by Microsoft on MSDN.
  */
-HBITMAP LoadSplashBitmap()
+HBITMAP LoadSplashBitmap(SplashImageType::Type Type)
 {
+	if (GSplashScreenFileNames[Type].IsEmpty())
+	{
+		return nullptr;
+	}
+
 	HRESULT hr = CoInitialize(NULL);
 
 	// The factory pointer
@@ -230,7 +258,7 @@ HBITMAP LoadSplashBitmap()
 
 	// Create a decoder
 	hr = Factory->CreateDecoderFromFilename(
-		(LPCTSTR)*GSplashScreenFileName, // Image to be decoded
+		(LPCTSTR)*GSplashScreenFileNames[Type], // Image to be decoded
 		NULL,                            // Do not prefer a particular vendor
 		GENERIC_READ,                    // Desired read access to the file
 		WICDecodeMetadataCacheOnDemand,  // Cache metadata when needed
@@ -440,10 +468,14 @@ uint32 WINAPI StartSplashScreenThread( LPVOID unused )
 		} 
 
 		// Load splash screen image, display it and handle all window's messages
-		GSplashScreenBitmap = LoadSplashBitmap();
-		if (GSplashScreenBitmap)
+		for (int32 i = 0; i < SplashImageType::NumImageTypes; ++i)
+		{
+			GSplashScreenBitmaps[i] = LoadSplashBitmap(static_cast<SplashImageType::Type>(i));
+		}
+
+		if (GSplashScreenBitmaps[SplashImageType::Background])
 		{
-			GetObjectW(GSplashScreenBitmap, sizeof(bm), &bm);
+			GetObjectW(GSplashScreenBitmaps[SplashImageType::Background], sizeof(bm), &bm);
 			const int32 WindowWidth = bm.bmWidth;
 			const int32 WindowHeight = bm.bmHeight;
 			int32 ScreenPosX = (GetSystemMetrics(SM_CXSCREEN) - WindowWidth) / 2;
@@ -474,7 +506,7 @@ uint32 WINAPI StartSplashScreenThread( LPVOID unused )
 
 	if (GSplashScreenWnd) // also implies GSplashScreenBitmap is not null.
 	{
-		check(GSplashScreenBitmap);
+		check(GSplashScreenBitmaps[SplashImageType::Background]);
 		if (bAllowFading)
 		{
 			// Set window to fully transparent to start out
@@ -551,24 +583,33 @@ uint32 WINAPI StartSplashScreenThread( LPVOID unused )
 			GSplashScreenTextRects[SplashTextType::GameName].left = 10;
 			GSplashScreenTextRects[SplashTextType::GameName].right = bm.bmWidth - 20;
 		}
+
+		GSplashScreenTextSizes[ SplashTextType::GameName ] = ESplashTextSizeType::Title;
+		GSplashScreenTextBrightness[ SplashTextType::GameName ] = ESplashTextBrightnessType::High;
 		
 		// Setup bounds for version info text 1
 		GSplashScreenTextRects[ SplashTextType::VersionInfo1 ].top = bm.bmHeight - 40;
 		GSplashScreenTextRects[ SplashTextType::VersionInfo1 ].bottom = bm.bmHeight - 20;
 		GSplashScreenTextRects[ SplashTextType::VersionInfo1 ].left = 10;
 		GSplashScreenTextRects[ SplashTextType::VersionInfo1 ].right = bm.bmWidth - 20;
+		GSplashScreenTextSizes[ SplashTextType::VersionInfo1 ] = ESplashTextSizeType::Normal;
+		GSplashScreenTextBrightness[ SplashTextType::VersionInfo1 ] = ESplashTextBrightnessType::Low;
 
 		// Setup bounds for copyright info text
 		GSplashScreenTextRects[ SplashTextType::CopyrightInfo ].top = bm.bmHeight - 16;
 		GSplashScreenTextRects[ SplashTextType::CopyrightInfo ].bottom = bm.bmHeight - 6;
 		GSplashScreenTextRects[ SplashTextType::CopyrightInfo ].left = bm.bmWidth - 180;
 		GSplashScreenTextRects[ SplashTextType::CopyrightInfo ].right = bm.bmWidth - 20;
+		GSplashScreenTextSizes[ SplashTextType::CopyrightInfo ] = ESplashTextSizeType::Small;
+		GSplashScreenTextBrightness[ SplashTextType::CopyrightInfo ] = ESplashTextBrightnessType::Low;
 
 		// Setup bounds for startup progress text
 		GSplashScreenTextRects[ SplashTextType::StartupProgress ].top = bm.bmHeight - 20;
 		GSplashScreenTextRects[ SplashTextType::StartupProgress ].bottom = bm.bmHeight;
 		GSplashScreenTextRects[ SplashTextType::StartupProgress ].left = 10;
 		GSplashScreenTextRects[ SplashTextType::StartupProgress ].right = bm.bmWidth - 20;
+		GSplashScreenTextSizes[ SplashTextType::StartupProgress ] = ESplashTextSizeType::Normal;
+		GSplashScreenTextBrightness[ SplashTextType::StartupProgress ] = ESplashTextBrightnessType::Low;
 
 		SetWindowText(GSplashScreenWnd, *GSplashScreenAppName.ToString());
 		ShowWindow(GSplashScreenWnd, SW_SHOW); 
@@ -617,8 +658,14 @@ uint32 WINAPI StartSplashScreenThread( LPVOID unused )
 			}
 		}
 
-		ensure(DeleteObject(GSplashScreenBitmap));
-		GSplashScreenBitmap = NULL;
+		for (int32 i = 0; i < SplashImageType::NumImageTypes; ++i)
+		{
+			if (GSplashScreenBitmaps[i])
+			{
+				ensure(DeleteObject(GSplashScreenBitmaps[i]));
+				GSplashScreenBitmaps[i] = NULL;
+			}
+		}
 	}
 
 	UnregisterClass(wc.lpszClassName, hInstance);
@@ -718,7 +765,7 @@ void FWindowsPlatformSplash::Show()
 			}
 
 			GSplashWindowCreationEvent = CreateEvent(NULL, FALSE, FALSE, NULL);
-			GSplashScreenFileName = SplashPath;
+			GSplashScreenFileNames[SplashImageType::Background] = MoveTemp(SplashPath);
 			DWORD ThreadID = 0;
 			GSplashScreenThread = CreateThread(NULL, 128 * 1024, (LPTHREAD_START_ROUTINE)StartSplashScreenThread, (LPVOID)NULL, STACK_SIZE_PARAM_IS_A_RESERVATION, &ThreadID);
 #if	STATS
@@ -768,6 +815,29 @@ bool FWindowsPlatformSplash::IsShown()
 	return (GSplashScreenThread != nullptr);
 }
 
+void FWindowsPlatformSplash::RegisterUserSplashText(const SplashTextType::Type InType, const FPlatformRect& InRect,
+	const ESplashTextSizeType InSize, const ESplashTextBrightnessType InBrightness, const ESplashTextFlags InFlags)
+{
+	check(InType >= SplashTextType::Type::User1);
+	check(!GSplashScreenThread);
+
+	GSplashScreenTextRects[InType] = RECT(InRect.Left, InRect.Top, InRect.Right, InRect.Bottom);
+	GSplashScreenTextSizes[InType] = InSize;
+	GSplashScreenTextBrightness[InType] = InBrightness;
+	GSplashScreenTextFlags[InType] = InFlags;
+}
+
+void FWindowsPlatformSplash::RegisterUserSplashImage(const SplashImageType::Type InType, const FIntPoint& InPos,
+	const TCHAR* InFilename)
+{
+	check(InType >= SplashImageType::User1);
+	check(!GSplashScreenThread);
+
+	bool bCustom;
+	GetSplashPath(InFilename, GSplashScreenFileNames[InType], bCustom, false);
+	GSplashScreenPos[InType] = InPos;
+}
+
 void FWindowsPlatformSplash::SetProgress(int ProgressPercent)
 {
 	extern FWindowsApplication* WindowsApplication;
@@ -819,6 +889,41 @@ void FWindowsPlatformSplash::SetSplashText( const SplashTextType::Type InType, c
 	}
 }
 
+void FWindowsPlatformSplash::SetSplashImage(const SplashImageType::Type InType, const TCHAR* InFilename)
+{
+	if (GSplashScreenThread)
+	{
+		const HBITMAP ImageHandle = LoadSplashBitmap(InType);
+		if (!ImageHandle)
+		{
+			return;
+		}
+
+		// Take a critical section since the splash thread may already be repainting using this image
+		FScopeLock ScopeLock(&GSplashScreenSynchronizationObject);
+		GSplashScreenBitmaps[InType] = ImageHandle;
+
+		SIZE Size;
+		if (GetBitmapDimensionEx(ImageHandle, &Size))
+		{
+			// Repaint the window
+			RECT Rect;
+			Rect.left = GSplashScreenPos[InType].X;
+			Rect.top = GSplashScreenPos[InType].Y;
+			Rect.right = GSplashScreenPos[InType].X + Size.cx;
+			Rect.bottom = GSplashScreenPos[InType].Y + Size.cy;
+
+			constexpr BOOL bErase = false;
+			InvalidateRect(GSplashScreenWnd, &Rect, bErase);
+		}
+	}
+	else // No drawing is running yet so we can just update the path here
+	{
+		bool bCustom;
+		GetSplashPath(InFilename, GSplashScreenFileNames[InType], bCustom, false);
+	}
+}
+
 APPLICATIONCORE_API HWND GetSplashScreenWindowHandle()
 {
 	return GSplashScreenWnd;
diff --git a/Engine/Source/Runtime/ApplicationCore/Public/GenericPlatform/GenericPlatformSplash.h b/Engine/Source/Runtime/ApplicationCore/Public/GenericPlatform/GenericPlatformSplash.h
--- a/Engine/Source/Runtime/ApplicationCore/Public/GenericPlatform/GenericPlatformSplash.h
+++ b/Engine/Source/Runtime/ApplicationCore/Public/GenericPlatform/GenericPlatformSplash.h
@@ -4,6 +4,9 @@
 
 #include "Containers/ContainersFwd.h"
 #include "CoreTypes.h"
+#include "Math/IntPoint.h"
+
+struct FPlatformRect;
 
 /**
  * SplashTextType defines the types of text on the splash screen
@@ -26,11 +29,69 @@ namespace SplashTextType
 
 		// ...
 
+		/** Licensee-usable slots */
+		User1,
+		User2,
+		User3,
+		User4,
+		User5,
+		User6,
+
+		// Only add further User types here, the rest should go above this block
+
 		/** Number of text types (must be final enum value) */
 		NumTextTypes
 	};
 }
 
+/**
+ * SplashImageType defines the types of images on the splash screen
+ */
+namespace SplashImageType
+{
+	enum Type
+	{
+		/** The main background image */
+		Background,
+
+		// ...
+
+		/** Licensee-usable slots */
+		User1,
+		User2,
+		User3,
+		User4,
+		User5,
+		User6,
+
+		// Only add further User types here, the rest should go above this block
+
+		/** Number of image types (must be final enum value) */
+		NumImageTypes
+	};
+}
+
+enum class ESplashTextSizeType : uint8
+{
+	Small,
+	Normal,
+	Title,
+};
+
+enum class ESplashTextBrightnessType : uint8
+{
+	Low,
+	Medium,
+	High,
+};
+
+enum class ESplashTextFlags : uint8
+{
+	None,
+	WordWrap = 1 << 0,
+};
+ENUM_CLASS_FLAGS(ESplashTextFlags)
+
 class FString;
 class FText;
 
@@ -56,7 +117,7 @@ struct FGenericPlatformSplash
 	/**
 	 * Sets the progress displayed on the application icon (for startup/loading progress).
 	 *
-	 * @param InType Progress value in percent.
+	 * @param ProgressPercent Progress value in percent.
 	 */
 	inline static void SetProgress(int ProgressPercent)
 	{
@@ -75,6 +136,15 @@ struct FGenericPlatformSplash
 	{
 	}
 
+	/** Sets the image to be assigned to the specific slot */
+	inline static void SetSplashImage(const SplashImageType::Type InType, const TCHAR* InFilename) { }
+
+	/** Registers a new licensee text. Only User types are allowed, and only before the splash screen is initialized! */
+	inline static void RegisterUserSplashText(const SplashTextType::Type InType, const FPlatformRect& InRect, ESplashTextSizeType InSize, ESplashTextBrightnessType InBrightness, ESplashTextFlags Flags = ESplashTextFlags::None) { }
+
+	/** Registers a new licensee image. Only User types are allowed, and only before the splash screen is initialized! */
+	inline static void RegisterUserSplashImage(const SplashImageType::Type InType, const FIntPoint& InPos, const TCHAR* InFilename) { }
+
 	/**
 	 * Return whether the splash screen is being shown or not
 	 */
@@ -94,6 +164,6 @@ protected:
 	*
 	* @return true if a splash screen was found
 	*/
-	static APPLICATIONCORE_API bool GetSplashPath(const TCHAR* SplashFilename, FString& OutPath, bool& OutIsCustom);
+	static APPLICATIONCORE_API bool GetSplashPath(const TCHAR* SplashFilename, FString& OutPath, bool& OutIsCustom, bool bAllowCustom = true);
 	static APPLICATIONCORE_API bool GetSplashPath(const TCHAR* SplashFilename, const TCHAR* IconFilename, FString& OutPath, FString& OutIconPath, bool& OutIsCustom);
 };
diff --git a/Engine/Source/Runtime/ApplicationCore/Public/Windows/WindowsPlatformSplash.h b/Engine/Source/Runtime/ApplicationCore/Public/Windows/WindowsPlatformSplash.h
--- a/Engine/Source/Runtime/ApplicationCore/Public/Windows/WindowsPlatformSplash.h
+++ b/Engine/Source/Runtime/ApplicationCore/Public/Windows/WindowsPlatformSplash.h
@@ -20,7 +20,7 @@ struct FWindowsPlatformSplash
 	/**
 	 * Sets the progress displayed on the application icon (for startup/loading progress).
 	 *
-	 * @param InType Progress value in percent.
+	 * @param ProgressPercent Progress value in percent.
 	 */
 	static APPLICATIONCORE_API void SetProgress(int ProgressPercent);
 
@@ -33,6 +33,15 @@ struct FWindowsPlatformSplash
 	using FGenericPlatformSplash::SetSplashText;
 	static APPLICATIONCORE_API void SetSplashText( const SplashTextType::Type InType, const FText& InText );
 
+	/** Sets the image to be assigned to the specific slot */
+	static APPLICATIONCORE_API void SetSplashImage(const SplashImageType::Type InType, const TCHAR* InFilename);
+
+	/** Registers a new licensee text. Only User types are allowed, and only before the splash screen is initialized! */
+	static APPLICATIONCORE_API void RegisterUserSplashText(const SplashTextType::Type InType, const FPlatformRect& InRect, ESplashTextSizeType InSize, ESplashTextBrightnessType InBrightness, ESplashTextFlags InFlags = ESplashTextFlags::None);
+
+	/** Registers a new licensee image. Only User types are allowed, and only before the splash screen is initialized! */
+	static APPLICATIONCORE_API void RegisterUserSplashImage(const SplashImageType::Type InType, const FIntPoint& InPos, const TCHAR* InFilename);
+
 	/**
 	 * Return whether the splash screen is being shown or not
 	 */
