Fix potential memory leak when changing a splash screen image.

#Synced-CL 50656828

diff --git a/Engine/Source/Runtime/ApplicationCore/Private/Windows/WindowsPlatformSplash.cpp b/Engine/Source/Runtime/ApplicationCore/Private/Windows/WindowsPlatformSplash.cpp
--- a/Engine/Source/Runtime/ApplicationCore/Private/Windows/WindowsPlatformSplash.cpp
+++ b/Engine/Source/Runtime/ApplicationCore/Private/Windows/WindowsPlatformSplash.cpp
@@ -901,6 +901,12 @@ void FWindowsPlatformSplash::SetSplashImage(const SplashImageType::Type InType,
 
 		// Take a critical section since the splash thread may already be repainting using this image
 		FScopeLock ScopeLock(&GSplashScreenSynchronizationObject);
+
+        if (GSplashScreenBitmaps[InType])
+        {
+            ensure(DeleteObject(GSplashScreenBitmaps[InType]));
+        }
+
 		GSplashScreenBitmaps[InType] = ImageHandle;
 
 		SIZE Size;
