Allow IOS frame pacer to continually signal the RHI thread during Exit so that the app doesn't hang due to the timing of the wait event firing

#rnx

#Synced-CL 50477682

diff --git a/Engine/Source/Runtime/ApplicationCore/Private/IOS/IOSPlatformFramePacer.cpp b/Engine/Source/Runtime/ApplicationCore/Private/IOS/IOSPlatformFramePacer.cpp
--- a/Engine/Source/Runtime/ApplicationCore/Private/IOS/IOSPlatformFramePacer.cpp
+++ b/Engine/Source/Runtime/ApplicationCore/Private/IOS/IOSPlatformFramePacer.cpp
@@ -72,7 +72,8 @@ namespace IOSDisplayConstants
 			NextEvent->Trigger();
 		}
 		
-		ListeningEvents.Empty();
+		// Allow ListeningEvents to keep firing indefinitely during Engine Exit so that the
+		// RHIThread never hangs waiting due to frame limiting in FMetalDevice::EndDrawingViewport
 
 		return;
 	};
