Adding support for ConfigRules to Windows.

[RN] rnx

#Synced-CL 50573404

diff --git a/Engine/Plugins/Runtime/WindowsDeviceProfileSelector/Source/WindowsDeviceProfileSelector/Private/WindowsDeviceProfileSelectorModule.cpp b/Engine/Plugins/Runtime/WindowsDeviceProfileSelector/Source/WindowsDeviceProfileSelector/Private/WindowsDeviceProfileSelectorModule.cpp
--- a/Engine/Plugins/Runtime/WindowsDeviceProfileSelector/Source/WindowsDeviceProfileSelector/Private/WindowsDeviceProfileSelectorModule.cpp
+++ b/Engine/Plugins/Runtime/WindowsDeviceProfileSelector/Source/WindowsDeviceProfileSelector/Private/WindowsDeviceProfileSelectorModule.cpp
@@ -56,6 +56,14 @@ const FString FWindowsDeviceProfileSelectorModule::GetRuntimeDeviceProfileName()
 			}
 		}
 	}
+	
+	// Use override from ConfigRules if set
+	const FString* ConfigProfile = FPlatformMisc::GetConfigRulesVariable(TEXT("Profile"));
+	if (ConfigProfile != nullptr)
+	{
+		ProfileName = *ConfigProfile;
+		UE_LOG(LogWindows, Log, TEXT("Using ConfigRules Profile: [%s]"), *ProfileName);
+	}
 
 	return ProfileName;
 }
diff --git a/Engine/Source/Runtime/ApplicationCore/Private/Windows/WindowsConfigRules.cpp b/Engine/Source/Runtime/ApplicationCore/Private/Windows/WindowsConfigRules.cpp
new file mode 100644
--- /dev/null
+++ b/Engine/Source/Runtime/ApplicationCore/Private/Windows/WindowsConfigRules.cpp
@@ -0,0 +1,500 @@
+// Copyright Epic Games, Inc. All Rights Reserved.
+
+#include "Windows/WindowsConfigRules.h"
+
+#if WITH_WINDOWS_CONFIGRULES
+
+/** URL to download ConfigRules file from. Must be a WIDE string starting with http: or https: */
+#ifndef WINDOWS_CONFIGRULES_URL
+#define WINDOWS_CONFIGRULES_URL TEXT("")
+#endif
+
+/** Secret used to derive AES key. Must be an ASCII string or char array. */
+#ifndef WINDOWS_CONFIGRULES_KEY
+#define WINDOWS_CONFIGRULES_KEY ""
+#endif
+
+/** Max miliseconds to spend waiting for download. Otherwise just proceed with bundled file */
+#ifndef WINDOWS_CONFIGRULES_DOWNLOAD_TIMEOUT
+#define WINDOWS_CONFIGRULES_DOWNLOAD_TIMEOUT (3000)
+#endif
+
+/** Path to save downloaded binary rules file */
+#ifndef WINDOWS_CONFIGRULES_DOWNLOAD_PATH
+#define WINDOWS_CONFIGRULES_DOWNLOAD_PATH FPaths::ProjectPersistentDownloadDir() / TEXT("ConfigRules") / TEXT("configrules.bin")
+#endif
+
+/** Path to find bundled config rules file */
+#ifndef WINDOWS_CONFIGRULES_BUNDLED_PATH
+#define WINDOWS_CONFIGRULES_BUNDLED_PATH TEXT("configrules.bin")
+#endif
+
+#include "Windows/WindowsHWrapper.h"
+#include "RHIGlobals.h"
+#include "Misc/Parse.h"
+#include "Misc/CommandLine.h"
+
+#pragma comment(lib, "wininet.lib")
+#pragma comment(lib, "bcrypt.lib")
+#include <wininet.h>
+#include <bcrypt.h>
+
+#define CONFIGRULES_FILE_SIGNATURE (0x39d8)
+#define CONFIGRULES_KEY_SALT { 0x23, 0x71, 0xd3, 0xa3, 0x30, 0x71, 0x63, 0xe3 }
+#define CONFIGRULES_KEY_ROUNDS (1000)
+
+#pragma pack(push, 1)
+union FConfigRuleFileHeader // NOTE: Packed to match layout on disk
+{
+	uint8 Bytes[10];
+	struct
+	{
+		uint16 Signature;
+		int32 Version;
+		int32 UncompressedSize;
+	};
+};
+#pragma pack(pop)
+
+struct FConfigRuleFile
+{
+	TArray<uint8> FileBytes;
+	FConfigRuleFileHeader Header;
+	const TCHAR *FilePath;
+
+	void ParseHeader();
+	bool Decrypt(TArray<uint8>& DecryptedBytes) const;
+
+	FORCEINLINE TConstArrayView<uint8> GetContent() const
+	{
+		// Exclude header as actual content
+		const uint8* Ptr = FileBytes.GetData() + sizeof(FConfigRuleFileHeader);
+		const int32 Num = FMath::Max<int32>(0, FileBytes.Num() - sizeof(FConfigRuleFileHeader));
+		return MakeConstArrayView(Ptr, Num);
+	}
+};
+
+void FConfigRuleFile::ParseHeader()
+{
+	Header = {};
+	Header.Version = INDEX_NONE;
+
+	if (FileBytes.Num() >= sizeof(FConfigRuleFileHeader))
+	{
+		FConfigRuleFileHeader ReadHeader;
+		FMemory::Memcpy(ReadHeader.Bytes, FileBytes.GetData(), sizeof(FConfigRuleFileHeader));
+
+		// NOTE: Header is stored in big endian
+#if PLATFORM_LITTLE_ENDIAN
+		uint16 Signature = ByteSwap(ReadHeader.Signature);
+		int32 Version = ByteSwap(ReadHeader.Version);
+		int32 UncompressedSize = ByteSwap(ReadHeader.UncompressedSize);
+#else
+		uint16 Signature = ReadHeader.Signature;
+		int32 Version = ReadHeader.Version;
+		int32 UncompressedSize = ReadHeader.UncompressedSize;
+#endif
+
+		if (Signature == CONFIGRULES_FILE_SIGNATURE)
+		{
+			Header.Signature = Signature;
+			Header.Version = Version;
+			Header.UncompressedSize = UncompressedSize;
+		}
+	}
+}
+
+bool FConfigRuleFile::Decrypt(TArray<uint8>& DecryptedBytes) const
+{
+	TConstArrayView<uint8> BytesToDecrypt = GetContent();
+	DecryptedBytes.SetNumUninitialized(BytesToDecrypt.Num());
+
+	const uint8 Secret[] = WINDOWS_CONFIGRULES_KEY;
+	const uint8 Salt[] = CONFIGRULES_KEY_SALT;
+	const uint32 Rounds = CONFIGRULES_KEY_ROUNDS;
+
+	uint8 GeneratedKey[16] = {};
+
+	BCRYPT_ALG_HANDLE AlgoHandle = NULL;
+	BCRYPT_KEY_HANDLE KeyHandle = NULL;
+
+	NTSTATUS DeriveStatus = 0;
+	DeriveStatus |= ::BCryptOpenAlgorithmProvider(&AlgoHandle, BCRYPT_SHA1_ALGORITHM, NULL, BCRYPT_ALG_HANDLE_HMAC_FLAG);
+	DeriveStatus |= ::BCryptDeriveKeyPBKDF2(AlgoHandle,
+		(PUCHAR)Secret,
+		sizeof(Secret) - 1,
+		(PUCHAR)Salt,
+		sizeof(Salt),
+		Rounds,
+		(PUCHAR)GeneratedKey,
+		sizeof(GeneratedKey),
+		0
+	);
+	::BCryptCloseAlgorithmProvider(AlgoHandle, 0);
+
+	if (BCRYPT_SUCCESS(DeriveStatus))
+	{
+		NTSTATUS KeyStatus = 0;
+		KeyStatus |= ::BCryptOpenAlgorithmProvider(&AlgoHandle, BCRYPT_AES_ALGORITHM, NULL, 0);
+		KeyStatus |= ::BCryptSetProperty(AlgoHandle, BCRYPT_CHAINING_MODE, (PUCHAR)BCRYPT_CHAIN_MODE_ECB, sizeof(BCRYPT_CHAIN_MODE_ECB), 0);
+		KeyStatus |= ::BCryptGenerateSymmetricKey(AlgoHandle, &KeyHandle, NULL, 0, GeneratedKey, sizeof(GeneratedKey), 0);
+		::BCryptCloseAlgorithmProvider(AlgoHandle, 0);
+
+		if (BCRYPT_SUCCESS(KeyStatus))
+		{
+			ULONG DecryptedBytesNum = 0;
+			NTSTATUS DecryptStatus = ::BCryptDecrypt(
+				KeyHandle,
+				(PUCHAR)BytesToDecrypt.GetData(),
+				BytesToDecrypt.Num(),
+				NULL,
+				NULL,
+				0,
+				(PUCHAR)DecryptedBytes.GetData(),
+				DecryptedBytes.Num(),
+				&DecryptedBytesNum,
+				BCRYPT_BLOCK_PADDING
+			);
+			
+			if (BCRYPT_SUCCESS(DecryptStatus))
+			{
+				::BCryptDestroyKey(KeyHandle);
+				DecryptedBytes.SetNum(DecryptedBytesNum);
+				return true;
+			}
+		}
+		::BCryptDestroyKey(KeyHandle);
+	}
+
+	return false;
+}
+
+struct FWindowsConfigRulesDownload
+{
+	HANDLE HDownloadThread;
+	TArray<uint8> DownloadContent;
+
+	FWindowsConfigRulesDownload();
+	~FWindowsConfigRulesDownload();
+
+	void SaveDownloadedFile();
+};
+
+// NOTE: static boots up the download thread ASAP, and allows us to register with GetPreMainInitDelegate()
+static FWindowsConfigRulesDownload GWindowsConfigRulesDownload;
+
+TMap<FString, FString> FWindowsConfigRules::ConfigRuleVariablesMap;
+TArray<uint8> FWindowsConfigRules::ConfigRulesBytes;
+
+void FWindowsConfigRules::Init()
+{
+	const FString& DownloadPath = WINDOWS_CONFIGRULES_DOWNLOAD_PATH;
+	TArray<FString> FilePaths = {
+		DownloadPath, /* Previously downloaded file */
+		WINDOWS_CONFIGRULES_BUNDLED_PATH
+	};
+
+#if !UE_BUILD_SHIPPING
+	FString DebugConfigRulesPath = TEXT("");
+	if (FParse::Value(FCommandLine::Get(), TEXT("-DebugConfigRulesPath="), DebugConfigRulesPath))
+	{
+		FilePaths.Add(DebugConfigRulesPath);
+	}
+#endif // !UE_BUILD_SHIPPING
+
+	TArray<FConfigRuleFile> Files;
+	Files.Reserve(FilePaths.Num() + 1);
+
+	int32 LatestFileIndex = INDEX_NONE;
+	int32 LatestVersion = INDEX_NONE;
+
+	// Load file in to memory and find the one with highest version
+	for (const FString& FilePath : FilePaths)
+	{
+		const bool bExists = FPaths::FileExists(*FilePath);
+		if (bExists)
+		{
+			const int32 Index = Files.AddDefaulted();
+			FConfigRuleFile& File = Files[Index];
+			File.FilePath = *FilePath;
+
+			bool bLoadFileResult = FFileHelper::LoadFileToArray(File.FileBytes, *FilePath);
+			if (bLoadFileResult)
+			{
+				File.ParseHeader();
+				if (File.Header.Version >= 0 && File.Header.Version > LatestVersion)
+				{
+					LatestFileIndex = Index;
+					LatestVersion = File.Header.Version;
+				}
+			}
+			else
+			{
+				UE_LOG(LogConfigRules, Error, TEXT("Could not LoadFileToArray: %s"), *FilePath);
+			}
+		}
+		else
+		{
+			UE_LOG(LogConfigRules, Log, TEXT("Could not find file: %s"), *FilePath);
+		}
+	}
+
+	DWORD DownloadWaitCode = ::WaitForSingleObject(GWindowsConfigRulesDownload.HDownloadThread, WINDOWS_CONFIGRULES_DOWNLOAD_TIMEOUT);
+	DWORD DownloadExitCode = INDEX_NONE;
+	::GetExitCodeThread(GWindowsConfigRulesDownload.HDownloadThread, &DownloadExitCode);
+
+	if (DownloadExitCode == 0 && DownloadWaitCode == WAIT_OBJECT_0)
+	{
+		// NOTE: Performing the save here since paths/cmdline might not be initialized on the download thread yet
+		GWindowsConfigRulesDownload.SaveDownloadedFile();
+
+		// If download was successful check if it has a higher version
+		const int32 Index = Files.AddDefaulted();
+		FConfigRuleFile& File = Files[Index];
+		File.FilePath = *DownloadPath;
+
+		File.FileBytes = MoveTemp(GWindowsConfigRulesDownload.DownloadContent);
+		File.ParseHeader();
+		if (File.Header.Version >= 0 && File.Header.Version > LatestVersion)
+		{
+			LatestFileIndex = Index;
+			LatestVersion = File.Header.Version;
+		}
+	}
+	else
+	{
+		UE_LOG(LogConfigRules, Warning, TEXT("DownloadThread unsuccessful, only using bundled file"));
+	}
+	
+	// Decrypt and Inflate the latest file
+	if (LatestFileIndex != INDEX_NONE && LatestVersion >= 0)
+	{
+		const FConfigRuleFile& File = Files[LatestFileIndex];
+		UE_LOG(LogConfigRules, Log, TEXT("Using file %s with version %d"), File.FilePath, LatestVersion);
+
+		TArray<uint8> DecryptedBytes;
+		const bool bDecrypted = File.Decrypt(DecryptedBytes);
+
+		if (bDecrypted)
+		{
+			TArray<uint8> InflatedBytes;
+			InflatedBytes.SetNumUninitialized(File.Header.UncompressedSize);
+			const bool bInflated = FCompression::UncompressMemory(NAME_Zlib, InflatedBytes.GetData(), InflatedBytes.Num(), DecryptedBytes.GetData(), DecryptedBytes.Num());
+
+			if (bInflated)
+			{
+				ConfigRulesBytes = MoveTemp(InflatedBytes);
+			}
+			else
+			{
+				UE_LOG(LogConfigRules, Error, TEXT("Failed to Uncompress file"));
+			}
+		}
+		else
+		{
+			UE_LOG(LogConfigRules, Error, TEXT("Failed to Decrypt file"));
+		}
+	}
+	else
+	{
+		UE_LOG(LogConfigRules, Error, TEXT("Unable to find a file with valid version"));
+	}
+}
+
+void FWindowsConfigRules::Parse(TMap<FString, FString>&& PredefinedVariables)
+{
+	if (ConfigRulesBytes.Num())
+	{
+		ConfigRuleVariablesMap = FGenericConfigRules::ParseConfigRules(ConfigRulesBytes, MoveTemp(PredefinedVariables));
+	}
+}
+
+TMap<FString, FString> FWindowsConfigRules::GetPredefinedConfigRuleVars()
+{
+	TMap<FString, FString> PredefinedConfigRuleVars;
+
+	FString CPUVendor = FWindowsPlatformMisc::GetCPUVendor();
+	FString CPUBrand = FWindowsPlatformMisc::GetCPUBrand();
+	FString OSVersion = FWindowsPlatformMisc::GetOSVersion();
+	FString Cores = FString::FromInt(FWindowsPlatformMisc::NumberOfCores());
+	FString HostArch = FWindowsPlatformMisc::GetHostArchitecture();
+	FString MakeAndModel = FPlatformMisc::GetDeviceMakeAndModel();
+	FString BranchName = FApp::GetBranchName();
+	FString BuildVersion = FApp::GetBuildVersion();
+	FString BuildConfig = LexToString(FApp::GetBuildConfiguration());
+	FString BuildTarget = LexToString(FApp::GetBuildTargetType());
+
+	PredefinedConfigRuleVars.Emplace( "SRC_BranchName",			BranchName);
+	PredefinedConfigRuleVars.Emplace( "SRC_BuildVersion",		BuildVersion);
+	PredefinedConfigRuleVars.Emplace( "SRC_BuildConfig",		BuildConfig);
+	PredefinedConfigRuleVars.Emplace( "SRC_BuildTarget",		BuildTarget);
+	PredefinedConfigRuleVars.Emplace( "SRC_MakeAndModel",		MakeAndModel);
+	PredefinedConfigRuleVars.Emplace( "SRC_CPUVendor",			CPUVendor);
+	PredefinedConfigRuleVars.Emplace( "SRC_CPUBrand",			CPUBrand);
+	PredefinedConfigRuleVars.Emplace( "SRC_HostArch",			HostArch);
+	PredefinedConfigRuleVars.Emplace( "SRC_NumberOfCores",		Cores);
+	PredefinedConfigRuleVars.Emplace( "SRC_OSVersion",			OSVersion);
+
+	bool bIsHandheld = FWindowsPlatformMisc::IsDeviceGamingHandheld();
+	if (bIsHandheld)
+	{
+		PredefinedConfigRuleVars.Emplace( "SRC_IsHandheld", TEXT("1"));
+	}
+	
+	// Might not be available until RHIInit
+	PredefinedConfigRuleVars.Emplace( "SRC_GRHIDeviceId", 		FString::FromInt(GRHIDeviceId));
+	PredefinedConfigRuleVars.Emplace( "SRC_GRHIVendorId", 		FString::FromInt(GRHIVendorId));
+
+	return PredefinedConfigRuleVars;
+}
+
+static DWORD ReportInternetError(const TCHAR* Str)
+{
+	DWORD LastError = ::GetLastError();
+
+	if (LastError == ERROR_INTERNET_EXTENDED_ERROR)
+	{
+		DWORD IError = 0;
+		WCHAR IErrorStr[1024] = {};
+		DWORD IErrorStrLen = UE_ARRAY_COUNT(IErrorStr);
+		BOOL IErrorResult = ::InternetGetLastResponseInfoW(&IError, (LPWSTR)&IErrorStr, &IErrorStrLen);
+		UE_LOG(LogConfigRules, Log, TEXT("%s, LastError %d, Response (%d: %s)"),
+			Str, (uint32)LastError, (uint32)IError, IErrorStr);
+	}
+	else
+	{
+		UE_LOG(LogConfigRules, Log, TEXT("%s, LastError: %d"), Str, (uint32)LastError);
+	}
+
+	return LastError;
+}
+
+static DWORD WINAPI WinConfigRulesDownloadProc(void* Arg)
+{
+#if !UE_BUILD_SHIPPING
+	if (FParse::Param(::GetCommandLineW(), TEXT("NoConfigRulesDownload")))
+	{
+		return 0;
+	}
+#endif // !UE_BUILD_SHIPPING
+
+	FWindowsConfigRulesDownload* This = (FWindowsConfigRulesDownload *)Arg;
+	DWORD ExitStatus = 0;
+
+	HINTERNET HInternet = ::InternetOpenW(L"WinConfigRulesDownloadProc", INTERNET_OPEN_TYPE_DIRECT, NULL, NULL, 0);
+	if (HInternet)
+	{
+		FString DownloadUrl;
+#if !UE_BUILD_SHIPPING
+		if (!FParse::Value(::GetCommandLineW(), TEXT("-DebugConfigRulesUrl="), DownloadUrl))
+#endif // !UE_BUILD_SHIPPING
+		{
+			DownloadUrl = WINDOWS_CONFIGRULES_URL;
+		}
+
+		// This actually sends the request (blocking)
+		HINTERNET HFile = ::InternetOpenUrlW(HInternet, *DownloadUrl, NULL, 0, INTERNET_FLAG_SECURE, NULL);
+		if (HFile)
+		{
+			TArray<uint8>& DownloadContent = This->DownloadContent;
+
+			uint8 Buffer[4096];
+			DWORD ReadBytes = 0;
+			BOOL ReadResult = 0;
+			do
+			{
+				ReadResult = ::InternetReadFile(HFile, Buffer, UE_ARRAY_COUNT(Buffer), &ReadBytes);
+				if (ReadResult)
+				{
+					DownloadContent.Append(Buffer, (int32)ReadBytes);
+				}
+				else
+				{
+					ExitStatus = ReportInternetError(TEXT("Failed to read response"));
+				}
+			} while (ReadResult && ReadBytes != 0);
+
+			::InternetCloseHandle(HFile);
+			HFile = NULL;
+		}
+		else
+		{
+			ExitStatus = ReportInternetError(TEXT("Failed to open url"));
+		}
+
+		::InternetCloseHandle(HInternet);
+		HInternet = NULL;
+	}
+	else
+	{
+		ExitStatus = ReportInternetError(TEXT("Failed to open connection"));
+	}
+
+	return ExitStatus;
+}
+
+FWindowsConfigRulesDownload::FWindowsConfigRulesDownload()
+{
+	HDownloadThread = ::CreateThread(NULL, 32 << 10, WinConfigRulesDownloadProc, (void*)this, 0, NULL);
+}
+
+FWindowsConfigRulesDownload::~FWindowsConfigRulesDownload()
+{
+	if (HDownloadThread)
+	{
+		if (::WaitForSingleObject(HDownloadThread, 1000) == WAIT_OBJECT_0)
+		{
+			::CloseHandle(HDownloadThread);
+		}
+		HDownloadThread = NULL;
+	}
+}
+
+void FWindowsConfigRulesDownload::SaveDownloadedFile()
+{
+	const FString& FilePath = WINDOWS_CONFIGRULES_DOWNLOAD_PATH;
+	const FString& DirectoryPath = FPaths::GetPath(FilePath);
+
+	IPlatformFile& PlatformFile = IPlatformFile::GetPlatformPhysical();
+	if (PlatformFile.CreateDirectoryTree(*DirectoryPath))
+	{
+		IFileHandle* Handle = PlatformFile.OpenWrite(*FilePath, false, true);
+		if (Handle)
+		{
+			bool bWriteResult = Handle->Write(DownloadContent.GetData(), DownloadContent.Num());
+			if (bWriteResult)
+			{
+				UE_LOG(LogConfigRules, Log, TEXT("Wrote %d bytes to %s"), DownloadContent.Num(), *FilePath);
+			}
+			else
+			{
+				UE_LOG(LogConfigRules, Error, TEXT("Failed to write to %s"), *FilePath);
+			}
+
+			delete Handle;
+			Handle = nullptr;
+		}
+		else
+		{
+			UE_LOG(LogConfigRules, Error, TEXT("Failed to OpenWrite %s"), *FilePath);
+		}
+	}
+	else
+	{
+		UE_LOG(LogConfigRules, Error, TEXT("Failed to CreateDirectoryTree %s"), *DirectoryPath);
+	}
+}
+
+void FWindowsConfigRules::Release()
+{
+	ConfigRulesBytes.Empty();
+}
+
+/** Cleanup resources not needed after engine init */
+static FDelayedAutoRegisterHelper GWindowsConfigRulesEndOfInit(
+	EDelayedRegisterRunPhase::EndOfEngineInit, []
+	{
+		FWindowsConfigRules::Release();
+	});
+
+#endif // WITH_WINDOWS_CONFIGRULES
diff --git a/Engine/Source/Runtime/ApplicationCore/Public/Windows/WindowsConfigRules.h b/Engine/Source/Runtime/ApplicationCore/Public/Windows/WindowsConfigRules.h
new file mode 100644
--- /dev/null
+++ b/Engine/Source/Runtime/ApplicationCore/Public/Windows/WindowsConfigRules.h
@@ -0,0 +1,42 @@
+// Copyright Epic Games, Inc. All Rights Reserved.
+
+#pragma once
+
+#ifndef WITH_WINDOWS_CONFIGRULES
+#define WITH_WINDOWS_CONFIGRULES 0
+#endif // WITH_WINDOWS_CONFIGRULES
+
+#if WITH_WINDOWS_CONFIGRULES
+
+#include "CoreMinimal.h"
+#include "GenericPlatform/GenericConfigRules.h"
+
+class FWindowsConfigRules : public FGenericConfigRules
+{
+	static TMap<FString, FString> ConfigRuleVariablesMap;
+	static TArray<uint8> ConfigRulesBytes;
+
+public:
+
+	/** Finds the best config rules file and loads it for parsing */
+	static APPLICATIONCORE_API void Init();
+
+	/**
+	 * Parse the loaded config rules file and updates the variable map.
+	 *
+	 * A concrete use case for Reparsing is when some data in PredefinedVars is not available during the first parse and
+	 * we'd like to update the RulesMap with new data. E.g. RHI might not be initialized during first Parse but we'd still
+	 * like to apply other values before RHIInit. Then during RHIInit call Parse() again to get updated data.
+	 *
+	 * Can be reparsed until EndOfEngineInit or Release()
+	 */
+	static APPLICATIONCORE_API void Parse(TMap<FString, FString>&& PredefinedVariables);
+	static APPLICATIONCORE_API TMap<FString, FString> GetPredefinedConfigRuleVars();
+
+	/** Frees resouces used to parse config rules */
+	static APPLICATIONCORE_API void Release();
+	
+	static const TMap<FString, FString>& GetConfigRulesMap() { return ConfigRuleVariablesMap; }
+};
+
+#endif // WITH_WINDOWS_CONFIGRULES
diff --git a/Engine/Source/Runtime/Core/Core.Build.cs b/Engine/Source/Runtime/Core/Core.Build.cs
--- a/Engine/Source/Runtime/Core/Core.Build.cs
+++ b/Engine/Source/Runtime/Core/Core.Build.cs
@@ -113,6 +113,11 @@ public class Core : ModuleRules
 			{
 				PublicDefinitions.Add("UE_MEMORY_TAGS_TRACE_ENABLED=1");
 			}
+			
+			if (Target.bCompileAgainstApplicationCore)
+			{
+				PrivateIncludePathModuleNames.Add("ApplicationCore");
+			}
 		}
 		else if (Target.Platform == UnrealTargetPlatform.Mac)
 		{
diff --git a/Engine/Source/Runtime/Core/Private/Windows/WindowsPlatformMisc.cpp b/Engine/Source/Runtime/Core/Private/Windows/WindowsPlatformMisc.cpp
--- a/Engine/Source/Runtime/Core/Private/Windows/WindowsPlatformMisc.cpp
+++ b/Engine/Source/Runtime/Core/Private/Windows/WindowsPlatformMisc.cpp
@@ -36,6 +36,9 @@
 #include "Windows/WindowsPlatformCrashContext.h"
 #include "HAL/PlatformOutputDevices.h"
 #include "AutoRTFM.h"
+#if WITH_APPLICATION_CORE
+#include "Windows/WindowsConfigRules.h"
+#endif // WITH_APPLICATION_CORE
 
 #include "GenericPlatform/GenericPlatformChunkInstall.h"
 #include "GenericPlatform/GenericPlatformDriver.h"
@@ -977,6 +980,17 @@ void FWindowsPlatformMisc::PlatformPreInit()
 
 	FGenericPlatformMisc::PlatformPreInit();
 
+#if WITH_APPLICATION_CORE && WITH_WINDOWS_CONFIGRULES
+	{
+		SCOPED_BOOT_TIMING("FWindowsConfigRules::Init");
+		FWindowsConfigRules::Init();
+	}
+	{
+		SCOPED_BOOT_TIMING("FWindowsConfigRules::Parse");
+		FWindowsConfigRules::Parse(FWindowsConfigRules::GetPredefinedConfigRuleVars());
+	}
+#endif // WITH_APPLICATION_CORE && WITH_WINDOWS_CONFIGRULES
+
 	FThreadHeartBeat::Get().GetOnThreadStuck().BindStatic(&FGenericCrashContext::OnThreadStuck);
 	FThreadHeartBeat::Get().GetOnThreadUnstuck().BindStatic(&FGenericCrashContext::OnThreadUnstuck);
 
@@ -4475,4 +4489,52 @@ const FPlatformDriveStats* FWindowsPlatformMisc::GetDriveStats(WIDECHAR DriveLet
 	return nullptr;
 }
 
+void FWindowsPlatformMisc::ApplyWindowsCompatConfigRules(void)
+{
+#if WITH_APPLICATION_CORE && WITH_WINDOWS_CONFIGRULES
+
+	bool bSkipAllowList = false;
+#if !UE_BUILD_SHIPPING
+	bSkipAllowList = FParse::Param(FCommandLine::Get(), TEXT("SkipConfigRulesAllowList"));
+#endif
+
+	TArray<FString> WindowsCompatCVars;
+	if (bSkipAllowList || GConfig->GetArray(TEXT("WindowsCompatCVars"), TEXT("CVars"), WindowsCompatCVars, GEngineIni))
+	{
+		TSet<FString> AllowedCompatCVars(WindowsCompatCVars);
+		const TMap<FString, FString>& ConfigRuleVars = FPlatformMisc::GetConfigRuleVars();
+		for (const TTuple<FString, FString>& Pair : ConfigRuleVars)
+		{
+			const FString& Key = Pair.Key;
+			const FString& Value = Pair.Value;
+			static const TCHAR WindowsCompat[] = TEXT("WindowsCompat.");
+			if (Key.StartsWith(WindowsCompat))
+			{
+				FString CVarName = Key.Mid(UE_ARRAY_COUNT(WindowsCompat) - 1);
+				if (bSkipAllowList || AllowedCompatCVars.Contains(CVarName))
+				{
+					auto* CVar = IConsoleManager::Get().FindConsoleVariable(*CVarName);
+					if (CVar)
+					{
+						// set with HF priority
+						// configrules are therefore higher priority than deviceprofiles. (e.g. -dpcvars)
+						CVar->Set(*Value, ECVF_SetByHotfix);
+						UE_LOG(LogConfigRules, Log, TEXT("Compat Setting %s = %s"), *CVarName, *Value);
+					}
+				}
+			}
+		}
+	}
+#endif // WITH_APPLICATION_CORE && WITH_WINDOWS_CONFIGRULES
+}
+
+const class TMap<FString, FString>& FWindowsPlatformMisc::GetConfigRuleVars()
+{
+#if WITH_APPLICATION_CORE && WITH_WINDOWS_CONFIGRULES
+	return FWindowsConfigRules::GetConfigRulesMap();
+#else 
+	return FGenericPlatformMisc::GetConfigRuleVars();
+#endif // WITH_APPLICATION_CORE && WITH_WINDOWS_CONFIGRULES
+}
+
 #undef UE_CPUID
diff --git a/Engine/Source/Runtime/Core/Public/Windows/WindowsPlatformMisc.h b/Engine/Source/Runtime/Core/Public/Windows/WindowsPlatformMisc.h
--- a/Engine/Source/Runtime/Core/Public/Windows/WindowsPlatformMisc.h
+++ b/Engine/Source/Runtime/Core/Public/Windows/WindowsPlatformMisc.h
@@ -244,7 +244,7 @@ struct FWindowsPlatformMisc
 	 * @return	Returns true if cpuid is supported
 	 */
 	static CORE_API bool HasCPUIDInstruction();
-	
+
 	static CORE_API bool HasTPM2Support();
 	static CORE_API bool HasSecureBootSupport();
 	/**
@@ -394,6 +394,9 @@ struct FWindowsPlatformMisc
 	 * NOTE: This call will block (if GCollectStorageInformationAsync is enabled) until stats are available
 	 */
 	static CORE_API const FPlatformDriveStats* GetDriveStats(WIDECHAR DriveLetter);
+
+	static CORE_API void ApplyWindowsCompatConfigRules(void);
+	static CORE_API const class TMap<FString, FString>& GetConfigRuleVars();
 };
 
 
diff --git a/Engine/Source/Runtime/Launch/Private/LaunchEngineLoop.cpp b/Engine/Source/Runtime/Launch/Private/LaunchEngineLoop.cpp
--- a/Engine/Source/Runtime/Launch/Private/LaunchEngineLoop.cpp
+++ b/Engine/Source/Runtime/Launch/Private/LaunchEngineLoop.cpp
@@ -6517,6 +6517,13 @@ bool FEngineLoop::AppInit( )
 		// Apply boot hotfixes
 		UE::ConfigUtilities::ApplyCVarsFromBootHotfix();
 	}
+	
+#if PLATFORM_WINDOWS
+	{
+		SCOPED_BOOT_TIMING("ApplyWindowsCompatConfigRules");
+		FWindowsPlatformMisc::ApplyWindowsCompatConfigRules();
+	}
+#endif // PLATFORM_WINDOWS
 
 	// Apply config driven presets
 	FTraceAuxiliary::InitializePresets(FCommandLine::Get());
diff --git a/Engine/Source/Runtime/RHI/Private/Windows/WindowsDynamicRHI.cpp b/Engine/Source/Runtime/RHI/Private/Windows/WindowsDynamicRHI.cpp
--- a/Engine/Source/Runtime/RHI/Private/Windows/WindowsDynamicRHI.cpp
+++ b/Engine/Source/Runtime/RHI/Private/Windows/WindowsDynamicRHI.cpp
@@ -10,6 +10,7 @@
 #include "DataDrivenShaderPlatformInfo.h"
 
 #include "Windows/WindowsPlatformApplicationMisc.h"
+#include "Windows/WindowsConfigRules.h"
 
 THIRD_PARTY_INCLUDES_START
 #include "Windows/AllowWindowsPlatformTypes.h"
@@ -1368,6 +1369,18 @@ FDynamicRHI* PlatformCreateDynamicRHI()
 
 	if (DynamicRHIModule)
 	{
+		// Run ConfigRules to allow hotfixing RHI. Things like GRHIDeviceId, GRHIVendorId are only available at this stage
+#if WITH_WINDOWS_CONFIGRULES
+		{
+			SCOPED_BOOT_TIMING("FWindowsConfigRules::Parse");
+			FWindowsConfigRules::Parse(FWindowsConfigRules::GetPredefinedConfigRuleVars());
+		}
+		{
+			SCOPED_BOOT_TIMING("ApplyWindowsCompatConfigRules");
+			FWindowsPlatformMisc::ApplyWindowsCompatConfigRules();
+		}
+#endif // WITH_WINDOWS_CONFIGRULES
+	
 		// Create the dynamic RHI.
 		DynamicRHI = DynamicRHIModule->CreateRHI(RequestedFeatureLevel);
 		GLoadedRHIModuleName = LoadedRHIModuleName;
