Adding missing includes to WindowsConfigRules.cpp

#Synced-CL 50788899

diff --git a/Engine/Source/Runtime/ApplicationCore/Private/Windows/WindowsConfigRules.cpp b/Engine/Source/Runtime/ApplicationCore/Private/Windows/WindowsConfigRules.cpp
--- a/Engine/Source/Runtime/ApplicationCore/Private/Windows/WindowsConfigRules.cpp
+++ b/Engine/Source/Runtime/ApplicationCore/Private/Windows/WindowsConfigRules.cpp
@@ -31,8 +31,12 @@
 
 #include "Windows/WindowsHWrapper.h"
 #include "RHIGlobals.h"
-#include "Misc/Parse.h"
+#include "Misc/App.h"
 #include "Misc/CommandLine.h"
+#include "Misc/Compression.h"
+#include "Misc/FileHelper.h"
+#include "Misc/Parse.h"
+#include "Misc/Paths.h"
 
 #pragma comment(lib, "wininet.lib")
 #pragma comment(lib, "bcrypt.lib")
