Skip copy of monitorFriendlyDeviceName if friendlyNameFromEdid is not set

#Synced-CL 50851879

diff --git a/Engine/Source/Runtime/ApplicationCore/Private/Windows/WindowsApplication.cpp b/Engine/Source/Runtime/ApplicationCore/Private/Windows/WindowsApplication.cpp
--- a/Engine/Source/Runtime/ApplicationCore/Private/Windows/WindowsApplication.cpp
+++ b/Engine/Source/Runtime/ApplicationCore/Private/Windows/WindowsApplication.cpp
@@ -1978,9 +1978,12 @@ static void GetFriendlyNameDisplayDevices(const WCHAR DeviceName[CCHDEVICENAME],
 				TargetName.header.id = Path.targetInfo.id;
 				if (DisplayConfigGetDeviceInfo(&TargetName.header) == ERROR_SUCCESS)
 				{
-					FriendlyName = TargetName.monitorFriendlyDeviceName;
+					if (TargetName.flags.friendlyNameFromEdid)
+					{
+						FriendlyName = TargetName.monitorFriendlyDeviceName;
+						break;
+					}
 				}
-				break;
 			}
 		}
 	}
