fixing game and main thread deadlocking issues caused by multi-view rendering

#Synced-CL 50718145

diff --git a/Engine/Source/Runtime/ApplicationCore/Private/Mac/CocoaWindow.cpp b/Engine/Source/Runtime/ApplicationCore/Private/Mac/CocoaWindow.cpp
--- a/Engine/Source/Runtime/ApplicationCore/Private/Mac/CocoaWindow.cpp
+++ b/Engine/Source/Runtime/ApplicationCore/Private/Mac/CocoaWindow.cpp
@@ -282,6 +282,8 @@ NSString* NSPerformDragOperation = @"NSPerformDragOperation";
 
 	if (MacApplication)
 	{
+		const bool bWaitForThread = false;
+		
 		GameThreadCall(^{
 			if (MacApplication) // Another check because game thread may destroy MacApplication before it gets here
 			{
@@ -291,7 +293,7 @@ NSString* NSPerformDragOperation = @"NSPerformDragOperation";
 					MacApplication->OnWindowActivationChanged(Window.ToSharedRef(), EWindowActivation::Activate);
 				}
 			}
-		}, true, @[ NSDefaultRunLoopMode, UnrealResizeEventMode, UnrealShowEventMode, UnrealFullscreenEventMode, UnrealCloseEventMode ]);
+		}, bWaitForThread, @[ NSDefaultRunLoopMode, UnrealResizeEventMode, UnrealShowEventMode, UnrealFullscreenEventMode, UnrealCloseEventMode ]);
 	}
 }
 
@@ -303,6 +305,7 @@ NSString* NSPerformDragOperation = @"NSPerformDragOperation";
 
 	if (MacApplication)
 	{
+		const bool bWaitForThread = false;
 		GameThreadCall(^{
 			if (MacApplication) // Another check because game thread may destroy MacApplication before it gets here
 			{
@@ -312,7 +315,7 @@ NSString* NSPerformDragOperation = @"NSPerformDragOperation";
 					MacApplication->OnWindowActivationChanged(Window.ToSharedRef(), EWindowActivation::Deactivate);
 				}
 			}
-		}, true, @[ NSDefaultRunLoopMode, UnrealResizeEventMode, UnrealShowEventMode, UnrealFullscreenEventMode, UnrealCloseEventMode ]);
+		}, bWaitForThread, @[ NSDefaultRunLoopMode, UnrealResizeEventMode, UnrealShowEventMode, UnrealFullscreenEventMode, UnrealCloseEventMode ]);
 	}
 }
 
@@ -468,6 +471,8 @@ NSString* NSPerformDragOperation = @"NSPerformDragOperation";
 	SCOPED_AUTORELEASE_POOL;
 	if (MacApplication && sender == self && !bIsBeingOrderedFront) // Skip informing Slate if we're simply changing the z order of windows
 	{
+		const bool bWaitForThread = false;
+		
 		GameThreadCall(^{
 			if (MacApplication) // Another check because game thread may destroy MacApplication before it gets here
 			{
@@ -477,7 +482,7 @@ NSString* NSPerformDragOperation = @"NSPerformDragOperation";
 					MacApplication->OnWindowWillResize(Window.ToSharedRef());
 				}
 			}
-		}, true, @[ NSDefaultRunLoopMode, UnrealResizeEventMode, UnrealShowEventMode, UnrealFullscreenEventMode ]);
+		}, bWaitForThread, @[ NSDefaultRunLoopMode, UnrealResizeEventMode, UnrealShowEventMode, UnrealFullscreenEventMode ]);
 	}
 	return frameSize;
 }
diff --git a/Engine/Source/Runtime/ApplicationCore/Private/Mac/MacApplication.cpp b/Engine/Source/Runtime/ApplicationCore/Private/Mac/MacApplication.cpp
--- a/Engine/Source/Runtime/ApplicationCore/Private/Mac/MacApplication.cpp
+++ b/Engine/Source/Runtime/ApplicationCore/Private/Mac/MacApplication.cpp
@@ -603,13 +603,14 @@ void FMacApplication::DeferEvent(NSObject* Object)
 			{
 				if (DeferredEvent.Window)
 				{
+					const bool bWaitForThread = false;
 					GameThreadCall(^{
 						TSharedPtr<FMacWindow> Window = FindWindowByNSWindow(DeferredEvent.Window);
 						if (Window.IsValid())
 						{
 							OnWindowDidResize(Window.ToSharedRef());
 						}
-					}, true, @[ NSDefaultRunLoopMode, UnrealResizeEventMode, UnrealShowEventMode, UnrealFullscreenEventMode ]);
+					}, bWaitForThread, @[ NSDefaultRunLoopMode, UnrealResizeEventMode, UnrealShowEventMode, UnrealFullscreenEventMode ]);
 				}
 				return;
 			}
@@ -1418,9 +1419,10 @@ bool FMacApplication::OnWindowDestroyed(TSharedRef<FMacWindow> DestroyedWindow)
 			
 			if (NumMenus<=0)
 			{
+				const bool bWaitForThread = false;
 				GameThreadCall(^{
 					OnWindowActivationChanged(WindowToActivate.ToSharedRef(), EWindowActivation::Activate);
-				}, true, @[ NSDefaultRunLoopMode, UnrealResizeEventMode, UnrealShowEventMode, UnrealFullscreenEventMode, UnrealCloseEventMode ]);
+				}, bWaitForThread, @[ NSDefaultRunLoopMode, UnrealResizeEventMode, UnrealShowEventMode, UnrealFullscreenEventMode, UnrealCloseEventMode ]);
 			}
 		}
 	}
