Add CVar to control whether to show the splash screen on Windows.

#Synced-CL 50851879

diff --git a/Engine/Source/Runtime/ApplicationCore/Private/Windows/WindowsPlatformSplash.cpp b/Engine/Source/Runtime/ApplicationCore/Private/Windows/WindowsPlatformSplash.cpp
--- a/Engine/Source/Runtime/ApplicationCore/Private/Windows/WindowsPlatformSplash.cpp
+++ b/Engine/Source/Runtime/ApplicationCore/Private/Windows/WindowsPlatformSplash.cpp
@@ -43,6 +43,12 @@ THIRD_PARTY_INCLUDES_END
 #define UE_WINDOWS_SPLASH_ENABLE_DRAG WITH_EDITOR
 #endif
 
+static TAutoConsoleVariable<bool> CVarWindowsApplicationShowSplash(
+	TEXT("WindowsApplication.ShowSplash"),
+	true,
+	TEXT("Whether to show the splash screen on Windows (default = true).")
+);
+
 /**
  * Splash screen functions and static globals
  */
@@ -690,7 +696,7 @@ static void StartSetSplashText( const SplashTextType::Type InType, const TCHAR*
 
 void FWindowsPlatformSplash::Show()
 {
-	if( !GSplashScreenThread && FParse::Param(FCommandLine::Get(),TEXT("NOSPLASH")) != true )
+	if(!GSplashScreenThread && FParse::Param(FCommandLine::Get(),TEXT("NOSPLASH")) != true && CVarWindowsApplicationShowSplash->GetBool())
 	{
 		const FText GameName = FText::FromString( FApp::GetProjectName() );
 
