Add several config rules predefined variables for RHI and memory stats.

For driver versions, also extract one variable for each component of the driver version, to make it easier to create match rules.

Finally, refactor how predefined variables are stored and managed so that each additional set of variables can be appended without re-doing the work for previous ones.

#Synced-CL 50851879

diff --git a/Engine/Source/Runtime/ApplicationCore/Private/Windows/WindowsConfigRules.cpp b/Engine/Source/Runtime/ApplicationCore/Private/Windows/WindowsConfigRules.cpp
--- a/Engine/Source/Runtime/ApplicationCore/Private/Windows/WindowsConfigRules.cpp
+++ b/Engine/Source/Runtime/ApplicationCore/Private/Windows/WindowsConfigRules.cpp
@@ -3,6 +3,7 @@
 #include "Windows/WindowsConfigRules.h"
 
 #if WITH_WINDOWS_CONFIGRULES
+#include "GenericPlatform/GenericPlatformDriver.h"
 
 /** URL to download ConfigRules file from. Must be a WIDE string starting with http: or https: */
 #ifndef WINDOWS_CONFIGRULES_URL
@@ -187,6 +188,7 @@ struct FWindowsConfigRulesDownload
 // NOTE: static boots up the download thread ASAP, and allows us to register with GetPreMainInitDelegate()
 static FWindowsConfigRulesDownload GWindowsConfigRulesDownload;
 
+TMap<FString, FString> FWindowsConfigRules::PredefinedConfigRuleVars;
 TMap<FString, FString> FWindowsConfigRules::ConfigRuleVariablesMap;
 TArray<uint8> FWindowsConfigRules::ConfigRulesBytes;
 
@@ -305,18 +307,17 @@ void FWindowsConfigRules::Init()
 	}
 }
 
-void FWindowsConfigRules::Parse(TMap<FString, FString>&& PredefinedVariables)
+void FWindowsConfigRules::Parse()
 {
 	if (ConfigRulesBytes.Num())
 	{
-		ConfigRuleVariablesMap = FGenericConfigRules::ParseConfigRules(ConfigRulesBytes, MoveTemp(PredefinedVariables));
+		TMap<FString, FString> PredefinedVars = PredefinedConfigRuleVars;
+		ConfigRuleVariablesMap = FGenericConfigRules::ParseConfigRules(ConfigRulesBytes, MoveTemp(PredefinedVars));
 	}
 }
 
-TMap<FString, FString> FWindowsConfigRules::GetPredefinedConfigRuleVars()
+void FWindowsConfigRules::SetStandardPredefinedVars()
 {
-	TMap<FString, FString> PredefinedConfigRuleVars;
-
 	FString CPUVendor = FWindowsPlatformMisc::GetCPUVendor();
 	FString CPUBrand = FWindowsPlatformMisc::GetCPUBrand();
 	FString OSVersion = FWindowsPlatformMisc::GetOSVersion();
@@ -327,29 +328,83 @@ TMap<FString, FString> FWindowsConfigRules::GetPredefinedConfigRuleVars()
 	FString BuildVersion = FApp::GetBuildVersion();
 	FString BuildConfig = LexToString(FApp::GetBuildConfiguration());
 	FString BuildTarget = LexToString(FApp::GetBuildTargetType());
+	FString TotalPhysicalGB = FString::FromInt(FPlatformMemory::GetConstants().TotalPhysicalGB);
+	FString TotalVirtualGB = FString::FromInt(FPlatformMemory::GetConstants().TotalVirtual / 1024 / 1024 / 1024);
+
+	PredefinedConfigRuleVars.Emplace(TEXT("SRC_BranchName"), BranchName);
+	PredefinedConfigRuleVars.Emplace(TEXT("SRC_BuildVersion"), BuildVersion);
+	PredefinedConfigRuleVars.Emplace(TEXT("SRC_BuildConfig"), BuildConfig);
+	PredefinedConfigRuleVars.Emplace(TEXT("SRC_BuildTarget"), BuildTarget);
+	PredefinedConfigRuleVars.Emplace(TEXT("SRC_MakeAndModel"), MakeAndModel);
+	PredefinedConfigRuleVars.Emplace(TEXT("SRC_CPUVendor"), CPUVendor);
+	PredefinedConfigRuleVars.Emplace(TEXT("SRC_CPUBrand"), CPUBrand);
+	PredefinedConfigRuleVars.Emplace(TEXT("SRC_HostArch"), HostArch);
+	PredefinedConfigRuleVars.Emplace(TEXT("SRC_NumberOfCores"), Cores);
+	PredefinedConfigRuleVars.Emplace(TEXT("SRC_OSVersion"), OSVersion);
+	PredefinedConfigRuleVars.Emplace(TEXT("SRC_TotalPhysicalGB"), TotalPhysicalGB);
+	PredefinedConfigRuleVars.Emplace(TEXT("SRC_TotalVirtualGB"), TotalVirtualGB);
+	PredefinedConfigRuleVars.Emplace(TEXT("SRC_IsHandheld"), FWindowsPlatformMisc::IsDeviceGamingHandheld() ? TEXT("1") : TEXT("0"));
+}
 
-	PredefinedConfigRuleVars.Emplace( "SRC_BranchName",			BranchName);
-	PredefinedConfigRuleVars.Emplace( "SRC_BuildVersion",		BuildVersion);
-	PredefinedConfigRuleVars.Emplace( "SRC_BuildConfig",		BuildConfig);
-	PredefinedConfigRuleVars.Emplace( "SRC_BuildTarget",		BuildTarget);
-	PredefinedConfigRuleVars.Emplace( "SRC_MakeAndModel",		MakeAndModel);
-	PredefinedConfigRuleVars.Emplace( "SRC_CPUVendor",			CPUVendor);
-	PredefinedConfigRuleVars.Emplace( "SRC_CPUBrand",			CPUBrand);
-	PredefinedConfigRuleVars.Emplace( "SRC_HostArch",			HostArch);
-	PredefinedConfigRuleVars.Emplace( "SRC_NumberOfCores",		Cores);
-	PredefinedConfigRuleVars.Emplace( "SRC_OSVersion",			OSVersion);
-
-	bool bIsHandheld = FWindowsPlatformMisc::IsDeviceGamingHandheld();
-	if (bIsHandheld)
+void FWindowsConfigRules::SetRHIPredefinedVars(const FString& DynamicRHIModuleName, const FString& RequestedFeatureLevel)
+{
+	FGPUDriverInfo GPUDriverInfo = {};
+	if (!GRHIAdapterName.IsEmpty())
 	{
-		PredefinedConfigRuleVars.Emplace( "SRC_IsHandheld", TEXT("1"));
+		GPUDriverInfo = FPlatformMisc::GetGPUDriverInfo(GRHIAdapterName);
+	}
+
+	PredefinedConfigRuleVars.Emplace(TEXT("SRC_DynamicRHIModuleName"), DynamicRHIModuleName);
+	PredefinedConfigRuleVars.Emplace(TEXT("SRC_DynamicRHIFeatureLevel"), RequestedFeatureLevel);
+	PredefinedConfigRuleVars.Emplace(TEXT("SRC_GRHIDeviceId"), FString::FromInt(GRHIDeviceId));
+	PredefinedConfigRuleVars.Emplace(TEXT("SRC_GRHIVendorId"), FString::FromInt(GRHIVendorId));
+	PredefinedConfigRuleVars.Emplace(TEXT("SRC_GRHIAdapterName"), GRHIAdapterName);
+	PredefinedConfigRuleVars.Emplace(TEXT("SRC_GRHIAdapterUserDriverVersion"), GPUDriverInfo.UserDriverVersion);
+	PredefinedConfigRuleVars.Emplace(TEXT("SRC_GRHIAdapterInternalDriverVersion"), GPUDriverInfo.InternalDriverVersion);
+	PredefinedConfigRuleVars.Emplace(TEXT("SRC_GRHIAdapterDriverDate"), GPUDriverInfo.DriverDate);
+
+	TArray<FString> DriverDateComponents;
+	GPUDriverInfo.DriverDate.ParseIntoArray(DriverDateComponents, TEXT("-"), true);
+	if (DriverDateComponents.Num() == 3)
+	{
+		PredefinedConfigRuleVars.Emplace(TEXT("SRC_GRHIAdapterDriverDateMonth"), DriverDateComponents[0]);
+		PredefinedConfigRuleVars.Emplace(TEXT("SRC_GRHIAdapterDriverDateDay"),   DriverDateComponents[1]);
+		PredefinedConfigRuleVars.Emplace(TEXT("SRC_GRHIAdapterDriverDateYear"),  DriverDateComponents[2]);
 	}
-	
-	// Might not be available until RHIInit
-	PredefinedConfigRuleVars.Emplace( "SRC_GRHIDeviceId", 		FString::FromInt(GRHIDeviceId));
-	PredefinedConfigRuleVars.Emplace( "SRC_GRHIVendorId", 		FString::FromInt(GRHIVendorId));
 
-	return PredefinedConfigRuleVars;
+	// Split the driver version(s) into a field for each component, usually integers (but some vendors have strings like AMD with 'Adrenalin').
+	// Since the number of components varies by vendor, collect components from the LAST (component 1) to the FIRST (component N),
+	// to make it easier to build match rules.
+	TArray<FString> UserDriverVersionComponents;
+	GPUDriverInfo.UserDriverVersion.ParseIntoArray(UserDriverVersionComponents, TEXT("."), true);
+	if (!UserDriverVersionComponents.IsEmpty())
+	{
+		uint32 ComponentNumFromLast = 1;
+		for (int32 Idx = UserDriverVersionComponents.Num() - 1; Idx >= 0; --Idx)
+		{
+			FString NameStr = FString(TEXT("SRC_GRHIAdapterUserDriverVersionComponent")) + FString::FromInt(ComponentNumFromLast);
+			PredefinedConfigRuleVars.Emplace(NameStr, UserDriverVersionComponents[Idx]);
+			ComponentNumFromLast++;
+		}
+	}
+
+	TArray<FString> InternalDriverVersionComponents;
+	GPUDriverInfo.InternalDriverVersion.ParseIntoArray(InternalDriverVersionComponents, TEXT("."), true);
+	if (!InternalDriverVersionComponents.IsEmpty())
+	{
+		uint32 ComponentNumFromLast = 1;
+		for (int32 Idx = InternalDriverVersionComponents.Num() - 1; Idx >= 0; --Idx)
+		{
+			FString NameStr = FString(TEXT("SRC_GRHIAdapterInternalDriverVersionComponent")) + FString::FromInt(ComponentNumFromLast);
+			PredefinedConfigRuleVars.Emplace(NameStr, InternalDriverVersionComponents[Idx]);
+			ComponentNumFromLast++;
+		}
+	}
+}
+
+void FWindowsConfigRules::SetAdditionalPredefinedVars(TMap<FString, FString>&& AdditionalPredefinedVars)
+{
+	PredefinedConfigRuleVars.Append(MoveTemp(AdditionalPredefinedVars));
 }
 
 static DWORD ReportInternetError(const TCHAR* Str)
@@ -491,6 +546,7 @@ void FWindowsConfigRulesDownload::SaveDownloadedFile()
 
 void FWindowsConfigRules::Release()
 {
+	PredefinedConfigRuleVars.Empty();
 	ConfigRulesBytes.Empty();
 }
 
diff --git a/Engine/Source/Runtime/ApplicationCore/Public/Windows/WindowsConfigRules.h b/Engine/Source/Runtime/ApplicationCore/Public/Windows/WindowsConfigRules.h
--- a/Engine/Source/Runtime/ApplicationCore/Public/Windows/WindowsConfigRules.h
+++ b/Engine/Source/Runtime/ApplicationCore/Public/Windows/WindowsConfigRules.h
@@ -13,6 +13,7 @@
 
 class FWindowsConfigRules : public FGenericConfigRules
 {
+	static TMap<FString, FString> PredefinedConfigRuleVars;
 	static TMap<FString, FString> ConfigRuleVariablesMap;
 	static TArray<uint8> ConfigRulesBytes;
 
@@ -30,10 +31,20 @@ public:
 	 *
 	 * Can be reparsed until EndOfEngineInit or Release()
 	 */
-	static APPLICATIONCORE_API void Parse(TMap<FString, FString>&& PredefinedVariables);
-	static APPLICATIONCORE_API TMap<FString, FString> GetPredefinedConfigRuleVars();
+	static APPLICATIONCORE_API void Parse();
 
-	/** Frees resouces used to parse config rules */
+	/** Sets a number of standard predefined variables such as build version, CPU information, etc. */
+	static APPLICATIONCORE_API void SetStandardPredefinedVars();
+
+	/** 
+	 * Sets a number of RHI-specific predefined variables such as adapter name, driver version, etc.
+	 * Requires the global GRHIAdapterName to be set by the time this is called.
+	*/
+	static APPLICATIONCORE_API void SetRHIPredefinedVars(const FString& DynamicRHIModuleName, const FString& RequestedFeatureLevel);
+
+	static APPLICATIONCORE_API void SetAdditionalPredefinedVars(TMap<FString, FString>&& AdditionalPredefinedVars);
+	
+	/** Frees resources used to parse config rules */
 	static APPLICATIONCORE_API void Release();
 	
 	static const TMap<FString, FString>& GetConfigRulesMap() { return ConfigRuleVariablesMap; }
diff --git a/Engine/Source/Runtime/Core/Private/Windows/WindowsPlatformMisc.cpp b/Engine/Source/Runtime/Core/Private/Windows/WindowsPlatformMisc.cpp
--- a/Engine/Source/Runtime/Core/Private/Windows/WindowsPlatformMisc.cpp
+++ b/Engine/Source/Runtime/Core/Private/Windows/WindowsPlatformMisc.cpp
@@ -987,7 +987,8 @@ void FWindowsPlatformMisc::PlatformPreInit()
 	}
 	{
 		SCOPED_BOOT_TIMING("FWindowsConfigRules::Parse");
-		FWindowsConfigRules::Parse(FWindowsConfigRules::GetPredefinedConfigRuleVars());
+		FWindowsConfigRules::SetStandardPredefinedVars();
+		FWindowsConfigRules::Parse();
 	}
 #endif // WITH_APPLICATION_CORE && WITH_WINDOWS_CONFIGRULES
 
diff --git a/Engine/Source/Runtime/RHI/Private/Windows/WindowsDynamicRHI.cpp b/Engine/Source/Runtime/RHI/Private/Windows/WindowsDynamicRHI.cpp
--- a/Engine/Source/Runtime/RHI/Private/Windows/WindowsDynamicRHI.cpp
+++ b/Engine/Source/Runtime/RHI/Private/Windows/WindowsDynamicRHI.cpp
@@ -1373,7 +1373,8 @@ FDynamicRHI* PlatformCreateDynamicRHI()
 #if WITH_WINDOWS_CONFIGRULES
 		{
 			SCOPED_BOOT_TIMING("FWindowsConfigRules::Parse");
-			FWindowsConfigRules::Parse(FWindowsConfigRules::GetPredefinedConfigRuleVars());
+			FWindowsConfigRules::SetRHIPredefinedVars(LoadedRHIModuleName, LexToString(RequestedFeatureLevel));
+			FWindowsConfigRules::Parse();
 		}
 		{
 			SCOPED_BOOT_TIMING("ApplyWindowsCompatConfigRules");
