[Input] Add a "OnInputDeviceHardwareInfoAvailable" delegate to the input device mapper which will fire before the actual OnInputDeviceConnectionChange. This gives the UInputDeviceSubsystem a chance to look for any metadata that is associated with the connected device and correctly associate it with a new FInputDeviceId. This gaurantees that any game or UI code bound to the "normal" OnInputDeviceConnectionChange delegate will have easy access to this info while it is executing.

[RN] Add a OnInputDeviceHardwareInfoAvailable delegate to the input device mapper, which allows for code to bind to connection events and be gaurneteed to run before other systems like Common Input/User code. This makes it possible for the UInputDeviceSubsystem to query the FInputDeviceScope, and make the necessary associations of the known hardware device Id to the new FInputDeviceId.
#tests test in PIE with xinput, windualshock, and GameInput plugins to cover windows. Also test on console targets with complex controller connection logic.

#Synced-CL 50969376

diff --git a/Engine/Source/Runtime/ApplicationCore/Private/GenericPlatform/GenericPlatformInputDeviceMapper.cpp b/Engine/Source/Runtime/ApplicationCore/Private/GenericPlatform/GenericPlatformInputDeviceMapper.cpp
--- a/Engine/Source/Runtime/ApplicationCore/Private/GenericPlatform/GenericPlatformInputDeviceMapper.cpp
+++ b/Engine/Source/Runtime/ApplicationCore/Private/GenericPlatform/GenericPlatformInputDeviceMapper.cpp
@@ -4,6 +4,7 @@
 #include "HAL/PlatformApplicationMisc.h"
 #include "Misc/CoreDelegates.h"
 #include "Misc/ConfigCacheIni.h"
+#include "GenericPlatform/GenericApplicationMessageHandler.h"
 #include "GenericPlatform/InputDeviceMappingPolicy.h"
 
 ///////////////////////////////////////////////////////////////////////////
@@ -11,6 +12,7 @@
 
 DEFINE_LOG_CATEGORY(LogInputDeviceMapper);
 
+IPlatformInputDeviceMapper::FOnInputDeviceHardwareInfoAvailable IPlatformInputDeviceMapper::OnInputDeviceHardwareInfoAvailable;
 IPlatformInputDeviceMapper::FOnUserInputDeviceConnectionChange IPlatformInputDeviceMapper::OnInputDeviceConnectionChange;
 IPlatformInputDeviceMapper::FOnUserInputDevicePairingChange IPlatformInputDeviceMapper::OnInputDevicePairingChange;
 
@@ -358,6 +360,12 @@ bool IPlatformInputDeviceMapper::Internal_MapInputDeviceToUser(FInputDeviceId De
 	FPlatformInputDeviceState& InputDeviceState = MappedInputDevices.FindOrAdd(DeviceId);
 	InputDeviceState.OwningPlatformUser = UserId;
 	InputDeviceState.ConnectionState = ConnectionState;
+
+	// If there is a valid FInputDeviceScope at the time of connection, then the hardware info is available.
+	if (FInputDeviceScope::GetCurrent() != nullptr)
+	{
+		OnInputDeviceHardwareInfoAvailable.Broadcast(ConnectionState, UserId, DeviceId);	
+	}
 	
 	// Broadcast delegates to let listeners know that the platform user has had an input device change
 	OnInputDeviceConnectionChange.Broadcast(ConnectionState, UserId, DeviceId);
diff --git a/Engine/Source/Runtime/ApplicationCore/Public/GenericPlatform/GenericPlatformInputDeviceMapper.h b/Engine/Source/Runtime/ApplicationCore/Public/GenericPlatform/GenericPlatformInputDeviceMapper.h
--- a/Engine/Source/Runtime/ApplicationCore/Public/GenericPlatform/GenericPlatformInputDeviceMapper.h
+++ b/Engine/Source/Runtime/ApplicationCore/Public/GenericPlatform/GenericPlatformInputDeviceMapper.h
@@ -179,6 +179,19 @@ public:
 
 	//////////////////////////////////////////////////////////////////////////////
 	// Delegates for listening to input device changes
+
+	/**
+	 * Callback which is fired before the "user" input device connection change. When this is fired, the input device metadata may not yet be populated.
+	 * 
+	 * This is intended for any "lower level" systems which want to directly access device connection changes before other engine plugins and game code does.
+	 * 
+	 * This should only be fired if there is a valid FInputDeviceScope which can be used to derive the hardware info from.
+	 * 
+	 * @param NewConnectionState	The new connection state of this device
+	 * @param FPlatformUserId		The User ID whose input device has changed
+	 * @param FInputDeviceId		The Input Device ID that has changed connection
+	 */
+	DECLARE_MULTICAST_DELEGATE_ThreeParams(FOnInputDeviceHardwareInfoAvailable, EInputDeviceConnectionState /* NewConnectionState */, FPlatformUserId /* PlatformUserId */, FInputDeviceId /* InputDeviceId */);
 	
 	/**
 	 * Callback for handling an Input Device's connection state change.
@@ -198,6 +211,7 @@ public:
 	 */
 	DECLARE_MULTICAST_DELEGATE_ThreeParams(FOnUserInputDevicePairingChange, FInputDeviceId /* InputDeviceId */, FPlatformUserId /* NewUserPlatformId */, FPlatformUserId /* OldUserPlatformId */);
 
+	FOnInputDeviceHardwareInfoAvailable& GetOnInputDeviceHardwareInfoAvailable() const { return OnInputDeviceHardwareInfoAvailable; }
 	FOnUserInputDeviceConnectionChange& GetOnInputDeviceConnectionChange() const { return OnInputDeviceConnectionChange; }
 	FOnUserInputDevicePairingChange& GetOnInputDevicePairingChange() const { return OnInputDevicePairingChange; }
 
@@ -302,6 +316,12 @@ protected:
 	 */
 	virtual bool ShouldBroadcastLegacyDelegates() const = 0;
 
+	/**
+	 * Callback which is fired right before the actual OnInputDeviceConnectionChange function for any systems which want to be notified
+	 * before any game code is, if there is a valid FInputDeviceScope.
+	 */
+	static APPLICATIONCORE_API FOnInputDeviceHardwareInfoAvailable OnInputDeviceHardwareInfoAvailable;
+	
 	/** Callback when input devices are disconnected/reconnected */
 	static APPLICATIONCORE_API FOnUserInputDeviceConnectionChange OnInputDeviceConnectionChange;
 
diff --git a/Engine/Source/Runtime/Engine/Private/GameFramework/InputDeviceSubsystem.cpp b/Engine/Source/Runtime/Engine/Private/GameFramework/InputDeviceSubsystem.cpp
--- a/Engine/Source/Runtime/Engine/Private/GameFramework/InputDeviceSubsystem.cpp
+++ b/Engine/Source/Runtime/Engine/Private/GameFramework/InputDeviceSubsystem.cpp
@@ -201,7 +201,7 @@ void UInputDeviceSubsystem::Initialize(FSubsystemCollectionBase& Collection)
 	InputPreprocessor = MakeShared<FInputDeviceSubsystemProcessor>();
 	FSlateApplication::Get().RegisterInputPreProcessor(InputPreprocessor, EInputPreProcessorType::Engine);
 
-	IPlatformInputDeviceMapper::Get().GetOnInputDeviceConnectionChange().AddUObject(this, &UInputDeviceSubsystem::HandleDeviceConnectionChange);
+	IPlatformInputDeviceMapper::Get().GetOnInputDeviceHardwareInfoAvailable().AddUObject(this, &UInputDeviceSubsystem::HandleDeviceConnectionChange);
 
 #if WITH_EDITOR
 	FEditorDelegates::PreBeginPIE.AddUObject(this, &UInputDeviceSubsystem::OnPrePIEStarted);
@@ -221,7 +221,7 @@ void UInputDeviceSubsystem::Deinitialize()
 	}
 	InputPreprocessor.Reset();
 
-	IPlatformInputDeviceMapper::Get().GetOnInputDeviceConnectionChange().RemoveAll(this);
+	IPlatformInputDeviceMapper::Get().GetOnInputDeviceHardwareInfoAvailable().RemoveAll(this);
 }
 
 bool UInputDeviceSubsystem::ShouldCreateSubsystem(UObject* Outer) const
