Fix implicit cast warnings

#rnx

#Synced-CL 50851879

diff --git a/Engine/Source/Runtime/ApplicationCore/Private/Windows/WindowsConfigRules.cpp b/Engine/Source/Runtime/ApplicationCore/Private/Windows/WindowsConfigRules.cpp
--- a/Engine/Source/Runtime/ApplicationCore/Private/Windows/WindowsConfigRules.cpp
+++ b/Engine/Source/Runtime/ApplicationCore/Private/Windows/WindowsConfigRules.cpp
@@ -328,8 +328,8 @@ void FWindowsConfigRules::SetStandardPredefinedVars()
 	FString BuildVersion = FApp::GetBuildVersion();
 	FString BuildConfig = LexToString(FApp::GetBuildConfiguration());
 	FString BuildTarget = LexToString(FApp::GetBuildTargetType());
-	FString TotalPhysicalGB = FString::FromInt(FPlatformMemory::GetConstants().TotalPhysicalGB);
-	FString TotalVirtualGB = FString::FromInt(FPlatformMemory::GetConstants().TotalVirtual / 1024 / 1024 / 1024);
+	FString TotalPhysicalGB = FString::FromInt(static_cast<int32>(FPlatformMemory::GetConstants().TotalPhysicalGB));
+	FString TotalVirtualGB = FString::FromInt(static_cast<int32>(FPlatformMemory::GetConstants().TotalVirtual / 1024 / 1024 / 1024));
 
 	PredefinedConfigRuleVars.Emplace(TEXT("SRC_BranchName"), BranchName);
 	PredefinedConfigRuleVars.Emplace(TEXT("SRC_BuildVersion"), BuildVersion);
