VisionOS: fixing deprecation warning.

#Synced-CL 51074604

diff --git a/Engine/Source/Runtime/ApplicationCore/Private/IOS/IOSAppDelegate.cpp b/Engine/Source/Runtime/ApplicationCore/Private/IOS/IOSAppDelegate.cpp
--- a/Engine/Source/Runtime/ApplicationCore/Private/IOS/IOSAppDelegate.cpp
+++ b/Engine/Source/Runtime/ApplicationCore/Private/IOS/IOSAppDelegate.cpp
@@ -743,7 +743,7 @@ static IOSAppDelegate* CachedDelegate = nil;
 			{
 				// iOS sdk 26 deprecated AVAudioSessionCategoryOptionAllowBluetooth in favor of AVAudioSessionCategoryOptionAllowBluetoothHFP
 				// but AVAudioSessionCategoryOptionAllowBluetoothHFP is not present in sdk 16 despite what the docs/header say
-#if (defined(__IPHONE_26_0) && __IPHONE_OS_VERSION_MAX_ALLOWED >= __IPHONE_26_0)
+#if (defined(__IPHONE_26_0) && __IPHONE_OS_VERSION_MAX_ALLOWED >= __IPHONE_26_0) || PLATFORM_VISIONOS
 				Options |= AVAudioSessionCategoryOptionAllowBluetoothHFP;
 #else
 				Options |= AVAudioSessionCategoryOptionAllowBluetooth;
@@ -774,7 +774,7 @@ static IOSAppDelegate* CachedDelegate = nil;
 			{
 				// iOS sdk 26 deprecated AVAudioSessionCategoryOptionAllowBluetooth in favor of AVAudioSessionCategoryOptionAllowBluetoothHFP
 				// but AVAudioSessionCategoryOptionAllowBluetoothHFP is not present in sdk 16 despite what the docs/header say
-#if (defined(__IPHONE_26_0) && __IPHONE_OS_VERSION_MAX_ALLOWED >= __IPHONE_26_0)
+#if (defined(__IPHONE_26_0) && __IPHONE_OS_VERSION_MAX_ALLOWED >= __IPHONE_26_0) || PLATFORM_VISIONOS
 				Options |= AVAudioSessionCategoryOptionAllowBluetoothHFP;
 #else
 				Options |= AVAudioSessionCategoryOptionAllowBluetooth;
