[iOS] Removing two Scooped Boot Timing macros that are being executed outside the game thread, which is no longer supported since november 2025.

#Synced-CL 51174899

diff --git a/Engine/Source/Runtime/ApplicationCore/Private/IOS/IOSView.cpp b/Engine/Source/Runtime/ApplicationCore/Private/IOS/IOSView.cpp
--- a/Engine/Source/Runtime/ApplicationCore/Private/IOS/IOSView.cpp
+++ b/Engine/Source/Runtime/ApplicationCore/Private/IOS/IOSView.cpp
@@ -144,7 +144,6 @@ MTL::Device* GMetalDevice = nullptr;
 	// the check for the function pointer itself is to determine if the Metal framework exists, before calling it
 	if ((bSupportsMetal || bSupportsMetalMRT) && MTLCreateSystemDefaultDevice != NULL)
 	{
-		SCOPED_BOOT_TIMING("CreateMetalDevice");
 		// if the device is unable to run with Metal (pre-A7), this will return nullptr
 		GMetalDevice = (__bridge MTL::Device*)MTLCreateSystemDefaultDevice();
 
@@ -175,8 +174,6 @@ MTL::Device* GMetalDevice = nullptr;
 
 - (id)initInternal:(CGRect)Frame
 {
-	SCOPED_BOOT_TIMING("[FIOSView initInternal]");
-
 	CachedMarkedText = nil;
 
 	check(GMetalDevice);
