PlatformWindow: Reconcile with recent change to EWindowZone order

#rnx
#tests Slate.User;Slate.Window

#Synced-CL 51251369

diff --git a/Engine/Source/Runtime/ApplicationCore/Private/Mac/MacApplication.cpp b/Engine/Source/Runtime/ApplicationCore/Private/Mac/MacApplication.cpp
--- a/Engine/Source/Runtime/ApplicationCore/Private/Mac/MacApplication.cpp
+++ b/Engine/Source/Runtime/ApplicationCore/Private/Mac/MacApplication.cpp
@@ -2043,27 +2043,7 @@ EWindowZone::Type FMacApplication::GetCurrentWindowZone(const TSharedRef<FMacWin
 
 bool FMacApplication::IsEdgeZone(EWindowZone::Type Zone) const
 {
-	switch (Zone)
-	{
-		case EWindowZone::NotInWindow:
-		case EWindowZone::TopLeftBorder:
-		case EWindowZone::TopBorder:
-		case EWindowZone::TopRightBorder:
-		case EWindowZone::LeftBorder:
-		case EWindowZone::RightBorder:
-		case EWindowZone::BottomLeftBorder:
-		case EWindowZone::BottomBorder:
-		case EWindowZone::BottomRightBorder:
-			return true;
-		case EWindowZone::TitleBar:
-		case EWindowZone::ClientArea:
-		case EWindowZone::MinimizeButton:
-		case EWindowZone::MaximizeButton:
-		case EWindowZone::CloseButton:
-		case EWindowZone::SysMenu:
-		default:
-			return false;
-	}
+	return Zone == EWindowZone::NotInWindow || EWindowZone::IsBorder(Zone);
 }
 
 bool FMacApplication::IsPrintableKey(uint32 Character) const
diff --git a/Engine/Source/Runtime/ApplicationCore/Private/Windows/WindowsApplication.cpp b/Engine/Source/Runtime/ApplicationCore/Private/Windows/WindowsApplication.cpp
--- a/Engine/Source/Runtime/ApplicationCore/Private/Windows/WindowsApplication.cpp
+++ b/Engine/Source/Runtime/ApplicationCore/Private/Windows/WindowsApplication.cpp
@@ -2893,7 +2893,7 @@ int32 FWindowsApplication::ProcessMessage( HWND hwnd, uint32 msg, WPARAM wParam,
 							Zone = EWindowZone::ClientArea;
 						}
 
-						static const LRESULT Results[] = {HTNOWHERE, HTTOPLEFT, HTTOP, HTTOPRIGHT, HTLEFT, HTCLIENT,
+						static const LRESULT Results[] = {HTNOWHERE, HTCLIENT, HTTOPLEFT, HTTOP, HTTOPRIGHT, HTLEFT,
 							HTRIGHT, HTBOTTOMLEFT, HTBOTTOM, HTBOTTOMRIGHT,
 							HTCAPTION, HTMINBUTTON, HTMAXBUTTON, HTCLOSE, HTSYSMENU};
 
diff --git a/Engine/Source/Runtime/ApplicationCore/Public/GenericPlatform/GenericApplicationMessageHandler.h b/Engine/Source/Runtime/ApplicationCore/Public/GenericPlatform/GenericApplicationMessageHandler.h
--- a/Engine/Source/Runtime/ApplicationCore/Public/GenericPlatform/GenericApplicationMessageHandler.h
+++ b/Engine/Source/Runtime/ApplicationCore/Public/GenericPlatform/GenericApplicationMessageHandler.h
@@ -80,8 +80,9 @@ enum class EWindowActivation : uint8
 namespace EWindowZone
 {
 	/**
-	 * The Window Zone is the window area we are currently over to send back to the operating system
-	 * for operating system compliance.
+	 * The Window Zone is the window area we are currently over to send back to the operating system for operating system compliance.
+	 * 
+	 * Must remain synchronized with platform enum maps in FWindowsApplication::ProcessMessage and FLinuxWindow::HitTest.
 	 */
 	enum Type
 	{
@@ -104,6 +105,8 @@ namespace EWindowZone
 		CloseButton			= 13,
 		SysMenu				= 14,
 
+		COUNT = 15,
+
 		/** No zone specified */
 		Unspecified	= 0,
 	};
