[Apple Input] Fix a bug on Mac where the left and right thumbstick direction buttons (up/down/left/right) were only firing the left/right keys, not the correct directions. This bug appears to have been introduced with 48988104 back in 2025

#rnx
#Synced-CL 51092797

diff --git a/Engine/Source/Runtime/ApplicationCore/Private/Apple/AppleControllerInterface.cpp b/Engine/Source/Runtime/ApplicationCore/Private/Apple/AppleControllerInterface.cpp
--- a/Engine/Source/Runtime/ApplicationCore/Private/Apple/AppleControllerInterface.cpp
+++ b/Engine/Source/Runtime/ApplicationCore/Private/Apple/AppleControllerInterface.cpp
@@ -631,8 +631,8 @@ void FAppleControllerInterface::HandleVirtualButtonGamepad(const FGamepadKeyName
 		bWasPositivePressed = PrevState.Find(FGamepadKeyNames::LeftAnalogX)->AxisValue >= RepeatDeadzone;
 		bPositivePressed = CurrentState.Find(FGamepadKeyNames::LeftAnalogX)->AxisValue >= RepeatDeadzone;
 		
-		HandleInputInternal(FGamepadKeyNames::LeftStickDown, ControllerIndex, bNegativePressed, bWasNegativePressed);
-		HandleInputInternal(FGamepadKeyNames::LeftStickUp, ControllerIndex, bPositivePressed, bWasPositivePressed);
+		HandleInputInternal(FGamepadKeyNames::LeftStickLeft, ControllerIndex, bNegativePressed, bWasNegativePressed);
+		HandleInputInternal(FGamepadKeyNames::LeftStickRight, ControllerIndex, bPositivePressed, bWasPositivePressed);
 	}
 	else if (UEButtonNegative == FGamepadKeyNames::LeftStickDown && UEButtonPositive == FGamepadKeyNames::LeftStickUp)
 	{
@@ -653,8 +653,8 @@ void FAppleControllerInterface::HandleVirtualButtonGamepad(const FGamepadKeyName
 		bWasPositivePressed = PrevState.Find(FGamepadKeyNames::RightAnalogX)->AxisValue >= RepeatDeadzone;
 		bPositivePressed = CurrentState.Find(FGamepadKeyNames::RightAnalogX)->AxisValue >= RepeatDeadzone;
 		
-		HandleInputInternal(FGamepadKeyNames::LeftStickDown, ControllerIndex, bNegativePressed, bWasNegativePressed);
-		HandleInputInternal(FGamepadKeyNames::LeftStickUp, ControllerIndex, bPositivePressed, bWasPositivePressed);
+		HandleInputInternal(FGamepadKeyNames::RightStickLeft, ControllerIndex, bNegativePressed, bWasNegativePressed);
+		HandleInputInternal(FGamepadKeyNames::RightStickRight, ControllerIndex, bPositivePressed, bWasPositivePressed);
 	}
 	else if (UEButtonNegative == FGamepadKeyNames::RightStickDown && UEButtonPositive == FGamepadKeyNames::RightStickUp)
 	{
