Add more GPU information to config rules for Windows

SRC_ExpectedDedicatedVideoMemoryMB can hotfix earlier during launch but may hold the value for the wrong GPU for some users. SRC_RHIDedicatedVideoMemoryMB will always represent the correct GPU.

Similarly SRC_ExpectedGPUVendorId/SRC_ExpectedGPUDeviceId/SRC_ExpectedGPUAdapterName are the earlier and less reliable counterparts of SRC_GRHIVendorId/SRC_GRHIDeviceId/SRC_GRHIAdapterName.

#Synced-CL 51366129

diff --git a/Engine/Source/Runtime/ApplicationCore/Private/Windows/WindowsConfigRules.cpp b/Engine/Source/Runtime/ApplicationCore/Private/Windows/WindowsConfigRules.cpp
--- a/Engine/Source/Runtime/ApplicationCore/Private/Windows/WindowsConfigRules.cpp
+++ b/Engine/Source/Runtime/ApplicationCore/Private/Windows/WindowsConfigRules.cpp
@@ -38,6 +38,7 @@
 #include "Misc/FileHelper.h"
 #include "Misc/Parse.h"
 #include "Misc/Paths.h"
+#include "Windows/WindowsD3D.h"
 
 #pragma comment(lib, "wininet.lib")
 #pragma comment(lib, "bcrypt.lib")
@@ -344,6 +345,12 @@ void FWindowsConfigRules::SetStandardPredefinedVars()
 	PredefinedConfigRuleVars.Emplace(TEXT("SRC_TotalPhysicalGB"), TotalPhysicalGB);
 	PredefinedConfigRuleVars.Emplace(TEXT("SRC_TotalVirtualGB"), TotalVirtualGB);
 	PredefinedConfigRuleVars.Emplace(TEXT("SRC_IsHandheld"), FWindowsPlatformMisc::IsDeviceGamingHandheld() ? TEXT("1") : TEXT("0"));
+
+	const FWindowsGPUInfo GPUInfo = FWindowsD3D::GetExpectedGPUInfo();
+	PredefinedConfigRuleVars.Emplace(TEXT("SRC_ExpectedGPUVendorId"), FString::Printf(TEXT("%u"), GPUInfo.VendorId));
+	PredefinedConfigRuleVars.Emplace(TEXT("SRC_ExpectedGPUDeviceId"), FString::Printf(TEXT("%u"), GPUInfo.DeviceId));
+	PredefinedConfigRuleVars.Emplace(TEXT("SRC_ExpectedDedicatedVideoMemoryMB"), FString::Printf(TEXT("%f"), static_cast<double>(GPUInfo.DedicatedVideoMemory) / 1048576.));
+	PredefinedConfigRuleVars.Emplace(TEXT("SRC_ExpectedGPUAdapterName"), GPUInfo.AdapterName.IsEmpty() ? FPlatformMisc::GetPrimaryGPUBrand() : GPUInfo.AdapterName);
 }
 
 void FWindowsConfigRules::SetRHIPredefinedVars(const FString& DynamicRHIModuleName, const FString& RequestedFeatureLevel)
@@ -356,8 +363,9 @@ void FWindowsConfigRules::SetRHIPredefinedVars(const FString& DynamicRHIModuleNa
 
 	PredefinedConfigRuleVars.Emplace(TEXT("SRC_DynamicRHIModuleName"), DynamicRHIModuleName);
 	PredefinedConfigRuleVars.Emplace(TEXT("SRC_DynamicRHIFeatureLevel"), RequestedFeatureLevel);
-	PredefinedConfigRuleVars.Emplace(TEXT("SRC_GRHIDeviceId"), FString::FromInt(GRHIDeviceId));
-	PredefinedConfigRuleVars.Emplace(TEXT("SRC_GRHIVendorId"), FString::FromInt(GRHIVendorId));
+	PredefinedConfigRuleVars.Emplace(TEXT("SRC_GRHIDeviceId"), FString::Printf(TEXT("%u"), GRHIDeviceId));
+	PredefinedConfigRuleVars.Emplace(TEXT("SRC_GRHIVendorId"), FString::Printf(TEXT("%u"), GRHIVendorId));
+	PredefinedConfigRuleVars.Emplace(TEXT("SRC_RHIDedicatedVideoMemoryMB"), FString::Printf(TEXT("%f"), static_cast<double>(GRHIGlobals.GpuInfo.DedicatedVideoMemory) / 1048576.));
 	PredefinedConfigRuleVars.Emplace(TEXT("SRC_GRHIAdapterName"), GRHIAdapterName);
 	PredefinedConfigRuleVars.Emplace(TEXT("SRC_GRHIAdapterUserDriverVersion"), GPUDriverInfo.UserDriverVersion);
 	PredefinedConfigRuleVars.Emplace(TEXT("SRC_GRHIAdapterInternalDriverVersion"), GPUDriverInfo.InternalDriverVersion);
diff --git a/Engine/Source/Runtime/D3D12RHI/Private/Windows/WindowsD3D12Device.cpp b/Engine/Source/Runtime/D3D12RHI/Private/Windows/WindowsD3D12Device.cpp
--- a/Engine/Source/Runtime/D3D12RHI/Private/Windows/WindowsD3D12Device.cpp
+++ b/Engine/Source/Runtime/D3D12RHI/Private/Windows/WindowsD3D12Device.cpp
@@ -1003,6 +1003,7 @@ void FD3D12DynamicRHIModule::FindAdapter()
 		GRHIAdapterName = AdapterDesc.Description;
 		GRHIDeviceId = AdapterDesc.DeviceId;
 		GRHIDeviceRevision = AdapterDesc.Revision;
+		GRHIGlobals.GpuInfo.DedicatedVideoMemory = AdapterDesc.DedicatedVideoMemory;
 	}
 	else
 	{
diff --git a/Engine/Source/Runtime/RHI/Public/RHIGlobals.h b/Engine/Source/Runtime/RHI/Public/RHIGlobals.h
--- a/Engine/Source/Runtime/RHI/Public/RHIGlobals.h
+++ b/Engine/Source/Runtime/RHI/Public/RHIGlobals.h
@@ -138,6 +138,9 @@ struct FRHIGlobals
 
 		// Vendor-specific generation number of the GPU. 0 if unknown.
 		uint32 AdapterGeneration = 0;
+
+		// Amount of VRAM (currently implemented for D3D11/D3D12).
+		uint64 DedicatedVideoMemory = 0;
 	} GpuInfo;
 
 	// true if the RHI supports Pixel Shader UAV
diff --git a/Engine/Source/Runtime/Windows/D3D11RHI/Private/Windows/WindowsD3D11Device.cpp b/Engine/Source/Runtime/Windows/D3D11RHI/Private/Windows/WindowsD3D11Device.cpp
--- a/Engine/Source/Runtime/Windows/D3D11RHI/Private/Windows/WindowsD3D11Device.cpp
+++ b/Engine/Source/Runtime/Windows/D3D11RHI/Private/Windows/WindowsD3D11Device.cpp
@@ -783,6 +783,7 @@ void FD3D11DynamicRHIModule::FindAdapter()
 	GRHIVendorId = ChosenAdapter.DXGIAdapterDesc.VendorId;
 	GRHIDeviceId = ChosenAdapter.DXGIAdapterDesc.DeviceId;
 	GRHIDeviceRevision = ChosenAdapter.DXGIAdapterDesc.Revision;
+	GRHIGlobals.GpuInfo.DedicatedVideoMemory = ChosenAdapter.DXGIAdapterDesc.DedicatedVideoMemory;
 	GRHIDeviceIsIntegrated = ChosenAdapter.bUMA;
 	GRHIDeviceIsUMASeparateVA = ChosenAdapter.bUMA;
 
