Ability to specify raw text Config Rules for testing on PC

#Synced-CL 51336708

diff --git a/Engine/Source/Runtime/ApplicationCore/Private/Windows/WindowsConfigRules.cpp b/Engine/Source/Runtime/ApplicationCore/Private/Windows/WindowsConfigRules.cpp
--- a/Engine/Source/Runtime/ApplicationCore/Private/Windows/WindowsConfigRules.cpp
+++ b/Engine/Source/Runtime/ApplicationCore/Private/Windows/WindowsConfigRules.cpp
@@ -195,6 +195,22 @@ TArray<uint8> FWindowsConfigRules::ConfigRulesBytes;
 
 void FWindowsConfigRules::Init()
 {
+#if !UE_BUILD_SHIPPING
+	// Use a raw .txt file for quickly testing rule behavior.
+	// Does not consider any other config rule file, nor does it check version.
+	FString TestRawConfigRulesPath = TEXT("");
+	if (FParse::Value(FCommandLine::Get(), TEXT("-TestRawConfigRulesPath="), TestRawConfigRulesPath))
+	{
+		bool bResult = FFileHelper::LoadFileToArray(ConfigRulesBytes, *TestRawConfigRulesPath);
+		if (!bResult)
+		{
+			ConfigRulesBytes.Reset();
+			UE_LOG(LogConfigRules, Error, TEXT("Failed to load Testing file: %s"), *TestRawConfigRulesPath);
+		}
+		return;
+	}
+#endif
+
 	const FString& DownloadPath = WINDOWS_CONFIGRULES_DOWNLOAD_PATH;
 	TArray<FString> FilePaths = {
 		DownloadPath, /* Previously downloaded file */
