Fix crash while exiting in FActorFolders destuctor.

#rnx
#Synced-CL 51377346

diff --git a/Engine/Source/Editor/UnrealEd/Private/EditorActorFolders.cpp b/Engine/Source/Editor/UnrealEd/Private/EditorActorFolders.cpp
--- a/Engine/Source/Editor/UnrealEd/Private/EditorActorFolders.cpp
+++ b/Engine/Source/Editor/UnrealEd/Private/EditorActorFolders.cpp
@@ -63,7 +63,10 @@ FActorFolders::~FActorFolders()
 
 	if (GEditor)
 	{
-		UActorEditorContextSubsystem::Get()->UnregisterClient(this);
+		if (UActorEditorContextSubsystem* ActorEditorContextSubsystem = UActorEditorContextSubsystem::Get())
+		{
+			ActorEditorContextSubsystem->UnregisterClient(this);
+		}
 	}
 }
 
diff --git a/Engine/Source/Runtime/ApplicationCore/Private/Linux/LinuxPlatformApplicationMisc.cpp b/Engine/Source/Runtime/ApplicationCore/Private/Linux/LinuxPlatformApplicationMisc.cpp
--- a/Engine/Source/Runtime/ApplicationCore/Private/Linux/LinuxPlatformApplicationMisc.cpp
+++ b/Engine/Source/Runtime/ApplicationCore/Private/Linux/LinuxPlatformApplicationMisc.cpp
@@ -285,6 +285,26 @@ static void SDLCALL SDLLogOutputFunction(void *userdata, int category, SDL_LogPr
 	UE_LOG(LogSDL3, Log, TEXT("%s"), UTF8_TO_TCHAR(message));
 }
 
+static int UEX11ErrorHandler(const char *errortext, int error_code, int request_code, int minor_code, void *userdata)
+{
+	UE_LOG(LogInit, Warning, TEXT("X11 Error: %s.  Error Code: %d.  Request Code: %d.  Minor Code: %d."),
+		UTF8_TO_TCHAR(errortext),
+		error_code,
+		request_code,
+		minor_code);
+
+	return 0;
+}
+
+static int UEX11IOErrorHandler(void *userdata)
+{
+	UE_LOG(LogInit, Error, TEXT("Fatal X11 IO Error: Lost connection to X Server"));
+	ensure(false);   // send a crash report if we can
+	FLinuxPlatformMisc::RequestExit(true);
+
+	return 0;
+}
+
 bool FLinuxPlatformApplicationMisc::InitSDL()
 {
 	if (!GInitializedSDL)
@@ -363,6 +383,12 @@ bool FLinuxPlatformApplicationMisc::InitSDL()
 			UE_LOG(LogInit, Log, TEXT("Using SDL video driver '%s'"),
 				UTF8_TO_TCHAR(SdlVideoDriver)
 			);
+
+			if (!strcmp(SdlVideoDriver, "x11"))
+			{
+				// Set up an error handler for X11 so we don't get the default version which just calls exit()
+				SDL_SetX11ErrorHandler(UEX11ErrorHandler, nullptr, UEX11IOErrorHandler, nullptr);
+			}
 		}
 
 
@@ -380,6 +406,33 @@ bool FLinuxPlatformApplicationMisc::InitSDL()
 	return true;
 }
 
+bool FLinuxPlatformApplicationMisc::IsUsingWayland()
+{
+	static int IsUsingWayland = -1;
+
+	// If we haven't initalized SDL yet, then we can't say if we are using Wayland
+	if (!GInitializedSDL)
+	{
+		UE_LOG(LogInit, Warning, TEXT("Querying IsUsingWayland before SDL is initalized"));
+		return false;
+	}
+
+	if (IsUsingWayland == -1)
+	{
+		char const* SdlVideoDriver = SDL_GetCurrentVideoDriver();
+		if (SdlVideoDriver)
+		{
+			IsUsingWayland = !strcmp(SdlVideoDriver, "wayland");
+		}
+		else
+		{
+			UE_LOG(LogInit, Warning, TEXT("SDL is not returning a video driver while querying IsUsingWayland"));	
+			return false;		
+		}
+	}
+	return IsUsingWayland;
+}
+
 void FLinuxPlatformApplicationMisc::TearDown()
 {
 	FGenericPlatformApplicationMisc::TearDown();
@@ -387,6 +440,10 @@ void FLinuxPlatformApplicationMisc::TearDown()
 	if (GInitializedSDL)
 	{
 		UE_LOG(LogInit, Log, TEXT("Tearing down SDL."));
+		if (!IsUsingWayland())
+		{
+			SDL_SetX11ErrorHandler(nullptr, nullptr, nullptr, nullptr);
+		}
 		SDL_Quit();
 		GInitializedSDL = false;
 
diff --git a/Engine/Source/Runtime/ApplicationCore/Public/Linux/LinuxPlatformApplicationMisc.h b/Engine/Source/Runtime/ApplicationCore/Public/Linux/LinuxPlatformApplicationMisc.h
--- a/Engine/Source/Runtime/ApplicationCore/Public/Linux/LinuxPlatformApplicationMisc.h
+++ b/Engine/Source/Runtime/ApplicationCore/Public/Linux/LinuxPlatformApplicationMisc.h
@@ -25,6 +25,7 @@ struct FLinuxPlatformApplicationMisc : public FGenericPlatformApplicationMisc
 	static APPLICATIONCORE_API void ClipboardCopy(const TCHAR* Str);
 	static APPLICATIONCORE_API void ClipboardPaste(class FString& Dest);
 	static bool FullscreenSameAsWindowedFullscreen() { return true; }
+	static APPLICATIONCORE_API bool IsUsingWayland();			// Is SDL using native Wayland video driver
 
 	// Unix specific
 	static APPLICATIONCORE_API void EarlyUnixInitialization(class FString& OutCommandLine);
