Correct spelling.

#rnx
#Synced-CL 51540233

diff --git a/Engine/Source/Runtime/ApplicationCore/Private/Linux/LinuxPlatformApplicationMisc.cpp b/Engine/Source/Runtime/ApplicationCore/Private/Linux/LinuxPlatformApplicationMisc.cpp
--- a/Engine/Source/Runtime/ApplicationCore/Private/Linux/LinuxPlatformApplicationMisc.cpp
+++ b/Engine/Source/Runtime/ApplicationCore/Private/Linux/LinuxPlatformApplicationMisc.cpp
@@ -410,10 +410,10 @@ bool FLinuxPlatformApplicationMisc::IsUsingWayland()
 {
 	static int IsUsingWayland = -1;
 
-	// If we haven't initalized SDL yet, then we can't say if we are using Wayland
+	// If we haven't initialized SDL yet, then we can't say if we are using Wayland
 	if (!GInitializedSDL)
 	{
-		UE_LOG(LogInit, Warning, TEXT("Querying IsUsingWayland before SDL is initalized"));
+		UE_LOG(LogInit, Warning, TEXT("Querying IsUsingWayland before SDL is initialized"));
 		return false;
 	}
 
