Fix shutdown when using dummy SDL video device

#rnx
#Synced-CL 51540233

diff --git a/Engine/Source/Runtime/ApplicationCore/Private/Linux/LinuxPlatformApplicationMisc.cpp b/Engine/Source/Runtime/ApplicationCore/Private/Linux/LinuxPlatformApplicationMisc.cpp
--- a/Engine/Source/Runtime/ApplicationCore/Private/Linux/LinuxPlatformApplicationMisc.cpp
+++ b/Engine/Source/Runtime/ApplicationCore/Private/Linux/LinuxPlatformApplicationMisc.cpp
@@ -433,6 +433,33 @@ bool FLinuxPlatformApplicationMisc::IsUsingWayland()
 	return IsUsingWayland;
 }
 
+bool FLinuxPlatformApplicationMisc::IsUsingX11()
+{
+	static int IsUsingX11 = -1;
+
+	// If we haven't initialized SDL yet, then we can't say if we are using X11
+	if (!GInitializedSDL)
+	{
+		UE_LOG(LogInit, Warning, TEXT("Querying IsUsingX11 before SDL is initialized"));
+		return false;
+	}
+
+	if (IsUsingX11 == -1)
+	{
+		char const* SdlVideoDriver = SDL_GetCurrentVideoDriver();
+		if (SdlVideoDriver)
+		{
+			IsUsingX11 = !strcmp(SdlVideoDriver, "x11");
+		}
+		else
+		{
+			UE_LOG(LogInit, Warning, TEXT("SDL is not returning a video driver while querying IsUsingX11"));	
+			return false;		
+		}
+	}
+	return IsUsingX11;
+}
+
 void FLinuxPlatformApplicationMisc::TearDown()
 {
 	FGenericPlatformApplicationMisc::TearDown();
@@ -440,7 +467,7 @@ void FLinuxPlatformApplicationMisc::TearDown()
 	if (GInitializedSDL)
 	{
 		UE_LOG(LogInit, Log, TEXT("Tearing down SDL."));
-		if (!IsUsingWayland())
+		if (IsUsingX11())
 		{
 			SDL_SetX11ErrorHandler(nullptr, nullptr, nullptr, nullptr);
 		}
diff --git a/Engine/Source/Runtime/ApplicationCore/Public/Linux/LinuxPlatformApplicationMisc.h b/Engine/Source/Runtime/ApplicationCore/Public/Linux/LinuxPlatformApplicationMisc.h
--- a/Engine/Source/Runtime/ApplicationCore/Public/Linux/LinuxPlatformApplicationMisc.h
+++ b/Engine/Source/Runtime/ApplicationCore/Public/Linux/LinuxPlatformApplicationMisc.h
@@ -26,6 +26,7 @@ struct FLinuxPlatformApplicationMisc : public FGenericPlatformApplicationMisc
 	static APPLICATIONCORE_API void ClipboardPaste(class FString& Dest);
 	static bool FullscreenSameAsWindowedFullscreen() { return true; }
 	static APPLICATIONCORE_API bool IsUsingWayland();			// Is SDL using native Wayland video driver
+	static APPLICATIONCORE_API bool IsUsingX11();			    // Is SDL using X11 video driver
 
 	// Unix specific
 	static APPLICATIONCORE_API void EarlyUnixInitialization(class FString& OutCommandLine);
