[Apple Input] Fix another copy/paste error in the apple controller interface that blocks the gamepad thumbstick directions from working correctly

#rnx
#Synced-CL 51542463

diff --git a/Engine/Source/Runtime/ApplicationCore/Private/Apple/AppleControllerInterface.cpp b/Engine/Source/Runtime/ApplicationCore/Private/Apple/AppleControllerInterface.cpp
--- a/Engine/Source/Runtime/ApplicationCore/Private/Apple/AppleControllerInterface.cpp
+++ b/Engine/Source/Runtime/ApplicationCore/Private/Apple/AppleControllerInterface.cpp
@@ -516,7 +516,7 @@ void FAppleControllerInterface::SendControllerEvents()
 		HandleAnalogGamepad(FGamepadKeyNames::RightTriggerAnalog, i);
 		HandleAnalogGamepad(FGamepadKeyNames::LeftTriggerAnalog, i);
 
-		HandleVirtualButtonGamepad(FGamepadKeyNames::LeftStickRight, FGamepadKeyNames::LeftStickLeft, i);
+		HandleVirtualButtonGamepad(FGamepadKeyNames::LeftStickLeft, FGamepadKeyNames::LeftStickRight, i);
 		HandleVirtualButtonGamepad(FGamepadKeyNames::LeftStickDown, FGamepadKeyNames::LeftStickUp, i);
 		HandleVirtualButtonGamepad(FGamepadKeyNames::RightStickLeft, FGamepadKeyNames::RightStickRight, i);
 		HandleVirtualButtonGamepad(FGamepadKeyNames::RightStickDown, FGamepadKeyNames::RightStickUp, i);
@@ -664,8 +664,8 @@ void FAppleControllerInterface::HandleVirtualButtonGamepad(const FGamepadKeyName
 		bWasPositivePressed = PrevState.Find(FGamepadKeyNames::RightAnalogY)->AxisValue >= RepeatDeadzone;
 		bPositivePressed = CurrentState.Find(FGamepadKeyNames::RightAnalogY)->AxisValue >= RepeatDeadzone;
 		
-		HandleInputInternal(FGamepadKeyNames::LeftStickDown, ControllerIndex, bNegativePressed, bWasNegativePressed);
-		HandleInputInternal(FGamepadKeyNames::LeftStickUp, ControllerIndex, bPositivePressed, bWasPositivePressed);
+		HandleInputInternal(FGamepadKeyNames::RightStickDown, ControllerIndex, bNegativePressed, bWasNegativePressed);
+		HandleInputInternal(FGamepadKeyNames::RightStickUp, ControllerIndex, bPositivePressed, bWasPositivePressed);
 	}
 }
 
