Make the Windows Splash Screen match the OS window style on Windows 11.

#rnx
#Synced-CL 51626022

diff --git a/Engine/Source/Runtime/ApplicationCore/Private/Windows/WindowsPlatformSplash.cpp b/Engine/Source/Runtime/ApplicationCore/Private/Windows/WindowsPlatformSplash.cpp
--- a/Engine/Source/Runtime/ApplicationCore/Private/Windows/WindowsPlatformSplash.cpp
+++ b/Engine/Source/Runtime/ApplicationCore/Private/Windows/WindowsPlatformSplash.cpp
@@ -28,6 +28,7 @@
 THIRD_PARTY_INCLUDES_START
 #pragma warning(push)
 #pragma warning(disable:4005)
+#include <dwmapi.h>
 #include <strsafe.h>
 #include <wincodec.h>
 #pragma warning(pop)
@@ -507,6 +508,10 @@ uint32 WINAPI StartSplashScreenThread( LPVOID unused )
 				(HMENU) NULL,
 				hInstance,
 				(LPVOID) NULL);
+
+			// Enable W11 window style (drop shadow, thin border, slightly rounded)
+			constexpr DWM_WINDOW_CORNER_PREFERENCE Corners = DWMWCP_ROUND;
+			DwmSetWindowAttribute(GSplashScreenWnd, DWMWA_WINDOW_CORNER_PREFERENCE, &Corners, sizeof(Corners));
 		}
 	}
 
