Fix compilation of Windows Splash Screen on W10 SDK.

#rnx
#Synced-CL 51626022

diff --git a/Engine/Source/Runtime/ApplicationCore/Private/Windows/WindowsPlatformSplash.cpp b/Engine/Source/Runtime/ApplicationCore/Private/Windows/WindowsPlatformSplash.cpp
--- a/Engine/Source/Runtime/ApplicationCore/Private/Windows/WindowsPlatformSplash.cpp
+++ b/Engine/Source/Runtime/ApplicationCore/Private/Windows/WindowsPlatformSplash.cpp
@@ -509,9 +509,11 @@ uint32 WINAPI StartSplashScreenThread( LPVOID unused )
 				hInstance,
 				(LPVOID) NULL);
 
+#if defined(NTDDI_WIN10_CO) && WDK_NTDDI_VERSION >= NTDDI_WIN10_CO
 			// Enable W11 window style (drop shadow, thin border, slightly rounded)
 			constexpr DWM_WINDOW_CORNER_PREFERENCE Corners = DWMWCP_ROUND;
 			DwmSetWindowAttribute(GSplashScreenWnd, DWMWA_WINDOW_CORNER_PREFERENCE, &Corners, sizeof(Corners));
+#endif
 		}
 	}
 
