[Android] Fix black screen when app is backgrounded repeatedly during early startup

Fixed native window look-up by preserving the preivous handle.

#rnx
#Synced-CL 51619035

diff --git a/Engine/Source/Runtime/ApplicationCore/Private/Android/AndroidWindow.cpp b/Engine/Source/Runtime/ApplicationCore/Private/Android/AndroidWindow.cpp
--- a/Engine/Source/Runtime/ApplicationCore/Private/Android/AndroidWindow.cpp
+++ b/Engine/Source/Runtime/ApplicationCore/Private/Android/AndroidWindow.cpp
@@ -116,6 +116,7 @@ TSharedRef<FAndroidWindow> FAndroidWindow::Make()
 		FNativeAccessor MyANativeAccessor = Window->GetANativeAccessor(true);
 		FAndroidWindowManager::FPendingWindowAccessor PendingAccessor = FAndroidWindowManager::GetPendingWindowAccessor();
 
+		MyANativeAccessor.SetANativeWindow(PendingAccessor.GetPreviousANativeWindow()); // first set previous handle, this is retained when the current handle is set next
 		MyANativeAccessor.SetANativeWindow(PendingAccessor.GetANativeWindow());
 		PendingAccessor.SetANativeWindow(nullptr);
 	}
diff --git a/Engine/Source/Runtime/ApplicationCore/Public/Android/AndroidWindow.h b/Engine/Source/Runtime/ApplicationCore/Public/Android/AndroidWindow.h
--- a/Engine/Source/Runtime/ApplicationCore/Public/Android/AndroidWindow.h
+++ b/Engine/Source/Runtime/ApplicationCore/Public/Android/AndroidWindow.h
@@ -211,6 +211,7 @@ public:
 	// this is used to collect the native window before the game has created its own FAndroidWindow.
 	class FPendingWindowAccessor
 	{
+		static inline ANativeWindow* PreviousPendingANativeWindow = nullptr;
 		static inline ANativeWindow* PendingANativeWindow = nullptr;
 		static inline UE::FMutex Lock;
 		FPendingWindowAccessor(const FPendingWindowAccessor&) = delete;
@@ -219,7 +220,13 @@ public:
 		FPendingWindowAccessor() { Lock.Lock(); }
 		~FPendingWindowAccessor() { Lock.Unlock(); }
 
-		void SetANativeWindow(ANativeWindow* ANativeHandle) { PendingANativeWindow = ANativeHandle; }
+		void SetANativeWindow(ANativeWindow* ANativeHandle) 
+		{ 
+			PreviousPendingANativeWindow = PendingANativeWindow;
+			PendingANativeWindow = ANativeHandle;
+		}
+
+		ANativeWindow* GetPreviousANativeWindow() const { return PreviousPendingANativeWindow; }
 		ANativeWindow* GetANativeWindow() const { return PendingANativeWindow; }
 	};
 
