[iOS/macOS] Fixes for Clang 21 implicit conversion warnings.

Resolved integer-to-float/double implicit conversions across iOS and Mac platform code by adding explicit casts, and fixed SNew self-assignment warnings by switching to SAssignNew.

#Synced-CL 51857304

diff --git a/Engine/Plugins/Experimental/Animation/HierarchyTable/Source/Editor/Private/HierarchyTableFactory.cpp b/Engine/Plugins/Experimental/Animation/HierarchyTable/Source/Editor/Private/HierarchyTableFactory.cpp
--- a/Engine/Plugins/Experimental/Animation/HierarchyTable/Source/Editor/Private/HierarchyTableFactory.cpp
+++ b/Engine/Plugins/Experimental/Animation/HierarchyTable/Source/Editor/Private/HierarchyTableFactory.cpp
@@ -103,7 +103,8 @@ bool UHierarchyTableFactory::ConfigureTableType()
 		Options.bAllowViewOptions = false;
 	}
 	
-	TSharedPtr<class SWindow> PickerWindow = SNew(SWindow)
+	TSharedPtr<class SWindow> PickerWindow;
+	PickerWindow = SAssignNew(PickerWindow, SWindow)
 		.Title(INVTEXT("Pick Type"))
 		.ClientSize(FVector2D(500, 600))
 		.SupportsMinimize(false)
@@ -163,7 +164,8 @@ bool UHierarchyTableFactory::ConfigureElementType()
 		Options.bAllowViewOptions = false;
 	}
 	
-	TSharedPtr<class SWindow> PickerWindow = SNew(SWindow)
+	TSharedPtr<class SWindow> PickerWindow;
+	PickerWindow = SAssignNew(PickerWindow, SWindow)
 		.Title(INVTEXT("Pick Type"))
 		.ClientSize(FVector2D(500, 600))
 		.SupportsMinimize(false)
diff --git a/Engine/Plugins/Experimental/UAF/UAF/Source/UAFEditor/Private/UAF/BlendProfile/BlendProfileFactory.cpp b/Engine/Plugins/Experimental/UAF/UAF/Source/UAFEditor/Private/UAF/BlendProfile/BlendProfileFactory.cpp
--- a/Engine/Plugins/Experimental/UAF/UAF/Source/UAFEditor/Private/UAF/BlendProfile/BlendProfileFactory.cpp
+++ b/Engine/Plugins/Experimental/UAF/UAF/Source/UAFEditor/Private/UAF/BlendProfile/BlendProfileFactory.cpp
@@ -57,7 +57,8 @@ bool UUAFBlendProfileFactory::ConfigureBlendProfileType()
 {
 	bool bConfirmClicked = false;
 
-	TSharedRef<SWindow> Window = SNew(SWindow)
+	TSharedPtr<SWindow> Window;
+	Window = SAssignNew(Window, SWindow)
 		.Title(LOCTEXT("Title", "Choose Blend Profile Type"))
 		.ClientSize(FVector2D(400, 400))
 		.SupportsMinimize(false)
@@ -95,7 +96,7 @@ bool UUAFBlendProfileFactory::ConfigureBlendProfileType()
 		];
 
 
-	GEditor->EditorAddModalWindow(Window);
+	GEditor->EditorAddModalWindow(Window.ToSharedRef());
 
 	return bConfirmClicked;
 }
diff --git a/Engine/Source/Runtime/ApplicationCore/Private/IOS/IOSAppDelegate.cpp b/Engine/Source/Runtime/ApplicationCore/Private/IOS/IOSAppDelegate.cpp
--- a/Engine/Source/Runtime/ApplicationCore/Private/IOS/IOSAppDelegate.cpp
+++ b/Engine/Source/Runtime/ApplicationCore/Private/IOS/IOSAppDelegate.cpp
@@ -480,7 +480,7 @@ static IOSAppDelegate* CachedDelegate = nil;
 			}
 #endif
 
-			if (OutputVolume != self.LastVolume || self.bForceEmitVolume)
+			if (static_cast<float>(OutputVolume) != self.LastVolume || self.bForceEmitVolume)
 			{
 				FGraphEventRef AudioMuteTask = FFunctionGraphTask::CreateAndDispatchWhenReady([bMuted, OutputVolume]()
 					{
@@ -488,7 +488,7 @@ static IOSAppDelegate* CachedDelegate = nil;
 						FCoreDelegates::AudioMuteDelegate.Broadcast(bMuted, OutputVolume);
 					}, TStatId(), NULL, ENamedThreads::GameThread);
 
-				self.LastVolume = OutputVolume;
+				self.LastVolume = static_cast<float>(OutputVolume);
 				self.bForceEmitVolume = false;
 			}
 
diff --git a/Engine/Source/Runtime/ApplicationCore/Private/IOS/IOSView.cpp b/Engine/Source/Runtime/ApplicationCore/Private/IOS/IOSView.cpp
--- a/Engine/Source/Runtime/ApplicationCore/Private/IOS/IOSView.cpp
+++ b/Engine/Source/Runtime/ApplicationCore/Private/IOS/IOSView.cpp
@@ -286,12 +286,12 @@ MTL::Device* GMetalDevice = nullptr;
 			if (ScreenHeight > ScreenWidth)
 			{
 				Height = RequestedResX;
-				Width = FMath::TruncToInt(Height * AspectRatio + 0.5f);
+				Width = FMath::TruncToInt(static_cast<float>(Height) * AspectRatio + 0.5f);
 			}
 			else
 			{
 				Width = RequestedResX;
-				Height = FMath::TruncToInt(Width * AspectRatio + 0.5f);
+				Height = FMath::TruncToInt(static_cast<float>(Width) * AspectRatio + 0.5f);
 			}
 		}
 		else if (RequestedResY > 0)
@@ -300,12 +300,12 @@ MTL::Device* GMetalDevice = nullptr;
 			if (ScreenHeight > ScreenWidth)
 			{
 				Width = RequestedResY;
-				Height = FMath::TruncToInt(Width * AspectRatio + 0.5f);
+				Height = FMath::TruncToInt(static_cast<float>(Width) * AspectRatio + 0.5f);
 			}
 			else
 			{
 				Height = RequestedResY;
-				Width = FMath::TruncToInt(Height * AspectRatio + 0.5f);
+				Width = FMath::TruncToInt(static_cast<float>(Height) * AspectRatio + 0.5f);
 			}
 		}
 		else
diff --git a/Engine/Source/Runtime/ApplicationCore/Private/Mac/MacPlatformSplash.cpp b/Engine/Source/Runtime/ApplicationCore/Private/Mac/MacPlatformSplash.cpp
--- a/Engine/Source/Runtime/ApplicationCore/Private/Mac/MacPlatformSplash.cpp
+++ b/Engine/Source/Runtime/ApplicationCore/Private/Mac/MacPlatformSplash.cpp
@@ -247,8 +247,8 @@ void FMacPlatformSplash::Show()
 			[SplashScreenFileName release];
 			
 			NSBitmapImageRep* BitmapRep = [NSBitmapImageRep imageRepWithData: [GSplashScreenImage TIFFRepresentation]];
-			float ImageWidth = [BitmapRep pixelsWide];
-			float ImageHeight = [BitmapRep pixelsHigh];
+			float ImageWidth = static_cast<float>([BitmapRep pixelsWide]);
+			float ImageHeight = static_cast<float>([BitmapRep pixelsHigh]);
 			[GSplashScreenImage setSize:NSMakeSize(ImageWidth, ImageHeight)];
 			
 			ContentRect.origin.x = 0;
@@ -263,7 +263,7 @@ void FMacPlatformSplash::Show()
 			// Setup bounds for game name
 			GSplashScreenTextRects[ SplashTextType::GameName ].origin.x = OriginX;
 			GSplashScreenTextRects[ SplashTextType::GameName ].origin.y = GIsEditor ? OriginY + 34 : OriginY + 18;
-			GSplashScreenTextRects[ SplashTextType::GameName ].size.width = ImageWidth - 2 * OriginX;
+			GSplashScreenTextRects[ SplashTextType::GameName ].size.width = ImageWidth - static_cast<float>(2 * OriginX);
 			GSplashScreenTextRects[ SplashTextType::GameName ].size.height = 2 * FontHeight;
 			
 			// Setup bounds for texts
@@ -273,7 +273,7 @@ void FMacPlatformSplash::Show()
 			
 			GSplashScreenTextRects[ SplashTextType::VersionInfo1 ].size.width =
 			GSplashScreenTextRects[ SplashTextType::CopyrightInfo ].size.width =
-			GSplashScreenTextRects[ SplashTextType::StartupProgress ].size.width = ImageWidth - 2 * OriginX;
+			GSplashScreenTextRects[ SplashTextType::StartupProgress ].size.width = ImageWidth - static_cast<float>(2 * OriginX);
 			
 			GSplashScreenTextRects[ SplashTextType::VersionInfo1 ].size.height =
 			GSplashScreenTextRects[ SplashTextType::StartupProgress ].size.height = FontHeight;
diff --git a/Engine/Source/Runtime/ApplicationCore/Private/Mac/MacPlatformSurvey.cpp b/Engine/Source/Runtime/ApplicationCore/Private/Mac/MacPlatformSurvey.cpp
--- a/Engine/Source/Runtime/ApplicationCore/Private/Mac/MacPlatformSurvey.cpp
+++ b/Engine/Source/Runtime/ApplicationCore/Private/Mac/MacPlatformSurvey.cpp
@@ -31,7 +31,7 @@ bool FMacPlatformSurvey::GetSurveyResults( FHardwareSurveyResults& OutResults, b
 	uint64_t FreeMem = (Stats.free_count + Stats.inactive_count) * PageSize;
 	uint64_t UsedMem = (Stats.active_count + Stats.wire_count) * PageSize;
 	uint64_t TotalPhys = FreeMem + UsedMem;
-	OutResults.MemoryMB = uint32(float(TotalPhys/1024.0/1024.0)+ .1f);
+	OutResults.MemoryMB = uint32(float(static_cast<double>(TotalPhys)/1024.0/1024.0)+ .1f);
 
 	// Identify Display Devices
 	FDisplayMetrics DisplayMetrics;
@@ -71,7 +71,7 @@ bool FMacPlatformSurvey::GetSurveyResults( FHardwareSurveyResults& OutResults, b
 		Size = sizeof(int64);
 		if (sysctlbyname("hw.cpufrequency", &CPUSpeed, &Size, NULL, 0) == 0)
 		{
-			OutResults.CPUClockGHz = (float)(0.000000001 * CPUSpeed);
+			OutResults.CPUClockGHz = (float)(0.000000001 * static_cast<double>(CPUSpeed));
 		}
 		else
 		{
diff --git a/Engine/Source/Runtime/ApplicationCore/Private/Mac/MacWindow.cpp b/Engine/Source/Runtime/ApplicationCore/Private/Mac/MacWindow.cpp
--- a/Engine/Source/Runtime/ApplicationCore/Private/Mac/MacWindow.cpp
+++ b/Engine/Source/Runtime/ApplicationCore/Private/Mac/MacWindow.cpp
@@ -502,7 +502,7 @@ bool FMacWindow::IsPointInWindow( int32 X, int32 Y ) const
 		if (WindowHandle->bIsOnActiveSpace)
 		{
 			const float DPIScaleFactor = GetDPIScaleFactor();
-			NSPoint CursorPoint = NSMakePoint(X / DPIScaleFactor, WindowFrame.size.height - (Y / DPIScaleFactor + 1));
+			NSPoint CursorPoint = NSMakePoint(static_cast<float>(X) / DPIScaleFactor, WindowFrame.size.height - (static_cast<float>(Y) / DPIScaleFactor + 1));
 			PointInWindow = (NSPointInRect(CursorPoint, VisibleFrame) == YES);
 		}
 	}
@@ -653,9 +653,9 @@ void FMacWindow::ApplySizeAndModeChanges(int32 X, int32 Y, int32 Width, int32 He
 		{
 			WindowHandle.TargetWindowMode = WindowMode;
 
-			const float DPIScaleFactor = FPlatformApplicationMisc::GetDPIScaleFactorAtPoint(X, Y);
-			Width = FMath::CeilToInt(Width / DPIScaleFactor);
-			Height = FMath::CeilToInt(Height / DPIScaleFactor);
+			const float DPIScaleFactor = FPlatformApplicationMisc::GetDPIScaleFactorAtPoint(static_cast<float>(X), static_cast<float>(Y));
+			Width = FMath::CeilToInt(static_cast<float>(Width) / DPIScaleFactor);
+			Height = FMath::CeilToInt(static_cast<float>(Height) / DPIScaleFactor);
 
 			const FVector2D CocoaPosition = FMacApplication::ConvertSlatePositionToCocoa(X, Y);
 			NSRect Rect = NSMakeRect(CocoaPosition.X, CocoaPosition.Y - Height + 1, FMath::Max(Width, 1), FMath::Max(Height, 1));
diff --git a/Engine/Source/Runtime/Core/Private/Apple/ApplePlatformFile.cpp b/Engine/Source/Runtime/Core/Private/Apple/ApplePlatformFile.cpp
--- a/Engine/Source/Runtime/Core/Private/Apple/ApplePlatformFile.cpp
+++ b/Engine/Source/Runtime/Core/Private/Apple/ApplePlatformFile.cpp
@@ -49,9 +49,9 @@ namespace
 		}
 
 		return FFileStatData(
-			MacEpoch + FTimespan::FromSeconds(FileInfo.st_ctime), 
-			MacEpoch + FTimespan::FromSeconds(FileInfo.st_atime), 
-			MacEpoch + FTimespan::FromSeconds(FileInfo.st_mtime), 
+			MacEpoch + FTimespan::FromSeconds(static_cast<double>(FileInfo.st_ctime)), 
+			MacEpoch + FTimespan::FromSeconds(static_cast<double>(FileInfo.st_atime)), 
+			MacEpoch + FTimespan::FromSeconds(static_cast<double>(FileInfo.st_mtime)), 
 			FileSize,
 			bIsDirectory,
 			!(FileInfo.st_mode & S_IWUSR)
diff --git a/Engine/Source/Runtime/Core/Private/Apple/ApplePlatformMemory.cpp b/Engine/Source/Runtime/Core/Private/Apple/ApplePlatformMemory.cpp
--- a/Engine/Source/Runtime/Core/Private/Apple/ApplePlatformMemory.cpp
+++ b/Engine/Source/Runtime/Core/Private/Apple/ApplePlatformMemory.cpp
@@ -215,10 +215,10 @@ void FApplePlatformMemory::Init()
 	
 	const FPlatformMemoryConstants& MemoryConstants = FPlatformMemory::GetConstants();
 	UE_LOG(LogInit, Log, TEXT("Memory total: Physical=%.1fGB (%dGB approx) Pagefile=%.1fGB Virtual=%.1fGB"),
-		   float(MemoryConstants.TotalPhysical/1024.0/1024.0/1024.0),
-		   MemoryConstants.TotalPhysicalGB,
-		   float((MemoryConstants.TotalVirtual-MemoryConstants.TotalPhysical)/1024.0/1024.0/1024.0),
-		   float(MemoryConstants.TotalVirtual/1024.0/1024.0/1024.0) );
+	   float(static_cast<double>(MemoryConstants.TotalPhysical)/1024.0/1024.0/1024.0),
+	   MemoryConstants.TotalPhysicalGB,
+	   float(static_cast<double>(MemoryConstants.TotalVirtual-MemoryConstants.TotalPhysical)/1024.0/1024.0/1024.0),
+	   float(static_cast<double>(MemoryConstants.TotalVirtual)/1024.0/1024.0/1024.0) );
 	
 }
 
diff --git a/Engine/Source/Runtime/Core/Private/Apple/ApplePlatformTime.cpp b/Engine/Source/Runtime/Core/Private/Apple/ApplePlatformTime.cpp
--- a/Engine/Source/Runtime/Core/Private/Apple/ApplePlatformTime.cpp
+++ b/Engine/Source/Runtime/Core/Private/Apple/ApplePlatformTime.cpp
@@ -40,7 +40,7 @@ double FApplePlatformTime::Seconds()
 {
 	uint64 Cycles = mach_absolute_time();
 	// Add big number to make bugs apparent where return value is being passed to float
-	return Cycles * GetSecondsPerCycle() + GetSecondsTimeOffset();
+	return static_cast<double>(Cycles) * GetSecondsPerCycle() + GetSecondsTimeOffset();
 }
 
 uint32 FApplePlatformTime::Cycles()
diff --git a/Engine/Source/Runtime/Core/Private/IOS/IOSPlatformFile.cpp b/Engine/Source/Runtime/Core/Private/IOS/IOSPlatformFile.cpp
--- a/Engine/Source/Runtime/Core/Private/IOS/IOSPlatformFile.cpp
+++ b/Engine/Source/Runtime/Core/Private/IOS/IOSPlatformFile.cpp
@@ -49,9 +49,9 @@ namespace
 		}
 
 		return FFileStatData(
-			IOSEpoch + FTimespan::FromSeconds(FileInfo.st_ctime), 
-			IOSEpoch + FTimespan::FromSeconds(FileInfo.st_atime), 
-			IOSEpoch + FTimespan::FromSeconds(FileInfo.st_mtime), 
+			IOSEpoch + FTimespan::FromSeconds(static_cast<double>(FileInfo.st_ctime)),
+			IOSEpoch + FTimespan::FromSeconds(static_cast<double>(FileInfo.st_atime)),
+			IOSEpoch + FTimespan::FromSeconds(static_cast<double>(FileInfo.st_mtime)), 
 			FileSize,
 			bIsDirectory,
 			!(FileInfo.st_mode & S_IWUSR)
diff --git a/Engine/Source/Runtime/Core/Private/IOS/IOSPlatformMisc.cpp b/Engine/Source/Runtime/Core/Private/IOS/IOSPlatformMisc.cpp
--- a/Engine/Source/Runtime/Core/Private/IOS/IOSPlatformMisc.cpp
+++ b/Engine/Source/Runtime/Core/Private/IOS/IOSPlatformMisc.cpp
@@ -1271,7 +1271,7 @@ void FIOSPlatformMisc::ShareURL(const FString& URL, const FText& Description, in
 	NSURL* SharedURL = [NSURL URLWithString:[NSString stringWithFString:URL]];
     
     float ScaleFactor = (float)[IOSAppDelegate GetDelegate].IOSView.contentScaleFactor;
-    __block CGRect PopoverLocation = CGRectMake(LocationHintX/ScaleFactor, LocationHintY/ScaleFactor, 10, 10);
+    __block CGRect PopoverLocation = CGRectMake(static_cast<float>(LocationHintX)/ScaleFactor, static_cast<float>(LocationHintY)/ScaleFactor, 10, 10);
 
 	dispatch_async(dispatch_get_main_queue(),^ {
 		NSArray* ObjectsToShare = @[SharedString, SharedURL];
