- Changed WIndowsConsoleOutputDevice2.cpp (-newconsole) to send input to Engine Cmd handler, even if python is installed (Python was stealing the comment input)
#Synced-CL 51858477

diff --git a/Engine/Source/Runtime/ApplicationCore/Private/Windows/WindowsConsoleOutputDevice2.cpp b/Engine/Source/Runtime/ApplicationCore/Private/Windows/WindowsConsoleOutputDevice2.cpp
--- a/Engine/Source/Runtime/ApplicationCore/Private/Windows/WindowsConsoleOutputDevice2.cpp
+++ b/Engine/Source/Runtime/ApplicationCore/Private/Windows/WindowsConsoleOutputDevice2.cpp
@@ -2516,10 +2516,11 @@ public:
 				// The reasoning is that the new MultiServer Proxy does not execute TickGroups (UE_SUPPORT_FOR_ACTOR_TICK_DISABLE).
 				ExecuteOnGameThread(UE_SOURCE_LOCATION, [Cmd = FString(Command)]()
 					{
-						if (IModularFeatures::Get().IsModularFeatureAvailable(IConsoleCommandExecutor::ModularFeatureName()))
+						TArray<IConsoleCommandExecutor*> CmdExecutors = IModularFeatures::Get().GetModularFeatureImplementations<IConsoleCommandExecutor>(IConsoleCommandExecutor::ModularFeatureName());
+						if (IConsoleCommandExecutor** CmdExecutor = CmdExecutors.FindByPredicate([](IConsoleCommandExecutor* CmdExecutor) { return CmdExecutor->GetName() == NAME_Cmd; }))
 						{
 							UE_LOG(LogExec, Log, TEXT("Executing console command: %s"), *Cmd);
-							IModularFeatures::Get().GetModularFeature<IConsoleCommandExecutor>(FName(IConsoleCommandExecutor::ModularFeatureName())).Exec(*Cmd);
+							(*CmdExecutor)->Exec(*Cmd);
 						}
 						else
 						{
