[Mac] Replace the current high-precision mouse movement implementation with the GCMouse-based implementation already used for iOS devices.

GCMouse provides raw mouse movement deltas which are better suited for high-precision mode, and our iOS implementation handles delta fraction accumulation between frames properly.

Mouse button events and scroll movement events are still handled via AppKit, as we require context about which window an event occurs in, and GCMouse does not provide that information. To achieve that, a new set of capabilities flags was added to FAppleMouseController, so we can only enable the pieces we actually need.

#Synced-CL 51901069

diff --git a/Engine/Source/Runtime/ApplicationCore/Private/Apple/AppleMouseController.cpp b/Engine/Source/Runtime/ApplicationCore/Private/Apple/AppleMouseController.cpp
--- a/Engine/Source/Runtime/ApplicationCore/Private/Apple/AppleMouseController.cpp
+++ b/Engine/Source/Runtime/ApplicationCore/Private/Apple/AppleMouseController.cpp
@@ -9,7 +9,7 @@
 #import "GameController/GCDeviceCursor.h"
 #import "GameController/GCMouseInput.h"
 
-FAppleMouseController::FAppleMouseController(const TSharedRef<FGenericApplicationMessageHandler>& InMessageHandler) : MessageHandler(InMessageHandler)
+FAppleMouseController::FAppleMouseController(const TSharedRef<FGenericApplicationMessageHandler>& InMessageHandler, EMouseControllerCapabilities InCapabilities) : MessageHandler(InMessageHandler), EnabledCapabilities(InCapabilities)
 {
 	if (@available(macOS 11.0, iOS 14.0, *))
 	{
@@ -83,41 +83,58 @@ void FAppleMouseController::HandleMouseConnected(GCMouse* Mouse)
 	
 	if (ensure(MouseInput))
 	{
-		MouseInput.mouseMovedHandler = ^(GCMouseInput * mouse, float deltaX, float deltaY)
+		if (EnumHasAnyFlags(EnabledCapabilities, EMouseControllerCapabilities::Movement))
 		{
-			CurrentMovementState.MouseDeltaX += deltaX;
-			CurrentMovementState.MouseDeltaY -= deltaY;
-		};
-		
-		MouseInput.scroll.valueChangedHandler = ^(GCControllerDirectionPad* dpad, float xValue, float yValue) 
-		{
-			CurrentMovementState.ScrollDelta += yValue;
-		};
-		
-		MouseInput.leftButton.pressedChangedHandler = ^(GCControllerButtonInput* eventButton, float value, BOOL pressed) 
-		{
-			HandleKeyEvent(pressed, EMouseButtons::Type::Left);
-		};
-
-		MouseInput.rightButton.pressedChangedHandler = ^(GCControllerButtonInput* eventButton, float value, BOOL pressed) {
-			HandleKeyEvent(pressed, EMouseButtons::Type::Right);
-		};
+			MouseInput.mouseMovedHandler = ^(GCMouseInput * mouse, float deltaX, float deltaY)
+			{
+				if (!bEnabled)
+				{
+					return;
+				}
+				CurrentMovementState.MouseDeltaX += deltaX;
+				CurrentMovementState.MouseDeltaY -= deltaY;
+			};
+		}
 
-		MouseInput.middleButton.pressedChangedHandler = ^(GCControllerButtonInput* eventButton, float value, BOOL pressed) 
+		if (EnumHasAnyFlags(EnabledCapabilities, EMouseControllerCapabilities::Scroll))
 		{
-			HandleKeyEvent(pressed, EMouseButtons::Type::Middle);
-		};
-		
-		constexpr uint64 MaxAuxButtonIndex = EMouseButtons::Type::Invalid - EMouseButtons::Type::Thumb01;
-		const uint64 CurrentMaxAuxButtonIndex = MaxAuxButtonIndex > MouseInput.auxiliaryButtons.count ? MouseInput.auxiliaryButtons.count : MaxAuxButtonIndex;
+			MouseInput.scroll.valueChangedHandler = ^(GCControllerDirectionPad* dpad, float xValue, float yValue)
+			{
+				if (!bEnabled)
+				{
+					return;
+				}
+				CurrentMovementState.ScrollDelta += yValue;
+			};
+		}
 
-		for (uint64 AuxButtonIndex = 0; AuxButtonIndex < CurrentMaxAuxButtonIndex; AuxButtonIndex++)
+		if (EnumHasAnyFlags(EnabledCapabilities, EMouseControllerCapabilities::Buttons))
 		{
-			MouseInput.auxiliaryButtons[AuxButtonIndex].pressedChangedHandler = ^(GCControllerButtonInput* eventButton, float value, BOOL pressed) 
+			MouseInput.leftButton.pressedChangedHandler = ^(GCControllerButtonInput* eventButton, float value, BOOL pressed)
 			{
-				const EMouseButtons::Type TargetButton = static_cast<EMouseButtons::Type>(EMouseButtons::Type::Thumb01 + AuxButtonIndex);
-				HandleKeyEvent(pressed, TargetButton);
+				HandleKeyEvent(pressed, EMouseButtons::Type::Left);
 			};
+
+			MouseInput.rightButton.pressedChangedHandler = ^(GCControllerButtonInput* eventButton, float value, BOOL pressed) {
+				HandleKeyEvent(pressed, EMouseButtons::Type::Right);
+			};
+
+			MouseInput.middleButton.pressedChangedHandler = ^(GCControllerButtonInput* eventButton, float value, BOOL pressed)
+			{
+				HandleKeyEvent(pressed, EMouseButtons::Type::Middle);
+			};
+
+			constexpr uint64 MaxAuxButtonIndex = EMouseButtons::Type::Invalid - EMouseButtons::Type::Thumb01;
+			const uint64 CurrentMaxAuxButtonIndex = MaxAuxButtonIndex > MouseInput.auxiliaryButtons.count ? MouseInput.auxiliaryButtons.count : MaxAuxButtonIndex;
+
+			for (uint64 AuxButtonIndex = 0; AuxButtonIndex < CurrentMaxAuxButtonIndex; AuxButtonIndex++)
+			{
+				MouseInput.auxiliaryButtons[AuxButtonIndex].pressedChangedHandler = ^(GCControllerButtonInput* eventButton, float value, BOOL pressed)
+				{
+					const EMouseButtons::Type TargetButton = static_cast<EMouseButtons::Type>(EMouseButtons::Type::Thumb01 + AuxButtonIndex);
+					HandleKeyEvent(pressed, TargetButton);
+				};
+			}
 		}
 	}
 	
@@ -126,6 +143,10 @@ void FAppleMouseController::HandleMouseConnected(GCMouse* Mouse)
 
 void FAppleMouseController::HandleKeyEvent(bool bIsPressed, EMouseButtons::Type Button)
 {
+	if (!bEnabled)
+	{
+		return;
+	}
 	FDeferredKeyEvent KeyEvent;
 	KeyEvent.ButtonID = Button;
 	KeyEvent.Type = bIsPressed ? EKeyEventType::Down : EKeyEventType::Up;
@@ -134,22 +155,30 @@ void FAppleMouseController::HandleKeyEvent(bool bIsPressed, EMouseButtons::Type
 
 void FAppleMouseController::DispatchMovementUpdates()
 {
-	FVector2D RawMovementDelta(CurrentMovementState.MouseDeltaX.exchange(0.0f), CurrentMovementState.MouseDeltaY.exchange(0.0f));
-	
-	// UE API only takes integer values as that is the unity most usb mouses report, but some devices like trackpads can return fractional, subpixel, movement deltas
-	// To no lose that fractional delta value, we store it in a pending delta property, and apply it in the next loop
-	RawMovementDelta += CurrentMovementState.PendingDeltaXY;
-	CurrentMovementState.PendingDeltaXY.X = FMath::Fractional(RawMovementDelta.X);
-	CurrentMovementState.PendingDeltaXY.Y = FMath::Fractional(RawMovementDelta.Y);
-	
-	FIntVector2 MovementDelta(FMath::TruncToInt(RawMovementDelta.X), FMath::TruncToInt(RawMovementDelta.Y));
-	
-	MessageHandler->OnRawMouseMove(MovementDelta.X, MovementDelta.Y);
-	MessageHandler->OnMouseMove();
-	MessageHandler->OnCursorSet();
-	MessageHandler->OnMouseWheel(CurrentMovementState.ScrollDelta.exchange(0.0f));
+	if (EnumHasAnyFlags(EnabledCapabilities, EMouseControllerCapabilities::Movement))
+	{
+		FVector2D RawMovementDelta(CurrentMovementState.MouseDeltaX.exchange(0.0f), CurrentMovementState.MouseDeltaY.exchange(0.0f));
+		
+		// UE API only takes integer values as that is the unity most usb mouses report, but some devices like trackpads can return fractional, subpixel, movement deltas
+		// To not lose that fractional delta value, we store it in a pending delta property, and apply it in the next loop
+		RawMovementDelta += CurrentMovementState.PendingDeltaXY;
+		CurrentMovementState.PendingDeltaXY.X = FMath::Fractional(RawMovementDelta.X);
+		CurrentMovementState.PendingDeltaXY.Y = FMath::Fractional(RawMovementDelta.Y);
+
+		FIntVector2 MovementDelta(FMath::TruncToInt(RawMovementDelta.X), FMath::TruncToInt(RawMovementDelta.Y));
+
+		MessageHandler->OnRawMouseMove(MovementDelta.X, MovementDelta.Y);
+		MessageHandler->OnMouseMove();
+		MessageHandler->OnCursorSet();
+		
+		MovementDispatchedDelegate.Broadcast(MovementDelta);
+	}
 	
-	MovementDispatchedDelegate.Broadcast(MovementDelta);
+	if (EnumHasAnyFlags(EnabledCapabilities, EMouseControllerCapabilities::Scroll))
+	{
+		const float ScrollDelta = CurrentMovementState.ScrollDelta.exchange(0.0f);
+		MessageHandler->OnMouseWheel(ScrollDelta);
+	}
 }
 
 void FAppleMouseController::DispatchKeyUpdates()
@@ -189,6 +218,28 @@ void FAppleMouseController::HandleMouseDisconnected()
 	UE_LOG(LogAppleController, Log, TEXT("[%hs] Handling physical mouse disconnection..."), __func__);
 }
 
+void FAppleMouseController::SetEnabled(bool bInEnabled)
+{
+	if (bEnabled == bInEnabled)
+	{
+		return;
+	}
+
+	UE_LOG(LogAppleController, Verbose, TEXT("[%hs] %s"), __func__, bInEnabled ? TEXT("Enabling") : TEXT("Disabling"));
+
+	bEnabled = bInEnabled;
+
+	if (!bInEnabled)
+	{
+		bNeedFlushPendingState = true;
+	}
+}
+
+bool FAppleMouseController::IsEnabled() const
+{
+	return bEnabled;
+}
+
 void FAppleMouseController::SetMessageHandler(const TSharedRef<FGenericApplicationMessageHandler>& InMessageHandler)
 {
 	MessageHandler = InMessageHandler;
@@ -201,8 +252,20 @@ void FAppleMouseController::SendControllerEvents()
 		return;
 	}
 
-	DispatchMovementUpdates();
-	DispatchKeyUpdates();
+	if (!bEnabled && !bNeedFlushPendingState.exchange(false))
+	{
+		return;
+	}
+
+	if (EnumHasAnyFlags(EnabledCapabilities, EMouseControllerCapabilities::Movement | EMouseControllerCapabilities::Scroll))
+	{
+		DispatchMovementUpdates();
+	}
+
+	if (EnumHasAnyFlags(EnabledCapabilities, EMouseControllerCapabilities::Buttons))
+	{
+		DispatchKeyUpdates();
+	}
 }
 
 bool FAppleMouseController::IsAnyMouseConnected() const
diff --git a/Engine/Source/Runtime/ApplicationCore/Private/Apple/AppleMouseController.h b/Engine/Source/Runtime/ApplicationCore/Private/Apple/AppleMouseController.h
--- a/Engine/Source/Runtime/ApplicationCore/Private/Apple/AppleMouseController.h
+++ b/Engine/Source/Runtime/ApplicationCore/Private/Apple/AppleMouseController.h
@@ -7,6 +7,7 @@
 #include <atomic>
 
 #include "Containers/SpscQueue.h"
+#include "Misc/EnumClassFlags.h"
 
 #if PLATFORM_MAC
 #include "Mac/MacSystemIncludes.h"
@@ -14,13 +15,24 @@
 
 #import "GameController/GCMouse.h"
 
+/** Controls which GCMouse handlers FAppleMouseController binds */
+enum class EMouseControllerCapabilities : uint8
+{
+	None       = 0,
+	Movement   = 1 << 0,
+	Scroll     = 1 << 1,
+	Buttons    = 1 << 2,
+	All        = Movement | Scroll | Buttons
+};
+ENUM_CLASS_FLAGS(EMouseControllerCapabilities)
+
 /**
  * Class that intercepts, queues and dispatches GCMouse events to UE's input system
  */
 class FAppleMouseController
 {
 public:
-	explicit FAppleMouseController(const TSharedRef<FGenericApplicationMessageHandler>& InMessageHandler);
+	explicit FAppleMouseController(const TSharedRef<FGenericApplicationMessageHandler>& InMessageHandler, EMouseControllerCapabilities InCapabilities = EMouseControllerCapabilities::All);
 	
 	~FAppleMouseController();
 	
@@ -33,6 +45,16 @@ public:
 	/** Process and dispatches all queued key events to UE's input system */
 	void SendControllerEvents();
 	
+	/**
+	 * Enables or disables mouse event accumulation and dispatch. Used it to in practice disable this controller,
+	 * but keeping it initialized
+	 * @param bInEnabled new enabled state
+	 */
+	void SetEnabled(bool bInEnabled);
+
+	/** Returns true if event accumulation and dispatch are active */
+	bool IsEnabled() const;
+
 	bool IsAnyMouseConnected() const;
 	
 	DECLARE_MULTICAST_DELEGATE(FAvailabilityChanged)
@@ -95,8 +117,11 @@ protected:
 
 	TSpscQueue<FDeferredKeyEvent> DeferredKeyEvents;
 	TSharedRef<FGenericApplicationMessageHandler> MessageHandler;
-	
+	EMouseControllerCapabilities EnabledCapabilities;
+
 	bool bIsInitialized = false;
+	std::atomic<bool> bEnabled = true;
+	std::atomic<bool> bNeedFlushPendingState = false;
 	
 	id ConnectedEventObserverID = nil;
 	id DisconnectedEventObserverID = nil;
diff --git a/Engine/Source/Runtime/ApplicationCore/Private/Mac/MacApplication.cpp b/Engine/Source/Runtime/ApplicationCore/Private/Mac/MacApplication.cpp
--- a/Engine/Source/Runtime/ApplicationCore/Private/Mac/MacApplication.cpp
+++ b/Engine/Source/Runtime/ApplicationCore/Private/Mac/MacApplication.cpp
@@ -31,6 +31,7 @@
 #include "Apple/ScopeAutoreleasePool.h"
 #include "Apple/AppleStringUtils.h"
 #include "Apple/AppleControllerInterface.h"
+#include "Apple/AppleMouseController.h"
 #include "HIDInputInterface.h"
 
 FMacApplication* MacApplication = nullptr;
@@ -48,6 +49,16 @@ extern "C" void MTDeviceStart(void*, int);
 extern "C" bool MTDeviceIsBuiltIn(void*);
 #endif
 
+namespace MacApplicationCVars
+{
+	static bool bUseNewMouseControllerForMovement = true;
+	static FAutoConsoleVariableRef CVarUseNewMouseControllerMovement(
+		TEXT("Slate.MacUseNewMouseControllerMovement"),
+		bUseNewMouseControllerForMovement,
+		TEXT("false: AppKit handles mouse movement in all modes (legacy). true: FAppleMouseController handles movement during high-precision mode; AppKit handles normal mode."),
+		ECVF_ReadOnly);
+}
+
 // Simple wrapper for Legacy HID Implementation and AppleControllerInterface GCController implementation selection
 // Remove this wrapper and HIDInputInterface once 10.14.0 is no longer supported
 static TAutoConsoleVariable<int32> CVarMacControllerPreferGCImpl(
@@ -58,6 +69,7 @@ static TAutoConsoleVariable<int32> CVarMacControllerPreferGCImpl(
 	TEXT("\t1: Use GCController implementation.\n")
 	TEXT("Default is 1."),
 	ECVF_ReadOnly);
+
 class FMacControllerInterface
 {
 public:
@@ -199,6 +211,15 @@ FMacApplication::FMacApplication()
         SystemResolutionX = ResolutionX;
         SystemResolutionY = ResolutionY;
     });
+
+	if (MacApplicationCVars::bUseNewMouseControllerForMovement)
+	{
+		// Only use the movement capabilities from the controller. For button input handling we need the context about on which window the event is happening
+		// but that is not supported by GC Mouse yet. In this setup we use Appkit to handle Mouse button and non precision movement modes, and the new Mouse controller
+		// (GC Mouse based) for high precision movement
+		MouseController = MakeShared<FAppleMouseController>(MessageHandler, EMouseControllerCapabilities::Movement);
+		MouseController->OnMovementDispatched().AddRaw(this, &FMacApplication::HandleMouseControllerMovement);
+	}
 }
 
 FMacApplication::~FMacApplication()
@@ -276,6 +297,13 @@ FMacApplication::~FMacApplication()
     FCoreDelegates::PreModal.RemoveAll(this);
     FCoreDelegates::PostModal.RemoveAll(this);
 #endif
+
+	if (MouseController)
+	{
+		MouseController->OnMovementDispatched().RemoveAll(this);
+		MouseController.Reset();
+	}
+
 	MacApplication = nullptr;
 }
 
@@ -283,6 +311,11 @@ void FMacApplication::SetMessageHandler(const TSharedRef<FGenericApplicationMess
 {
 	GenericApplication::SetMessageHandler(InMessageHandler);
 	HIDInput->SetMessageHandler(InMessageHandler);
+
+	if (MouseController)
+	{
+		MouseController->SetMessageHandler(InMessageHandler);
+	}
 }
 
 #if WITH_ACCESSIBILITY
@@ -334,6 +367,11 @@ void FMacApplication::PollGameDeviceState(const float TimeDelta)
 	HIDInput->Tick( TimeDelta );
 	HIDInput->SendControllerEvents();
 
+	if (MouseController)
+	{
+		MouseController->SendControllerEvents();
+	}
+
 	// Poll externally-implemented devices
 	for( auto DeviceIt = ExternalInputDevices.CreateIterator(); DeviceIt; ++DeviceIt )
 	{
@@ -421,8 +459,15 @@ TSharedPtr<FGenericWindow> FMacApplication::GetWindowUnderCursor()
 
 void FMacApplication::SetHighPrecisionMouseMode(const bool Enable, const TSharedPtr<FGenericWindow>& InWindow)
 {
+	UE_LOG(LogAppleController, Verbose, TEXT("[%hs] HighPrecisionMouseMode %s"), __func__, Enable ? TEXT("Enabled") : TEXT("Disabled"));
+
 	bUsingHighPrecisionMouseInput = Enable;
 	((FMacCursor*)Cursor.Get())->SetHighPrecisionMouseMode(Enable);
+
+	if (MouseController && MacApplicationCVars::bUseNewMouseControllerForMovement)
+	{
+		MouseController->SetEnabled(Enable);
+	}
 }
 
 bool FMacApplication::IsGamepadAttached() const
@@ -941,6 +986,36 @@ void FMacApplication::ResendEvent(NSEvent* Event)
 	});
 }
 
+void FMacApplication::HandleMouseControllerMovement(FIntVector2 Delta)
+{
+	// GCMouse reports raw hardware deltas — cursor warps (CGWarpMouseCursorPosition) do not
+	// generate GCMouse events, so no warp compensation is needed (unlike the NSEvent path).
+	if (bUsingHighPrecisionMouseInput)
+	{
+		FMacCursor* MacCursor = static_cast<FMacCursor*>(Cursor.Get());
+		FIntVector2 HighPrecisionMousePos = MacCursor->GetIntPosition();
+
+		HighPrecisionMousePos.X += Delta.X;
+		HighPrecisionMousePos.Y += Delta.Y;
+
+		MacCursor->UpdateCursorClipping(HighPrecisionMousePos);
+
+		FMacScreenRef Screen = FindScreenBySlatePosition(HighPrecisionMousePos.X, HighPrecisionMousePos.Y);
+		NSRect VisibleFrame = Screen->VisibleFramePixels;
+
+		// Avoid the menu bar & dock disclosure borders at the top & bottom of fullscreen windows
+		if (ActiveWindow.IsValid() && ActiveWindow->GetWindowMode() != EWindowMode::Windowed)
+		{
+			VisibleFrame.origin.y += 5;
+			VisibleFrame.size.height -= 10;
+		}
+
+		const int32 ClampedPosX = FMath::Clamp(HighPrecisionMousePos.X, static_cast<int32>(VisibleFrame.origin.x), static_cast<int32>(VisibleFrame.origin.x + VisibleFrame.size.width) - 1);
+		const int32 ClampedPosY = FMath::Clamp(HighPrecisionMousePos.Y, static_cast<int32>(VisibleFrame.origin.y), static_cast<int32>(VisibleFrame.origin.y + VisibleFrame.size.height) - 1);
+		MacCursor->SetPosition(ClampedPosX, ClampedPosY);
+	}
+}
+
 void FMacApplication::ProcessMouseMovedEvent(const FDeferredMacEvent& Event, TSharedPtr<FMacWindow> EventWindow)
 {
 	if (EventWindow.IsValid() && EventWindow->IsRegularWindow())
@@ -960,52 +1035,55 @@ void FMacApplication::ProcessMouseMovedEvent(const FDeferredMacEvent& Event, TSh
 
 	if (bUsingHighPrecisionMouseInput)
 	{
-		// Get the mouse position
-		FIntVector2 HighPrecisionMousePos = MacCursor->GetIntPosition();
+		if (!MacApplicationCVars::bUseNewMouseControllerForMovement)
+		{
+			// Get the mouse position
+			FIntVector2 HighPrecisionMousePos = MacCursor->GetIntPosition();
 
-		// Find the visible frame of the screen the cursor is currently on.
-		FMacScreenRef Screen = FindScreenBySlatePosition(HighPrecisionMousePos.X, HighPrecisionMousePos.Y);
-		NSRect VisibleFrame = Screen->VisibleFramePixels;
+			// Find the visible frame of the screen the cursor is currently on.
+			FMacScreenRef Screen = FindScreenBySlatePosition(HighPrecisionMousePos.X, HighPrecisionMousePos.Y);
+			NSRect VisibleFrame = Screen->VisibleFramePixels;
 
-		// Under OS X we disassociate the cursor and mouse position during hi-precision mouse input.
-		// The game snaps the mouse cursor back to the starting point when this is disabled, which
-		// accumulates mouse delta that we want to ignore.
-		const FIntVector2 AccumDelta = MacCursor->GetMouseWarpDelta();
+			// Under OS X we disassociate the cursor and mouse position during hi-precision mouse input.
+			// The game snaps the mouse cursor back to the starting point when this is disabled, which
+			// accumulates mouse delta that we want to ignore.
+			const FIntVector2 AccumDelta = MacCursor->GetMouseWarpDelta();
 
-		// Account for warping delta's
-		FIntVector2 Delta(FMath::TruncToInt(Event.Delta.X), FMath::TruncToInt(Event.Delta.Y));
-		const FIntVector2 WarpDelta(FMath::Abs(AccumDelta.X)<FMath::Abs(Delta.X) ? AccumDelta.X : Delta.X, FMath::Abs(AccumDelta.Y)<FMath::Abs(Delta.Y) ? AccumDelta.Y : Delta.Y);
-		Delta.X -= WarpDelta.X;
-		Delta.Y -= WarpDelta.Y;
+			// Account for warping delta's
+			FIntVector2 Delta(FMath::TruncToInt(Event.Delta.X), FMath::TruncToInt(Event.Delta.Y));
+			const FIntVector2 WarpDelta(FMath::Abs(AccumDelta.X)<FMath::Abs(Delta.X) ? AccumDelta.X : Delta.X, FMath::Abs(AccumDelta.Y)<FMath::Abs(Delta.Y) ? AccumDelta.Y : Delta.Y);
+			Delta.X -= WarpDelta.X;
+			Delta.Y -= WarpDelta.Y;
 
-		// Update to latest position
-		HighPrecisionMousePos.X += Delta.X;
-		HighPrecisionMousePos.Y += Delta.Y;
+			// Update to latest position
+			HighPrecisionMousePos.X += Delta.X;
+			HighPrecisionMousePos.Y += Delta.Y;
 
-		// Clip to lock rect
-		MacCursor->UpdateCursorClipping(HighPrecisionMousePos);
+			// Clip to lock rect
+			MacCursor->UpdateCursorClipping(HighPrecisionMousePos);
 
-		// Clamp to the current screen and avoid the menu bar and dock to prevent popups and other
-		// assorted potential for mouse abuse.
-		if (bUsingHighPrecisionMouseInput)
-		{
-			// Avoid the menu bar & dock disclosure borders at the top & bottom of fullscreen windows
-			if (EventWindow.IsValid() && EventWindow->GetWindowMode() != EWindowMode::Windowed)
+			// Clamp to the current screen and avoid the menu bar and dock to prevent popups and other
+			// assorted potential for mouse abuse.
+			if (bUsingHighPrecisionMouseInput)
 			{
-				VisibleFrame.origin.y += 5;
-				VisibleFrame.size.height -= 10;
+				// Avoid the menu bar & dock disclosure borders at the top & bottom of fullscreen windows
+				if (EventWindow.IsValid() && EventWindow->GetWindowMode() != EWindowMode::Windowed)
+				{
+					VisibleFrame.origin.y += 5;
+					VisibleFrame.size.height -= 10;
+				}
+				int32 ClampedPosX = FMath::Clamp(HighPrecisionMousePos.X, (int32)VisibleFrame.origin.x, (int32)(VisibleFrame.origin.x + VisibleFrame.size.width)-1);
+				int32 ClampedPosY = FMath::Clamp(HighPrecisionMousePos.Y, (int32)VisibleFrame.origin.y, (int32)(VisibleFrame.origin.y + VisibleFrame.size.height)-1);
+				MacCursor->SetPosition(ClampedPosX, ClampedPosY);
+			}
+			else
+			{
+				MacCursor->SetPosition(HighPrecisionMousePos.X, HighPrecisionMousePos.Y);
 			}
-			int32 ClampedPosX = FMath::Clamp(HighPrecisionMousePos.X, (int32)VisibleFrame.origin.x, (int32)(VisibleFrame.origin.x + VisibleFrame.size.width)-1);
-			int32 ClampedPosY = FMath::Clamp(HighPrecisionMousePos.Y, (int32)VisibleFrame.origin.y, (int32)(VisibleFrame.origin.y + VisibleFrame.size.height)-1);
-			MacCursor->SetPosition(ClampedPosX, ClampedPosY);
-		}
-		else
-		{
-			MacCursor->SetPosition(HighPrecisionMousePos.X, HighPrecisionMousePos.Y);
-		}
 
-		// Forward the delta on to Slate
-		MessageHandler->OnRawMouseMove(Delta.X, Delta.Y);
+			// Forward the delta on to Slate
+			MessageHandler->OnRawMouseMove(Delta.X, Delta.Y);
+		}
 	}
 	else
 	{
diff --git a/Engine/Source/Runtime/ApplicationCore/Public/Mac/MacApplication.h b/Engine/Source/Runtime/ApplicationCore/Public/Mac/MacApplication.h
--- a/Engine/Source/Runtime/ApplicationCore/Public/Mac/MacApplication.h
+++ b/Engine/Source/Runtime/ApplicationCore/Public/Mac/MacApplication.h
@@ -10,6 +10,7 @@
 #include "Misc/CoreDelegates.h"
 #include "GenericPlatform/IInputInterface.h"
 
+class FAppleMouseController;
 class IModularFeature;
 
 struct FDeferredMacEvent
@@ -363,6 +364,10 @@ private:
 	
 	TSharedRef<class FMacControllerInterface> HIDInput;
 
+	TSharedPtr<FAppleMouseController> MouseController;
+	
+	void HandleMouseControllerMovement(FIntVector2 Delta);
+
 	/** List of input devices implemented in external modules. */
 	TArray<TSharedPtr<class IInputDevice>> ExternalInputDevices;
 	bool bHasLoadedInputPlugins;
