[Input] Add a new FInputDeviceRegistry:

- FInputDeviceRegistry is  a thread safe registry to keep track of a FInputDeviceId to it's name and type. In the future, this will make FInputDeviceScope obsolete.
- The registry is populated upon device connection and FInputDeviceId assignment, opposed to FInputDeviceScope's requirement of being there for every single input event call on the message handler.
- Add Linux and Mac support for this device type tracking, which were completely missing FInputDeviceScopes.

#Synced-CL 52059832

diff --git a/Engine/Plugins/GameInputWindows/Source/GameInputWindows/Private/GameInputWindowsDevice.cpp b/Engine/Plugins/GameInputWindows/Source/GameInputWindows/Private/GameInputWindowsDevice.cpp
--- a/Engine/Plugins/GameInputWindows/Source/GameInputWindows/Private/GameInputWindowsDevice.cpp
+++ b/Engine/Plugins/GameInputWindows/Source/GameInputWindows/Private/GameInputWindowsDevice.cpp
@@ -9,6 +9,7 @@
 #include "HAL/IConsoleManager.h" // FAutoConsoleVariableRef
 #include "GenericPlatform/IInputInterface.h"
 #include "GameInputDeveloperSettings.h"
+#include "GenericPlatform/InputDeviceRegistry.h"
 
 FGameInputWindowsInputDevice::FGameInputWindowsInputDevice(const TSharedRef<FGenericApplicationMessageHandler>& InMessageHandler, IGameInput* InGameInput)
 	: IGameInputDeviceInterface(InMessageHandler, InGameInput)
@@ -129,6 +130,13 @@ void FGameInputWindowsInputDevice::HandleDeviceConnected(IGameInputDevice* Devic
 			DeviceMapper.GetUserIndexForPlatformUser(UserToAssign),
 			FString(UE::GameInput::GetHardwareDeviceIdentifierName(Device)));
 		
+		{
+			FInputDeviceDescriptor Desc;
+			Desc.HardwareDeviceHandle     = DeviceId;
+			Desc.InputDeviceName          = UE::GameInput::GetDeviceScopeClassName();
+			Desc.HardwareDeviceIdentifier = FName(UE::GameInput::GetHardwareDeviceIdentifierName(Device));
+			FInputDeviceRegistry::RegisterDevice(Desc);
+		}
 		const bool bSuccess = DeviceMapper.Internal_MapInputDeviceToUser(DeviceId, UserToAssign, EInputDeviceConnectionState::Connected);
 		if (ensure(bSuccess))
 		{
diff --git a/Engine/Plugins/Media/PixelStreaming2/Source/PixelStreaming2Input/Private/InputDevice.cpp b/Engine/Plugins/Media/PixelStreaming2/Source/PixelStreaming2Input/Private/InputDevice.cpp
--- a/Engine/Plugins/Media/PixelStreaming2/Source/PixelStreaming2Input/Private/InputDevice.cpp
+++ b/Engine/Plugins/Media/PixelStreaming2/Source/PixelStreaming2Input/Private/InputDevice.cpp
@@ -2,6 +2,7 @@
 
 #include "InputDevice.h"
 #include "Framework/Application/SlateApplication.h"
+#include "GenericPlatform/InputDeviceRegistry.h"
 
 TSharedPtr<UE::PixelStreaming2Input::FInputDevice> UE::PixelStreaming2Input::FInputDevice::InputDevice;
 
@@ -91,6 +92,14 @@ namespace UE::PixelStreaming2Input
 		IPlatformInputDeviceMapper& DeviceMapper = IPlatformInputDeviceMapper::Get();
 		FInputDeviceId				DeviceId = DeviceMapper.AllocateNewInputDeviceId();
 		FPlatformUserId				UserId = DeviceMapper.GetPlatformUserForNewlyConnectedDevice();
+		{
+			static const FName PixelStreamingInputDeviceName(TEXT("PixelStreaming2InputDevice"));
+			FInputDeviceDescriptor Desc;
+			Desc.HardwareDeviceHandle     = DeviceId;
+			Desc.InputDeviceName          = PixelStreamingInputDeviceName;
+			Desc.HardwareDeviceIdentifier = FName(TEXT("PixelStreaming2_Controller"));
+			FInputDeviceRegistry::RegisterDevice(Desc);
+		}
 		DeviceMapper.Internal_MapInputDeviceToUser(DeviceId, UserId, EInputDeviceConnectionState::Connected);
 
 		InputDevices.Add(NextControllerId, { DeviceId, UserId });
diff --git a/Engine/Plugins/Runtime/Windows/XInputDevice/Source/XInputDevice/Private/XInputInterface.cpp b/Engine/Plugins/Runtime/Windows/XInputDevice/Source/XInputDevice/Private/XInputInterface.cpp
--- a/Engine/Plugins/Runtime/Windows/XInputDevice/Source/XInputDevice/Private/XInputInterface.cpp
+++ b/Engine/Plugins/Runtime/Windows/XInputDevice/Source/XInputDevice/Private/XInputInterface.cpp
@@ -7,6 +7,7 @@
 #include "Windows/WindowsApplication.h"
 #include "GenericPlatform/GenericApplication.h"
 #include "GenericPlatform/GenericPlatformInputDeviceMapper.h"
+#include "GenericPlatform/InputDeviceRegistry.h"
 #include "Misc/ConfigCacheIni.h"
 
 #pragma pack (push,8)
@@ -142,6 +143,14 @@ void XInputInterface::GetPlatformUserAndDevice(int32 InControllerId, EInputDevic
 		// If the controller is connected now but was not before, refresh the information
 		if (InDeviceState == EInputDeviceConnectionState::Connected)
 		{
+			// Registering this device with FInputDeviceRegistry so that systems can look up metadata about this device if they
+			// desire, such as icons specific to this controller type.
+			FInputDeviceDescriptor Desc;
+			Desc.InputDeviceName          = XInputInterfaceName;
+			Desc.HardwareDeviceHandle     = OutDeviceId;
+			Desc.HardwareDeviceIdentifier = FName(*XInputControllerIdentifier);
+			FInputDeviceRegistry::RegisterDevice(Desc);
+			
 			DeviceMapper.Internal_MapInputDeviceToUser(OutDeviceId, OutPlatformUserId, InDeviceState);
 		}
 	}
diff --git a/Engine/Source/Runtime/ApplicationCore/Private/Android/AndroidInputInterface.cpp b/Engine/Source/Runtime/ApplicationCore/Private/Android/AndroidInputInterface.cpp
--- a/Engine/Source/Runtime/ApplicationCore/Private/Android/AndroidInputInterface.cpp
+++ b/Engine/Source/Runtime/ApplicationCore/Private/Android/AndroidInputInterface.cpp
@@ -17,6 +17,7 @@
 #include "HAL/IConsoleManager.h"
 #include "IHapticDevice.h"
 #include "GenericPlatform/GenericPlatformInputDeviceMapper.h"
+#include "GenericPlatform/InputDeviceRegistry.h"
 #include "Containers/Ticker.h"
 
 #ifndef ANDROID_GAMEPAD_TRIGGER_THRESHOLD
@@ -321,6 +322,26 @@ void FAndroidInputInterface::MapControllerToPlayer(const FString& ControllerDesc
 		{
 			PlatformUserId = DeviceMapper.GetPlatformUserForNewlyConnectedDevice();
 		}
+
+		// Find the device name from InputDeviceInfoMap by matching descriptor.
+		FName DeviceName = InputClassName_DefaultGamepad;
+		for (const TPair<int32, FAndroidInputDeviceInfo>& Pair : InputDeviceInfoMap)
+		{
+			if (Pair.Value.Descriptor == ControllerDescriptor)
+			{
+				if (!Pair.Value.Name.IsNone())
+				{
+					DeviceName = Pair.Value.Name;
+				}
+				break;
+			}
+		}
+
+		FInputDeviceDescriptor Desc;
+		Desc.InputDeviceName          = DeviceName;
+		Desc.HardwareDeviceHandle     = DeviceId;
+		Desc.HardwareDeviceIdentifier = FName(*ControllerDescriptor);
+		FInputDeviceRegistry::RegisterDevice(Desc);
 	}
 	else if (State == EInputDeviceConnectionState::Disconnected)
 	{
diff --git a/Engine/Source/Runtime/ApplicationCore/Private/Apple/AppleControllerInterface.cpp b/Engine/Source/Runtime/ApplicationCore/Private/Apple/AppleControllerInterface.cpp
--- a/Engine/Source/Runtime/ApplicationCore/Private/Apple/AppleControllerInterface.cpp
+++ b/Engine/Source/Runtime/ApplicationCore/Private/Apple/AppleControllerInterface.cpp
@@ -6,6 +6,7 @@
 #include "HAL/PlatformTime.h"
 #include "Misc/ScopeLock.h"
 #include "GenericPlatform/GenericPlatformInputDeviceMapper.h"
+#include "GenericPlatform/InputDeviceRegistry.h"
 
 DEFINE_LOG_CATEGORY(LogAppleController);
 
@@ -391,6 +392,13 @@ void FAppleControllerInterface::HandleConnection(GCController* Controller)
 		FPlatformUserId UserId = FGenericPlatformMisc::GetPlatformUserForUserIndex(Controllers[ControllerIndex].PlayerIndex);
 		FInputDeviceId DeviceId = INPUTDEVICEID_NONE;
 		DeviceMapper.RemapControllerIdToPlatformUserAndDevice(Controllers[ControllerIndex].PlayerIndex, OUT UserId, OUT DeviceId);
+		{
+			FInputDeviceDescriptor Desc;
+			Desc.HardwareDeviceHandle     = DeviceId;
+			Desc.InputDeviceName          = GetControllerTypeName(Controllers[ControllerIndex].ControllerType);
+			Desc.HardwareDeviceIdentifier = FName(*HardwareDeviceIdentifier_DefaultGamepad);
+			FInputDeviceRegistry::RegisterDevice(Desc);
+		}
 		DeviceMapper.Internal_MapInputDeviceToUser(DeviceId, UserId, EInputDeviceConnectionState::Connected);
 		
 		UE_LOG(LogAppleController, Log, TEXT("New %s controller inserted, assigned to playerIndex %d"),
diff --git a/Engine/Source/Runtime/ApplicationCore/Private/GenericPlatform/InputDeviceRegistry.cpp b/Engine/Source/Runtime/ApplicationCore/Private/GenericPlatform/InputDeviceRegistry.cpp
new file mode 100644
--- /dev/null
+++ b/Engine/Source/Runtime/ApplicationCore/Private/GenericPlatform/InputDeviceRegistry.cpp
@@ -0,0 +1,57 @@
+// Copyright Epic Games, Inc. All Rights Reserved.
+
+#include "GenericPlatform/InputDeviceRegistry.h"
+#include "Logging/LogMacros.h"
+
+DEFINE_LOG_CATEGORY_STATIC(LogInputDeviceRegistry, Log, All);
+
+FInputDeviceRegistry& FInputDeviceRegistry::Get()
+{
+	static FInputDeviceRegistry Instance;
+	return Instance;
+}
+
+bool FInputDeviceRegistry::RegisterDevice(const FInputDeviceDescriptor& Descriptor)
+{
+	if (!Descriptor.HardwareDeviceHandle.IsValid())
+	{
+		UE_LOG(LogInputDeviceRegistry, Warning, TEXT("RegisterDevice failed: HardwareDeviceHandle is not valid. %s"), *Descriptor.ToString());
+		return false;
+	}
+
+	FInputDeviceRegistry& Registry = Get();
+	FWriteScopeLock ScopeLock(Registry.Lock);
+
+	UE_LOG(LogInputDeviceRegistry, Log, TEXT("RegisterDevice: %s"), *Descriptor.ToString());
+	
+	Registry.Descriptors.Add(Descriptor.HardwareDeviceHandle, Descriptor);
+	return true;
+}
+
+void FInputDeviceRegistry::RemoveDevice(const FInputDeviceId DeviceId)
+{
+	UE_LOG(LogInputDeviceRegistry, Log, TEXT("RemoveDevice: DeviceId=%d"), DeviceId.GetId());
+	
+	FInputDeviceRegistry& Registry = Get();
+	FWriteScopeLock ScopeLock(Registry.Lock);
+	Registry.Descriptors.Remove(DeviceId);
+}
+
+FString FInputDeviceDescriptor::ToString() const
+{
+	return FString::Printf(TEXT("DeviceId=%d InputDeviceName=%s HardwareDeviceIdentifier=%s"),
+		HardwareDeviceHandle.GetId(),
+		*InputDeviceName.ToString(),
+		*HardwareDeviceIdentifier.ToString());
+}
+
+TOptional<FInputDeviceDescriptor> FInputDeviceRegistry::FindDescriptor(const FInputDeviceId DeviceId)
+{
+	FInputDeviceRegistry& Registry = Get();
+	FReadScopeLock ScopeLock(Registry.Lock);
+	if (const FInputDeviceDescriptor* Found = Registry.Descriptors.Find(DeviceId))
+	{
+		return *Found;
+	}
+	return TOptional<FInputDeviceDescriptor>();
+}
diff --git a/Engine/Source/Runtime/ApplicationCore/Private/Linux/LinuxApplication.cpp b/Engine/Source/Runtime/ApplicationCore/Private/Linux/LinuxApplication.cpp
--- a/Engine/Source/Runtime/ApplicationCore/Private/Linux/LinuxApplication.cpp
+++ b/Engine/Source/Runtime/ApplicationCore/Private/Linux/LinuxApplication.cpp
@@ -12,6 +12,7 @@
 #include "IInputDeviceModule.h"
 #include "IHapticDevice.h"
 #include "GenericPlatform/GenericPlatformInputDeviceMapper.h"
+#include "GenericPlatform/InputDeviceRegistry.h"
 
 //
 // GameController thresholds
@@ -2144,6 +2145,14 @@ void FLinuxApplication::AddGameController(SDL_JoystickID Id)
 		}
 	}
 
+	{
+		static const FName LinuxGamepadInputDeviceName(TEXT("FLinuxApplication"));
+		FInputDeviceDescriptor Desc;
+		Desc.HardwareDeviceHandle     = ControllerState.DeviceId;
+		Desc.InputDeviceName          = LinuxGamepadInputDeviceName;
+		Desc.HardwareDeviceIdentifier = FName(UTF8_TO_TCHAR(SDL_GetGamepadName(Controller)));
+		FInputDeviceRegistry::RegisterDevice(Desc);
+	}
 	DeviceMapper.Internal_MapInputDeviceToUser(ControllerState.DeviceId, UserId, EInputDeviceConnectionState::Connected);
 }
 
diff --git a/Engine/Source/Runtime/ApplicationCore/Private/Mac/HIDInputInterface.cpp b/Engine/Source/Runtime/ApplicationCore/Private/Mac/HIDInputInterface.cpp
--- a/Engine/Source/Runtime/ApplicationCore/Private/Mac/HIDInputInterface.cpp
+++ b/Engine/Source/Runtime/ApplicationCore/Private/Mac/HIDInputInterface.cpp
@@ -6,6 +6,7 @@
 #include "Misc/ConfigCacheIni.h"
 #include "Templates/SharedPointer.h"
 #include "GenericPlatform/GenericPlatformInputDeviceMapper.h"
+#include "GenericPlatform/InputDeviceRegistry.h"
 
 static int32 GetDevicePropertyAsInt32(IOHIDDeviceRef DeviceRef, CFStringRef Property)
 {
@@ -445,6 +446,14 @@ void HIDInputInterface::OnNewHIDController(IOReturn Result, IOHIDDeviceRef Devic
 			FInputDeviceId DeviceId = INPUTDEVICEID_NONE;
 
 			DeviceMapper.RemapControllerIdToPlatformUserAndDevice(ControllerIndex, OUT PlatformUserId, OUT DeviceId);
+			{
+				static const FName HIDInputDeviceName(TEXT("HIDInputInterface"));
+				FInputDeviceDescriptor Desc;
+				Desc.HardwareDeviceHandle     = DeviceId;
+				Desc.InputDeviceName          = HIDInputDeviceName;
+				Desc.HardwareDeviceIdentifier = FName(TEXT("Mac_HIDController"));
+				FInputDeviceRegistry::RegisterDevice(Desc);
+			}
 			DeviceMapper.Internal_MapInputDeviceToUser(DeviceId, PlatformUserId, EInputDeviceConnectionState::Connected);
 
 			bIsGamepadAttached = true;
diff --git a/Engine/Source/Runtime/ApplicationCore/Private/Windows/WindowsApplication.cpp b/Engine/Source/Runtime/ApplicationCore/Private/Windows/WindowsApplication.cpp
--- a/Engine/Source/Runtime/ApplicationCore/Private/Windows/WindowsApplication.cpp
+++ b/Engine/Source/Runtime/ApplicationCore/Private/Windows/WindowsApplication.cpp
@@ -25,6 +25,7 @@
 #include "Tasks/Pipe.h"
 #include "Containers/LockFreeList.h"
 #include "HAL/RunnableThread.h"
+#include "GenericPlatform/InputDeviceRegistry.h"
 
 #if WITH_EDITOR
 #include "Modules/ModuleManager.h"
@@ -1314,6 +1315,17 @@ FWindowsApplication::FWindowsApplication( const HINSTANCE HInstance, const HICON
 	{
 		UE::WindowsApplication::RawInputThread::FWindowsRawInputRunnable::Initialize(HInstance);
 	}
+
+	// Register the keyboard and mouse device immediately so callers can look up
+	// hardware metadata from FInputDeviceRegistry before the first tick.
+	{
+		const FInputDeviceId KBMDeviceId = IPlatformInputDeviceMapper::Get().GetDefaultInputDevice();
+		FInputDeviceDescriptor Desc;
+		Desc.InputDeviceName          = UE::WindowsInput::InputClassName;
+		Desc.HardwareDeviceHandle     = KBMDeviceId;
+		Desc.HardwareDeviceIdentifier = FName(*UE::WindowsInput::KBMInputHardwareName);
+		FInputDeviceRegistry::RegisterDevice(Desc);
+	}
 }
 
 bool FWindowsApplication::SupportsForceMaxTouchpadSensitivity()
diff --git a/Engine/Source/Runtime/ApplicationCore/Public/GenericPlatform/InputDeviceRegistry.h b/Engine/Source/Runtime/ApplicationCore/Public/GenericPlatform/InputDeviceRegistry.h
new file mode 100644
--- /dev/null
+++ b/Engine/Source/Runtime/ApplicationCore/Public/GenericPlatform/InputDeviceRegistry.h
@@ -0,0 +1,91 @@
+// Copyright Epic Games, Inc. All Rights Reserved.
+
+#pragma once
+
+#include "GenericPlatform/GenericPlatformInputDeviceMapper.h"
+#include "HAL/CriticalSection.h"
+#include "Misc/ScopeRWLock.h"
+#include "Containers/Map.h"
+#include "Misc/Optional.h"
+#include "UObject/NameTypes.h"
+
+#define UE_API APPLICATIONCORE_API
+
+/**
+ * Stable per-device descriptor stored in FInputDeviceRegistry.
+ * Mirrors the data held by FInputDeviceScope but is keyed on FInputDeviceId
+ * and lives for the lifetime of the device rather than just the dispatch call.
+ */
+struct FInputDeviceDescriptor final
+{
+	/** The engine-assigned device ID. INPUTDEVICEID_NONE if unknown. */
+	FInputDeviceId HardwareDeviceHandle = INPUTDEVICEID_NONE;
+	
+	/**
+	 * Logical name of the input device interface or module which created this descriptor.
+	 * 
+	 * This is often times the name of the C++ class which is being used to create this event, like "XInputInterface".
+	 *
+	 * It is a generic grouping for input devices.
+	 */
+	FName InputDeviceName = NAME_None;
+
+	/**
+	 * Logical identifier for the hardware device. Platform-specific, not translated.
+	 * Stored as FName (replaces the FString used in FInputDeviceScope)
+	 *
+	 * This is a more specific name that represents a specific human interface device such as
+	 * "<PlatformName>_Controller". 
+	 */
+	FName HardwareDeviceIdentifier = NAME_None;
+
+	/** Returns a human-readable string representation of this descriptor, suitable for logging. */
+	UE_API FString ToString() const;
+};
+
+/**
+ * Thread-safe singleton registry mapping FInputDeviceId to FInputDeviceDescriptor.
+ *
+ * This is the intended replacement for FInputDeviceScope::GetCurrent() and is safe to call from
+ * multiple threads.
+ */
+class FInputDeviceRegistry final
+{
+private:
+	
+	/** Returns the process-wide singleton instance. */
+	static UE_API FInputDeviceRegistry& Get();
+
+public:
+	/**
+	 * Register or update the descriptor for a device.
+	 * Called by FInputDeviceScope's constructor and may be called directly at device-connect
+	 * time once all call sites have migrated away from FInputDeviceScope.
+	 * @return True if the descriptor was registered successfully, false if the descriptor's
+	 *         HardwareDeviceHandle is not valid (INPUTDEVICEID_NONE).
+	 */
+	static UE_API bool RegisterDevice(const FInputDeviceDescriptor& Descriptor);
+
+	/**
+	 * Explicitly remove the descriptor for a device from the registry.
+	 * The registry is intentionally persistent across connect/disconnect cycles so that
+	 * hardware metadata remains queryable for disconnected devices. Only call this when
+	 * the device ID itself is being permanently retired (e.g., device ID recycled by the
+	 * mapper, or engine shutdown cleanup). No-op if the device was never registered.
+	 * Safe to call from any thread.
+	 */
+	static UE_API void RemoveDevice(const FInputDeviceId DeviceId);
+
+	/**
+	 * Look up the descriptor for a device by its FInputDeviceId.
+	 * Returns an empty optional if the device has not been registered yet.
+	 * Safe to call from any thread.
+	 */
+	static UE_API TOptional<FInputDeviceDescriptor> FindDescriptor(const FInputDeviceId DeviceId);
+
+private:
+	mutable FRWLock Lock;
+	TMap<FInputDeviceId, FInputDeviceDescriptor> Descriptors;
+};
+
+#undef UE_API
\ No newline at end of file
