Fix a CJK text input bug on Windows 11 with the new Microsoft IME (TSF)

Defocusing a text field without defocusing the window itself (e.g. by focusing another widget in the same window) leaves the IME context in a stuck state, and it's not possible to toggle between CJK and english anymore when the text field is focused again.

That's because FWindowsTextInputMethodSystem::ActivateContext and DeactivateContext need to call SetFocus on the TSF thread manager in addition to AssociateFocus.
The Win32 documentation for AssociateFocus isn't clear on when SetFocus is called automatically, but it implies that it only does so when the window itself loses/gains focus.
This is consistent with my tests : the bug doesn't occur if I'm toggling between windows, only if I'm moving focus between a text field and another widget in the same window. In that case AssociateFocus does nothing and TSF is stuck with the wrong document manager.
Adding calls to SetFocus in addition to AssociateFocus (which is still needed) restores the correct TSF document manager in all cases.

#tests tested using Korean and Japanese input, also tested with the old Microsoft IME mode (IMM) to ensure no regressions
#rnx
#Synced-CL 51824576

diff --git a/Engine/Source/Runtime/ApplicationCore/Private/Windows/WindowsTextInputMethodSystem.cpp b/Engine/Source/Runtime/ApplicationCore/Private/Windows/WindowsTextInputMethodSystem.cpp
--- a/Engine/Source/Runtime/ApplicationCore/Private/Windows/WindowsTextInputMethodSystem.cpp
+++ b/Engine/Source/Runtime/ApplicationCore/Private/Windows/WindowsTextInputMethodSystem.cpp
@@ -860,15 +860,21 @@ void FWindowsTextInputMethodSystem::ActivateContext(const TSharedRef<ITextInputM
 
 		// TSF Implementation
 		TComPtr<FTextStoreACP>& TextStore = InternalContext.TSFContext;
+		Result = TSFThreadManager->SetFocus(TextStore->TSFDocumentManager);
+		if (FAILED(Result))
+		{
+			TCHAR ErrorMsg[1024];
+			FPlatformMisc::GetSystemErrorMessage(ErrorMsg, 1024, Result);
+			UE_LOG(LogWindowsTextInputMethodSystem, Error, TEXT("Activating a context failed while setting focus on a TSF document manager. %s (0x%08x)"), ErrorMsg, Result);
+		}
 		ITfDocumentMgr* Unused;
 		Result = TSFThreadManager->AssociateFocus(InternalContext.WindowHandle, TextStore->TSFDocumentManager, &Unused);
 		if (FAILED(Result))
 		{
 			TCHAR ErrorMsg[1024];
 			FPlatformMisc::GetSystemErrorMessage(ErrorMsg, 1024, Result);
-			UE_LOG(LogWindowsTextInputMethodSystem, Error, TEXT("Activating a context failed while setting focus on a TSF document manager. %s (0x%08x)"), ErrorMsg, Result);
+			UE_LOG(LogWindowsTextInputMethodSystem, Error, TEXT("Activating a context failed while associating focus on a TSF document manager. %s (0x%08x)"), ErrorMsg, Result);
 		}
-
 		UE_LOG(LogWindowsTextInputMethodSystem, Verbose, TEXT("Activated context %p!"), &(Context.Get()));
 	}
 }
@@ -889,6 +895,15 @@ void FWindowsTextInputMethodSystem::DeactivateContext(const TSharedRef<ITextInpu
 		::ImmNotifyIME(IMMContext, NI_COMPOSITIONSTR, CPS_COMPLETE, 0);
 		::ImmReleaseContext(InternalContext.WindowHandle, IMMContext);
 
+		// TSF Implementation
+		HRESULT Result = TSFThreadManager->SetFocus(TSFDisabledDocumentManager);
+		if (FAILED(Result))
+		{
+			TCHAR ErrorMsg[1024];
+			FPlatformMisc::GetSystemErrorMessage(ErrorMsg, 1024, Result);
+			UE_LOG(LogWindowsTextInputMethodSystem, Error, TEXT("Deactivating a context failed while setting focus on the disabled TSF document manager. %s (0x%08x)"), ErrorMsg, Result);
+		}
+
 		// Ensure all known windows are associated with a disabled IME
 		ClearStaleWindowHandles();
 		for (const TWeakPtr<FGenericWindow>& Window : KnownWindows)
@@ -899,17 +914,22 @@ void FWindowsTextInputMethodSystem::DeactivateContext(const TSharedRef<ITextInpu
 				const HWND Hwnd = reinterpret_cast<HWND>(WindowPtr->GetOSWindowHandle());
 				if (Hwnd)
 				{
-					//  TSF Implementation
+					// TSF Implementation
 					ITfDocumentMgr* Unused;
-					TSFThreadManager->AssociateFocus(Hwnd, TSFDisabledDocumentManager, &Unused);
+					Result = TSFThreadManager->AssociateFocus(Hwnd, TSFDisabledDocumentManager, &Unused);
+					if (FAILED(Result))
+					{
+						TCHAR ErrorMsg[1024];
+						FPlatformMisc::GetSystemErrorMessage(ErrorMsg, 1024, Result);
+						UE_LOG(LogWindowsTextInputMethodSystem, Error, TEXT("Deactivating a context failed while associating focus on the disabled TSF document manager. %s (0x%08x)"), ErrorMsg, Result);
+					}
 
 					// IMM Implementation
-					::ImmAssociateContext(InternalContext.WindowHandle, nullptr);
+					::ImmAssociateContext(Hwnd, nullptr);
 				}
 			}
-
-			UE_LOG(LogWindowsTextInputMethodSystem, Verbose, TEXT("Deactivated context %p!"), &(Context.Get()));
 		}
+		UE_LOG(LogWindowsTextInputMethodSystem, Verbose, TEXT("Deactivated context %p!"), &(Context.Get()));
 	}
 
 	// General Implementation
