Windows: Optimize maximized borderless windows

The primary aim of this CL is to reconcile differences between EWindowType::GameWindow and EWindowType::Normal that have no basis in whether a window is "game" or not. Before this change, maximized normal windows would allocate a swapchain that is oversized by 16 pixels in both dimensions. Now, normal windows enjoy the same optimization as game windows. The related logic is slightly simplified, and confusing explanations are replaced with more direct + correct ones.

[RN] Reduced VRAM usage of maximized editor windows
#Synced-CL 52110000

diff --git a/Engine/Source/Runtime/ApplicationCore/Private/Windows/WindowsApplication.cpp b/Engine/Source/Runtime/ApplicationCore/Private/Windows/WindowsApplication.cpp
--- a/Engine/Source/Runtime/ApplicationCore/Private/Windows/WindowsApplication.cpp
+++ b/Engine/Source/Runtime/ApplicationCore/Private/Windows/WindowsApplication.cpp
@@ -2540,12 +2540,12 @@ int32 FWindowsApplication::ProcessMessage( HWND hwnd, uint32 msg, WPARAM wParam,
 				// Let windows absorb this message if using the standard border
 				if ( wParam && !CurrentNativeEventWindow->GetDefinition().HasOSWindowBorder )
 				{
-					// Borderless game windows are not actually borderless, they have a thick border that we simply draw game content over (client
-					// rect contains the window border). When maximized Windows will bleed our border over the edges of the monitor. So that we
-					// don't draw content we are going to later discard, we change a maximized window's size and position so that the entire
-					// window rect (including the border) sits inside the monitor. The size adjustments here will be sent to WM_MOVE and
-					// WM_SIZE and the window will still be considered maximized.
-					if (CurrentNativeEventWindow->GetDefinition().Type == EWindowType::GameWindow && CurrentNativeEventWindow->IsMaximized())
+					// Borderless windows are not actually borderless: they have a thick border that we simply draw content over by making the client
+					// rect contain that border. When maximized, however, Windows forces the border to bleed over the edges of the monitor. To avoid
+					// drawing content outside of the viewable area, we counter this behavior by adjusting the size back inwards by the bleed amount,
+					// making the window and client rects again align exactly with the viewable area. The size adjustments here will be sent to
+					// WM_MOVE and WM_SIZE, but the window will still be considered maximized.
+					if (CurrentNativeEventWindow->IsMaximized())
 					{
 						// Ask the system for the window border size as this is the amount that Windows will bleed our window over the edge
 						// of our desired space. The value returned by CurrentNativeEventWindow will be incorrect for our usage here as it
@@ -2557,8 +2557,9 @@ int32 FWindowsApplication::ProcessMessage( HWND hwnd, uint32 msg, WPARAM wParam,
 
 						// A pointer to the window size data that Windows will use is passed to us in lParam
 						LPNCCALCSIZE_PARAMS ResizingRects = (LPNCCALCSIZE_PARAMS)lParam;
-						// The first rectangle contains the client rectangle of the resized window. Decrease window size on all sides by
-						// the border size.
+
+						// The first rectangle initially contains the new window size; we are expected to overwrite this with the desired client size.
+						// Make the client inset on all sides by the bleed amount.
 						ResizingRects->rgrc[0].left += WindowInfo.cxWindowBorders;
 						ResizingRects->rgrc[0].top += WindowInfo.cxWindowBorders;
 						ResizingRects->rgrc[0].right -= WindowInfo.cxWindowBorders;
@@ -2573,17 +2574,10 @@ int32 FWindowsApplication::ProcessMessage( HWND hwnd, uint32 msg, WPARAM wParam,
 						// A third rectangle is passed in that contains the source rectangle (client area from window pre-maximize).
 						// It's value should not be changed.
 
-						// The new window position. Pull in the window on all sides by the width of the window border so that the
-						// window fits entirely on screen. We'll draw over these borders with game content.
-						ResizingRects->lppos->x += WindowInfo.cxWindowBorders;
-						ResizingRects->lppos->y += WindowInfo.cxWindowBorders;
-						ResizingRects->lppos->cx -= 2 * WindowInfo.cxWindowBorders;
-						ResizingRects->lppos->cy -= 2 * WindowInfo.cxWindowBorders;
-
 						// Informs Windows to use the values as we altered them.
 						return WVR_VALIDRECTS;
 					}
-					else if (wParam)
+					else
 					{
 						//////////////////////////////////////////////////////
 						// Find out the schedule of when composition takes place 
diff --git a/Engine/Source/Runtime/ApplicationCore/Private/Windows/WindowsWindow.cpp b/Engine/Source/Runtime/ApplicationCore/Private/Windows/WindowsWindow.cpp
--- a/Engine/Source/Runtime/ApplicationCore/Private/Windows/WindowsWindow.cpp
+++ b/Engine/Source/Runtime/ApplicationCore/Private/Windows/WindowsWindow.cpp
@@ -370,36 +370,26 @@ void FWindowsWindow::OnTransparencySupportChanged(EWindowTransparency NewTranspa
 #endif
 }
 
-HRGN FWindowsWindow::MakeWindowRegionObject(bool bIncludeBorderWhenMaximized) const
+HRGN FWindowsWindow::MakeWindowRegionObject(bool bForQuery) const
 {
-	HRGN Region;
+	int X1, Y1, X2, Y2;
 	if ( RegionWidth != INDEX_NONE && RegionHeight != INDEX_NONE )
 	{
-		const bool bIsBorderlessGameWindow = GetDefinition().Type == EWindowType::GameWindow && !GetDefinition().HasOSWindowBorder;
+		const bool bIsBorderless = !GetDefinition().HasOSWindowBorder;
 		if (IsMaximized())
 		{
-			if (bIsBorderlessGameWindow)
-			{
-				// Windows caches the cxWindowBorders size at window creation. Even if borders are removed or resized Windows will continue to use this value when evaluating regions
-				// and sizing windows. When maximized this means that our window position will be offset from the screen origin by (-cxWindowBorders,-cxWindowBorders). We want to
-				// display only the region within the maximized screen area, so offset our upper left and lower right by cxWindowBorders.
-				WINDOWINFO WindowInfo;
-				FMemory::Memzero(WindowInfo);
-				WindowInfo.cbSize = sizeof(WindowInfo);
-				::GetWindowInfo(HWnd, &WindowInfo);
-
-				const int32 WindowBorderSize = bIncludeBorderWhenMaximized ? WindowInfo.cxWindowBorders : 0;
-				Region = CreateRectRgn(WindowBorderSize, WindowBorderSize, RegionWidth + WindowBorderSize, RegionHeight + WindowBorderSize);
-			}
-			else
-			{
-				const int32 WindowBorderSize = bIncludeBorderWhenMaximized ? GetWindowBorderSize() : 0;
-				Region = CreateRectRgn(WindowBorderSize, WindowBorderSize, RegionWidth - WindowBorderSize, RegionHeight - WindowBorderSize);
-			}
+			// Maximized windows are forcibly oversized by cxWindowBorders, so the rect we pass to SetWindowRgn will be interpreted as relative to
+			// -cxWindowBorders pixels from the screen's corner. We don't want anything being drawn on adjacent monitors, so we need to offset the rgn
+			// by this much (unless it is being used for a query against a cursor position that we have already remapped to client-relative).
+			int32 BorderSize = bForQuery ? 0 : GetWin32BorderSize();
+			X1 = BorderSize;
+			Y1 = BorderSize;
+			X2 = BorderSize + RegionWidth;
+			Y2 = BorderSize + RegionHeight;
 		}
 		else
 		{
-			const bool bUseCornerRadius  = WindowMode == EWindowMode::Windowed && !bIsBorderlessGameWindow &&
+			const bool bUseCornerRadius = WindowMode == EWindowMode::Windowed && !bIsBorderless &&
 #if ALPHA_BLENDED_WINDOWS
 				// Corner radii cause DWM window composition blending to fail, so we always set regions to full size rectangles
 				Definition->TransparencySupport != EWindowTransparency::PerPixel &&
@@ -410,14 +400,23 @@ HRGN FWindowsWindow::MakeWindowRegionObject(bool bIncludeBorderWhenMaximized) co
 			{
 				// CreateRoundRectRgn gives you a duff region that's 1 pixel smaller than you ask for. CreateRectRgn behaves correctly.
 				// This can be verified by uncommenting the assert below
-				Region = CreateRoundRectRgn( 0, 0, RegionWidth+1, RegionHeight+1, Definition->CornerRadius, Definition->CornerRadius );
+				HRGN Region = CreateRoundRectRgn( 0, 0, RegionWidth+1, RegionHeight+1, Definition->CornerRadius, Definition->CornerRadius );
+				if (!Region)
+				{
+					UE_LOGF(LogWindows, Error, "CreateRoundRectRgn for 0,0,%d+1,%d+1,%d,%d failed with error %d",
+						RegionWidth, RegionHeight, Definition->CornerRadius, Definition->CornerRadius, GetLastError());
+				}
 
 				// Test that a point that should be in the region, is in the region
 				// check(!!PtInRegion(Region, RegionWidth-1, RegionHeight/2));
+				return Region;
 			}
 			else
 			{
-				Region = CreateRectRgn( 0, 0, RegionWidth, RegionHeight );
+				X1 = 0;
+				Y1 = 0;
+				X2 = RegionWidth;
+				Y2 = RegionHeight;
 			}
 		}
 	}
@@ -425,9 +424,17 @@ HRGN FWindowsWindow::MakeWindowRegionObject(bool bIncludeBorderWhenMaximized) co
 	{
 		RECT rcWnd;
 		GetWindowRect( HWnd, &rcWnd );
-		Region = CreateRectRgn( 0, 0, rcWnd.right - rcWnd.left, rcWnd.bottom - rcWnd.top );
+		X1 = 0;
+		Y1 = 0;
+		X2 = rcWnd.right - rcWnd.left;
+		Y2 = rcWnd.bottom - rcWnd.top;
 	}
 
+	HRGN Region = CreateRectRgn(X1, Y1, X2, Y2);
+	if (!Region)
+	{
+		UE_LOGF(LogWindows, Error, "CreateRectRgn failed for %d,%d,%d,%d with error %d", X1, Y1, X2, Y2, GetLastError());
+	}
 	return Region;
 }
 
@@ -474,7 +481,12 @@ void FWindowsWindow::AdjustWindowRegion( int32 Width, int32 Height )
 	RegionWidth = Width;
 	RegionHeight = Height;
 
-	HRGN Region = MakeWindowRegionObject(true);
+	const bool bForQuery = false;
+	HRGN Region = MakeWindowRegionObject(bForQuery);
+	if (!Region)
+	{
+		return;
+	}
 
 	// NOTE: We explicitly don't delete the Region object, because the OS deletes the handle when it no longer needed after
 	// a call to SetWindowRgn.
@@ -1081,7 +1093,12 @@ bool FWindowsWindow::IsPointInWindow( int32 X, int32 Y ) const
 {
 	bool Result = false;
 
-	HRGN Region = MakeWindowRegionObject(false);
+	const bool bForQuery = true;
+	HRGN Region = MakeWindowRegionObject(bForQuery);
+	if (!Region)
+	{
+		return false;
+	}
 	Result = !!PtInRegion( Region, X, Y );
 	DeleteObject( Region );
 
@@ -1090,14 +1107,17 @@ bool FWindowsWindow::IsPointInWindow( int32 X, int32 Y ) const
 
 int32 FWindowsWindow::GetWindowBorderSize() const
 {
-	if (GetDefinition().Type == EWindowType::GameWindow && !GetDefinition().HasOSWindowBorder)
+	if (!GetDefinition().HasOSWindowBorder)
 	{
-		// Our borderless game windows actually have a thick border to allow sizing, which we draw over to simulate
-		// a borderless window. We return zero here so that the game will correctly behave as if this is truly a
-		// borderless window.
+		// Our borderless windows technically have a thick border, but we draw over them to simulate a borderless window.
 		return 0;
 	}
 
+	return GetWin32BorderSize();
+}
+
+int32 FWindowsWindow::GetWin32BorderSize() const
+{
 	WINDOWINFO WindowInfo;
 	FMemory::Memzero( WindowInfo );
 	WindowInfo.cbSize = sizeof( WindowInfo );
diff --git a/Engine/Source/Runtime/ApplicationCore/Public/Windows/WindowsWindow.h b/Engine/Source/Runtime/ApplicationCore/Public/Windows/WindowsWindow.h
--- a/Engine/Source/Runtime/ApplicationCore/Public/Windows/WindowsWindow.h
+++ b/Engine/Source/Runtime/ApplicationCore/Public/Windows/WindowsWindow.h
@@ -159,12 +159,19 @@ private:
 
 	APPLICATIONCORE_API void UpdateVisibility();
 
-	/** Creates an HRGN for the window's current region.  Remember to delete this when you're done with it using
-	   ::DeleteObject, unless you're passing it to SetWindowRgn(), which will absorb the reference itself. */
-	APPLICATIONCORE_API HRGN MakeWindowRegionObject(bool bIncludeBorderWhenMaximized) const;
+	/**
+	 * Creates an HRGN for the window's current region.  Remember to delete the returned handle when you're done with it using ::DeleteObject, unless
+     * you're passing it to SetWindowRgn(), which will absorb the reference itself.
+	 *
+	 * @param bForQuery Bypasses situational behavior compensation, so that the same compensation does not also need to be applied to the queried
+	 *     location.
+	 */
+	APPLICATIONCORE_API HRGN MakeWindowRegionObject(bool bForQuery) const;
 
 	APPLICATIONCORE_API void DisableTouchFeedback();
 
+	int32 GetWin32BorderSize() const;
+
 private:
 
 	/** The application that owns this window. */
