Disable [rotation] debug logging by default via IOS_ROTATION_DEBUG_LOGGING define
#Synced-CL 52106027

diff --git a/Engine/Source/Runtime/Apple/MetalRHI/Private/MetalViewport.cpp b/Engine/Source/Runtime/Apple/MetalRHI/Private/MetalViewport.cpp
--- a/Engine/Source/Runtime/Apple/MetalRHI/Private/MetalViewport.cpp
+++ b/Engine/Source/Runtime/Apple/MetalRHI/Private/MetalViewport.cpp
@@ -608,8 +608,10 @@ void FMetalViewport::Resize(uint32 InSizeX, uint32 InSizeY, bool bInIsFullscreen
 #if PLATFORM_IOS
 		Format = InFormat;
 		OrientationMask = CachedOrientationMask;
+#if IOS_ROTATION_DEBUG_LOGGING
 		FPlatformMisc::LowLevelOutputDebugStringf(TEXT("[rotation] FMetalViewport::Resize: OrientationMask %d"), OrientationMask);
 #endif
+#endif
 #endif
 	}
 #endif
diff --git a/Engine/Source/Runtime/ApplicationCore/Private/IOS/IOSAppDelegate.cpp b/Engine/Source/Runtime/ApplicationCore/Private/IOS/IOSAppDelegate.cpp
--- a/Engine/Source/Runtime/ApplicationCore/Private/IOS/IOSAppDelegate.cpp
+++ b/Engine/Source/Runtime/ApplicationCore/Private/IOS/IOSAppDelegate.cpp
@@ -1510,12 +1510,16 @@ static FAutoConsoleVariableRef CVarGEnableThermalsReport(
 
 	if(!self.bPlatformInit)
 	{
+#if IOS_ROTATION_DEBUG_LOGGING
 		FPlatformMisc::LowLevelOutputDebugStringf(TEXT("[rotation] didRotate ignoring call before bPlatformInit self.InterfaceOrientation=%d windowScene=%d device=%d\n"), self.InterfaceOrientation, (int)[self.Window.windowScene interfaceOrientation], (int)[UIDevice currentDevice].orientation);
+#endif
 		return;
 	}
 
+#if IOS_ROTATION_DEBUG_LOGGING
 	FPlatformMisc::LowLevelOutputDebugStringf(TEXT("[rotation] didRotate self.InterfaceOrientation=%d windowScene=%d device=%d\n"), self.InterfaceOrientation, (int)[self.Window.windowScene interfaceOrientation], (int)[UIDevice currentDevice].orientation);
-	
+#endif
+
 	UIInterfaceOrientationMask PotentialMask;
 	UIInterfaceOrientation NewOrientation = self.InterfaceOrientation;
 	bool bForceUpdate = false;
@@ -1531,7 +1535,9 @@ static FAutoConsoleVariableRef CVarGEnableThermalsReport(
 		PotentialMask = [IOSAppDelegate GetMaskFromUIInterfaceOrientation: PotentialOrientation];
 		if (PotentialOrientation != UIInterfaceOrientationUnknown && (PotentialMask & AllowedMask) != 0)
 		{
+#if IOS_ROTATION_DEBUG_LOGGING
 			FPlatformMisc::LowLevelOutputDebugStringf(TEXT("[rotation] didRotate: using device orientation\n"));
+#endif
 			NewOrientation = PotentialOrientation;
 		}
 		else
@@ -1542,7 +1548,9 @@ static FAutoConsoleVariableRef CVarGEnableThermalsReport(
 			if ((PotentialMask & AllowedMask) != 0)
 			{
 				NewOrientation = PotentialOrientation;
+#if IOS_ROTATION_DEBUG_LOGGING
 				FPlatformMisc::LowLevelOutputDebugStringf(TEXT("[rotation] didRotate: using windowScene interfaceOrientation\n"));
+#endif
 			}
 			else
 			{			
@@ -1584,7 +1592,9 @@ static FAutoConsoleVariableRef CVarGEnableThermalsReport(
 					NewOrientation = UIInterfaceOrientationPortraitUpsideDown;
 					PotentialMask = UIInterfaceOrientationMaskPortraitUpsideDown;
 				}
+#if IOS_ROTATION_DEBUG_LOGGING
 				FPlatformMisc::LowLevelOutputDebugStringf(TEXT("[rotation] didRotate: using first allowed orientation %d mask %d (AllowedMask was %d)\n"), NewOrientation, PotentialMask, AllowedMask);
+#endif
 			}
 		}
 	}
@@ -1597,14 +1607,18 @@ static FAutoConsoleVariableRef CVarGEnableThermalsReport(
 	{
 		static int32 Count=0;
 		int LocalCount = ++Count;
+#if IOS_ROTATION_DEBUG_LOGGING
 		FPlatformMisc::LowLevelOutputDebugStringf(TEXT("[rotation] didRotate: rotating to new orientation %d, count %d\n"), (int32)NewOrientation, LocalCount);
+#endif
 
 		self.InterfaceOrientation = NewOrientation;
 		FIOSApplication* Application = self.IOSApplication;
 		
 		FFunctionGraphTask::CreateAndDispatchWhenReady([NewOrientation, Application, LocalCount]()
 		{
+#if IOS_ROTATION_DEBUG_LOGGING
 			FPlatformMisc::LowLevelOutputDebugStringf(TEXT("[rotation] didRotate: calling OrientationChanged(%d) from GT, count %d"), (int32)NewOrientation, LocalCount);
+#endif
 			Application->OrientationChanged(NewOrientation);
 		}, TStatId(), NULL, ENamedThreads::GameThread);
 	}
diff --git a/Engine/Source/Runtime/ApplicationCore/Private/IOS/IOSApplication.cpp b/Engine/Source/Runtime/ApplicationCore/Private/IOS/IOSApplication.cpp
--- a/Engine/Source/Runtime/ApplicationCore/Private/IOS/IOSApplication.cpp
+++ b/Engine/Source/Runtime/ApplicationCore/Private/IOS/IOSApplication.cpp
@@ -373,7 +373,9 @@ void FIOSApplication::CacheDisplayMetrics(UIInterfaceOrientation Orientation)
 		CachedInsets = OrientatedInsets;
 	}
 
+#if IOS_ROTATION_DEBUG_LOGGING
 	FPlatformMisc::LowLevelOutputDebugStringf(TEXT("[rotation] CachedInsets top:%f left:%f bottom:%f right:%f\n"), CachedInsets.top, CachedInsets.left, CachedInsets.bottom, CachedInsets.right);
+#endif
 }
 #endif
 
@@ -386,7 +388,9 @@ TSharedRef< FGenericWindow > FIOSApplication::MakeWindow()
 
 void FIOSApplication::OrientationChanged(UIInterfaceOrientation NewOrientation)
 {
+#if IOS_ROTATION_DEBUG_LOGGING
 	FPlatformMisc::LowLevelOutputDebugStringf(TEXT("[rotation] FIOSApplication::OrientationChanged: %d\n"), NewOrientation);
+#endif
 
 	// queue up the size as we see it now all the way to the RHI 
 	[FAppleAsyncTask CreateTaskWithBlock : ^ bool(void)
@@ -433,7 +437,9 @@ void FIOSApplication::UpdateSafeZoneAfterRotation()
 		FDisplayMetrics::RebuildDisplayMetrics(DisplayMetrics);
 		App->BroadcastDisplayMetricsChanged(DisplayMetrics);
 		FCoreDelegates::OnSafeFrameChangedEvent.Broadcast();
+#if IOS_ROTATION_DEBUG_LOGGING
 		FPlatformMisc::LowLevelOutputDebugStringf(TEXT("[rotation] UpdateSafeZoneAfterRotation top:%f left:%f bottom:%f right:%f\n"), CachedInsets.top, CachedInsets.left, CachedInsets.bottom, CachedInsets.right);
+#endif
 		return true;
 	}];
 }
diff --git a/Engine/Source/Runtime/ApplicationCore/Private/IOS/IOSView.cpp b/Engine/Source/Runtime/ApplicationCore/Private/IOS/IOSView.cpp
--- a/Engine/Source/Runtime/ApplicationCore/Private/IOS/IOSView.cpp
+++ b/Engine/Source/Runtime/ApplicationCore/Private/IOS/IOSView.cpp
@@ -1188,13 +1188,16 @@ self.accessibilityElements = @[Window.accessibilityContainer];
 	const FIOSView* View = [AppDelegate IOSView];
 	if (GIOSDelayRotationUntilPresent && AppDelegate.bPlatformInit && View != nil && (int)View->LastFrameInterfaceOrientationMask != 0)
 	{
+#if IOS_ROTATION_DEBUG_LOGGING
 		FPlatformMisc::LowLevelOutputDebugStringf(TEXT("[rotation] supportedInterfaceOrientations: allowing %d\n"), (int)View->LastFrameInterfaceOrientationMask);
+#endif
 
 		// Only allow rotation to the orientation of the last frame we have rendered, until a new one is available in the new orientation.
 		return View->LastFrameInterfaceOrientationMask;
 	}
 	else
 	{
+#if IOS_ROTATION_DEBUG_LOGGING
 		if (GIOSDelayRotationUntilPresent)
 		{
 			if (!AppDelegate.bPlatformInit)
@@ -1211,6 +1214,7 @@ self.accessibilityElements = @[Window.accessibilityContainer];
 				FPlatformMisc::LowLevelOutputDebugStringf(TEXT("[rotation] supportedInterfaceOrientations: last %d\n"), (int)View->LastFrameInterfaceOrientationMask);
 			}
 		}
+#endif
 
 		return [self supportedInterfaceOrientations_Internal];
 	}
@@ -1284,7 +1288,9 @@ self.accessibilityElements = @[Window.accessibilityContainer];
 	const IOSAppDelegate *AppDelegate = [IOSAppDelegate GetDelegate];
 	const FIOSView *View = [AppDelegate IOSView];
 	
+#if IOS_ROTATION_DEBUG_LOGGING
 	FPlatformMisc::LowLevelOutputDebugStringf(TEXT("[rotation] viewWillTransitionToSize\n"));
+#endif
 
 	[super viewWillTransitionToSize : size withTransitionCoordinator : coordinator];
 
@@ -1306,8 +1312,10 @@ self.accessibilityElements = @[Window.accessibilityContainer];
 {
 	const FIOSView* View = [[IOSAppDelegate GetDelegate] IOSView];
 
+#if IOS_ROTATION_DEBUG_LOGGING
 	FPlatformMisc::LowLevelOutputDebugStringf(TEXT("[rotation] notifyPresentAfterRotate OrientationMask:%d SizeX:%d SizeY:%d\n"), NewOrientationMask, SizeX, SizeY);
-	
+#endif
+
 	View->LastFrameInterfaceOrientationMask = NewOrientationMask;
 
 	// Attempt to rotate to desired orientation
diff --git a/Engine/Source/Runtime/ApplicationCore/Public/IOS/IOSApplication.h b/Engine/Source/Runtime/ApplicationCore/Public/IOS/IOSApplication.h
--- a/Engine/Source/Runtime/ApplicationCore/Public/IOS/IOSApplication.h
+++ b/Engine/Source/Runtime/ApplicationCore/Public/IOS/IOSApplication.h
@@ -11,6 +11,11 @@
 #include "IOS/IOSSystemIncludes.h"
 #endif
 
+// Set to 1 to enable [rotation] debug logging across IOS files
+#ifndef IOS_ROTATION_DEBUG_LOGGING
+#define IOS_ROTATION_DEBUG_LOGGING 0
+#endif
+
 class FIOSInputInterface;
 class FIOSWindow;
 class IModularFeature;
diff --git a/Engine/Source/Runtime/Core/Private/IOS/IOSPlatformMisc.cpp b/Engine/Source/Runtime/Core/Private/IOS/IOSPlatformMisc.cpp
--- a/Engine/Source/Runtime/Core/Private/IOS/IOSPlatformMisc.cpp
+++ b/Engine/Source/Runtime/Core/Private/IOS/IOSPlatformMisc.cpp
@@ -481,7 +481,9 @@ void FIOSPlatformMisc::SetAllowedDeviceOrientation(EDeviceScreenOrientation NewA
 	IOSAppDelegate* AppDelegate = [IOSAppDelegate GetDelegate];
 	AppDelegate.IOSView->SupportedInterfaceOrientations = GetUIInterfaceOrientationMask(NewAllowedDeviceOrientation);
 
+#if IOS_ROTATION_DEBUG_LOGGING
 	FPlatformMisc::LowLevelOutputDebugStringf(TEXT("[rotation] SetAllowedDeviceOrientation mask %d\n"), (int)GetUIInterfaceOrientationMask(NewAllowedDeviceOrientation));
+#endif
 
 	// Check rotation state
 	[AppDelegate didRotate:nil];
