[Input] Add an non-shipping API to allow us to set fake input device descriptions in the editor and via some console commands. This will be very useful to help make some better tooling for "simulating" specific input devices, as well as testing things like Common UI/Input and how they handle the transition of swapping to different device types.

#rnx
#Synced-CL 52059832

diff --git a/Engine/Source/Runtime/ApplicationCore/Private/GenericPlatform/InputDeviceRegistry.cpp b/Engine/Source/Runtime/ApplicationCore/Private/GenericPlatform/InputDeviceRegistry.cpp
--- a/Engine/Source/Runtime/ApplicationCore/Private/GenericPlatform/InputDeviceRegistry.cpp
+++ b/Engine/Source/Runtime/ApplicationCore/Private/GenericPlatform/InputDeviceRegistry.cpp
@@ -2,6 +2,7 @@
 
 #include "GenericPlatform/InputDeviceRegistry.h"
 #include "Logging/LogMacros.h"
+#include "HAL/IConsoleManager.h"
 
 DEFINE_LOG_CATEGORY_STATIC(LogInputDeviceRegistry, Log, All);
 
@@ -49,9 +50,135 @@ TOptional<FInputDeviceDescriptor> FInputDeviceRegistry::FindDescriptor(const FIn
 {
 	FInputDeviceRegistry& Registry = Get();
 	FReadScopeLock ScopeLock(Registry.Lock);
+
+#if !UE_BUILD_SHIPPING
+	// Simulated descriptors take priority so that non-shipping tooling can override what
+	// FindDescriptor returns without touching the real runtime registry.
+	if (const FInputDeviceDescriptor* Simulated = Registry.SimulatedDescriptors.Find(DeviceId))
+	{
+		return *Simulated;
+	}
+#endif // !UE_BUILD_SHIPPING
+
 	if (const FInputDeviceDescriptor* Found = Registry.Descriptors.Find(DeviceId))
 	{
 		return *Found;
 	}
 	return TOptional<FInputDeviceDescriptor>();
 }
+
+#if !UE_BUILD_SHIPPING
+
+void FInputDeviceRegistry::SetSimulatedDescriptor(const FInputDeviceId DeviceId, const FInputDeviceDescriptor& SimulatedDescriptor)
+{
+	UE_LOG(LogInputDeviceRegistry, Log, TEXT("SetSimulatedDescriptor: DeviceId=%d %s"), DeviceId.GetId(), *SimulatedDescriptor.ToString());
+
+	FInputDeviceRegistry& Registry = Get();
+	FWriteScopeLock ScopeLock(Registry.Lock);
+	Registry.SimulatedDescriptors.Add(DeviceId, SimulatedDescriptor);
+}
+
+void FInputDeviceRegistry::ClearSimulatedDescriptor(const FInputDeviceId DeviceId)
+{
+	UE_LOG(LogInputDeviceRegistry, Log, TEXT("ClearSimulatedDescriptor: DeviceId=%d"), DeviceId.GetId());
+
+	FInputDeviceRegistry& Registry = Get();
+	FWriteScopeLock ScopeLock(Registry.Lock);
+	Registry.SimulatedDescriptors.Remove(DeviceId);
+}
+
+void FInputDeviceRegistry::ClearAllSimulatedDescriptors()
+{
+	UE_LOG(LogInputDeviceRegistry, Log, TEXT("ClearAllSimulatedDescriptors"));
+
+	FInputDeviceRegistry& Registry = Get();
+	FWriteScopeLock ScopeLock(Registry.Lock);
+	Registry.SimulatedDescriptors.Reset();
+}
+
+//-----------------------------------------------------------------------------
+// Non-shipping console commands
+//-----------------------------------------------------------------------------
+
+// Input.Devices.SetDeviceDesc <DeviceId> <InputDeviceName>::<HardwareDeviceIdentifier>
+//   Example: Input.Devices.SetDeviceDesc 1 XInputInterface::XInputController
+static FAutoConsoleCommand CmdSetSimulatedDescriptor(
+	TEXT("Input.Devices.SetDeviceDesc"),
+	TEXT("(Non-shipping) Simulate a device descriptor for a given device ID.\n")
+	TEXT("Usage: Input.Devices.SetDeviceDesc <DeviceId> <InputDeviceName>::<HardwareDeviceIdentifier>\n")
+	TEXT("Example: Input.Devices.SetDeviceDesc 1 XInputInterface::XInputController"),
+	FConsoleCommandWithArgsDelegate::CreateLambda([](const TArray<FString>& Args)
+	{
+		if (Args.Num() < 2)
+		{
+			UE_LOG(LogInputDeviceRegistry, Warning,
+				TEXT("Input.Devices.SetDeviceDesc: expected 2 arguments: <DeviceId> <InputDeviceName>::<HardwareDeviceIdentifier>"));
+			return;
+		}
+
+		const int32 RawDeviceId = FCString::Atoi(*Args[0]);
+		const FInputDeviceId DeviceId = FInputDeviceId::CreateFromInternalId(RawDeviceId);
+		if (!DeviceId.IsValid())
+		{
+			UE_LOG(LogInputDeviceRegistry, Warning,
+				TEXT("Input.Devices.SetDeviceDesc: DeviceId '%s' is not valid."), *Args[0]);
+			return;
+		}
+
+		// Parse "InputDeviceName::HardwareDeviceIdentifier"
+		FString InputDeviceName;
+		FString HardwareDeviceIdentifier;
+		if (!Args[1].Split(TEXT("::"), &InputDeviceName, &HardwareDeviceIdentifier))
+		{
+			UE_LOG(LogInputDeviceRegistry, Warning,
+				TEXT("Input.Devices.SetDeviceDesc: argument '%s' is not in the expected format <InputDeviceName>::<HardwareDeviceIdentifier>"),
+				*Args[1]);
+			return;
+		}
+
+		FInputDeviceDescriptor Desc;
+		Desc.HardwareDeviceHandle     = DeviceId;
+		Desc.InputDeviceName          = FName(*InputDeviceName);
+		Desc.HardwareDeviceIdentifier = FName(*HardwareDeviceIdentifier);
+		FInputDeviceRegistry::SetSimulatedDescriptor(DeviceId, Desc);
+	})
+);
+
+// Input.Devices.ClearSimulatedDescriptor <DeviceId>
+static FAutoConsoleCommand CmdClearSimulatedDescriptor(
+	TEXT("Input.Devices.ClearSimulatedDescriptor"),
+	TEXT("(Non-shipping) Remove the simulated descriptor for a given device ID, restoring the real one.\n")
+	TEXT("Usage: Input.Devices.ClearSimulatedDescriptor <DeviceId>"),
+	FConsoleCommandWithArgsDelegate::CreateLambda([](const TArray<FString>& Args)
+	{
+		if (Args.Num() < 1)
+		{
+			UE_LOG(LogInputDeviceRegistry, Warning,
+				TEXT("Input.Devices.ClearSimulatedDescriptor: expected 1 argument: <DeviceId>"));
+			return;
+		}
+
+		const int32 RawDeviceId = FCString::Atoi(*Args[0]);
+		const FInputDeviceId DeviceId = FInputDeviceId::CreateFromInternalId(RawDeviceId);
+		if (!DeviceId.IsValid())
+		{
+			UE_LOG(LogInputDeviceRegistry, Warning,
+				TEXT("Input.Devices.ClearSimulatedDescriptor: DeviceId '%s' is not valid."), *Args[0]);
+			return;
+		}
+
+		FInputDeviceRegistry::ClearSimulatedDescriptor(DeviceId);
+	})
+);
+
+// Input.Devices.ClearAllSimulatedDescriptors
+static FAutoConsoleCommand CmdClearAllSimulatedDescriptors(
+	TEXT("Input.Devices.ClearAllSimulatedDescriptors"),
+	TEXT("(Non-shipping) Remove all simulated device descriptors, restoring real ones for all devices."),
+	FConsoleCommandDelegate::CreateLambda([]()
+	{
+		FInputDeviceRegistry::ClearAllSimulatedDescriptors();
+	})
+);
+
+#endif // !UE_BUILD_SHIPPING
diff --git a/Engine/Source/Runtime/ApplicationCore/Public/GenericPlatform/InputDeviceRegistry.h b/Engine/Source/Runtime/ApplicationCore/Public/GenericPlatform/InputDeviceRegistry.h
--- a/Engine/Source/Runtime/ApplicationCore/Public/GenericPlatform/InputDeviceRegistry.h
+++ b/Engine/Source/Runtime/ApplicationCore/Public/GenericPlatform/InputDeviceRegistry.h
@@ -79,13 +79,51 @@ public:
 	/**
 	 * Look up the descriptor for a device by its FInputDeviceId.
 	 * Returns an empty optional if the device has not been registered yet.
+	 *
+	 * If a simulated descriptor has been set for this device (see SetSimulatedDescriptor),
+	 * that takes priority over the real descriptor.
 	 * Safe to call from any thread.
 	 */
 	static UE_API TOptional<FInputDeviceDescriptor> FindDescriptor(const FInputDeviceId DeviceId);
 
+#if !UE_BUILD_SHIPPING
+	/**
+	 * Override the descriptor returned by FindDescriptor for a specific device ID with a
+	 * simulated one. Intended for tooling and non-shipping testing (e.g. input preview,
+	 * device simulation) that needs to pretend a particular device is connected without
+	 * affecting the real runtime registry.
+	 *
+	 * The simulated descriptor is stored separately from the real registry and is never
+	 * written to by the platform layer. Call ClearSimulatedDescriptor to remove it.
+	 *
+	 * Not available in Shipping builds.
+	 */
+	static UE_API void SetSimulatedDescriptor(const FInputDeviceId DeviceId, const FInputDeviceDescriptor& SimulatedDescriptor);
+
+	/**
+	 * Remove the simulated descriptor for a device, restoring FindDescriptor to return
+	 * the real registered descriptor (if any).
+	 *
+	 * Not available in Shipping builds.
+	 */
+	static UE_API void ClearSimulatedDescriptor(const FInputDeviceId DeviceId);
+
+	/**
+	 * Remove all simulated descriptors at once (e.g. on PIE end or tool teardown).
+	 *
+	 * Not available in Shipping builds.
+	 */
+	static UE_API void ClearAllSimulatedDescriptors();
+#endif // !UE_BUILD_SHIPPING
+
 private:
 	mutable FRWLock Lock;
 	TMap<FInputDeviceId, FInputDeviceDescriptor> Descriptors;
+
+#if !UE_BUILD_SHIPPING
+	/** Descriptors set via SetSimulatedDescriptor. Checked first by FindDescriptor. */
+	TMap<FInputDeviceId, FInputDeviceDescriptor> SimulatedDescriptors;
+#endif // !UE_BUILD_SHIPPING
 };
 
-#undef UE_API
\ No newline at end of file
+#undef UE_API
