[Input] FInputDeviceScope refactor and deprecation.

- Deprecation the FInputDeviceScope type. Instead, you should use the new FInputDeviceRegistry type instead. This registry is thread safe and removes the need for Input API's to wrap every event with a FInputDeviceScope, which is very error prone and confusing.
- The FInputDeviceRegistry should be populated any time a new device is *connected*, and has an assigned FInputDeviceId. This way, the device name/type is evaluated once, upon connection, instead of every single tick.
- Effectively make the FInputDeviceScope no-op unless the UE_USE_LEGACY_INPUT_DEVICE_SCOPE define is set to 1. Users can define this in their Target.cs file if you still depend on FInputDeviceScope.
- For forward compatibility, when a device is connected (IPlatformInputDeviceMapper::Internal_MapInputDeviceToUser is called) if there is not already a new device registry entry available, we will attempt to populate it from a device scope that is available.
- Enforces a better contract for the FInputDeviceId to each platform and makes them all consistent.
- Linux and Mac were completely missing FInputDeviceScopes, and they now have a FInputDeviceRegistry call to resolve that.
- In CommonInput and the Input Device Subsystem, utilize this new FInputDeviceRegistry instead of FInputDeviceScope's. The names of the types for common input should be the same, and do not require asset updates. Same with the hardware device identifiers.

#nosubmit, missing Xbox thing.
#Synced-CL 52059832

#Synced-CL 52059832

diff --git a/Engine/Plugins/Experimental/RawInput/Source/RawInput/Private/Windows/RawInputWindows.cpp b/Engine/Plugins/Experimental/RawInput/Source/RawInput/Private/Windows/RawInputWindows.cpp
--- a/Engine/Plugins/Experimental/RawInput/Source/RawInput/Private/Windows/RawInputWindows.cpp
+++ b/Engine/Plugins/Experimental/RawInput/Source/RawInput/Private/Windows/RawInputWindows.cpp
@@ -1003,7 +1003,9 @@ void FRawInputWindows::SendControllerEvents()
 		// e.g if a button is still down or axis has a value (e.g. wheel not in centre)
 		if (DeviceEntry.bNeedsUpdate)
 		{
+PRAGMA_DISABLE_DEPRECATION_WARNINGS
 			FInputDeviceScope InputScope(this, RawInputInterfaceName, DeviceEntryPair.Key, DeviceEntry.DeviceData.DeviceName);
+PRAGMA_ENABLE_DEPRECATION_WARNINGS
 
 			for (FButtonData& DeviceButtonData : DeviceEntry.ButtonData)
 			{
diff --git a/Engine/Plugins/GameInputWindows/Source/GameInputWindows/Private/GameInputWindowsDevice.cpp b/Engine/Plugins/GameInputWindows/Source/GameInputWindows/Private/GameInputWindowsDevice.cpp
--- a/Engine/Plugins/GameInputWindows/Source/GameInputWindows/Private/GameInputWindowsDevice.cpp
+++ b/Engine/Plugins/GameInputWindows/Source/GameInputWindows/Private/GameInputWindowsDevice.cpp
@@ -124,11 +124,13 @@ void FGameInputWindowsInputDevice::HandleDeviceConnected(IGameInputDevice* Devic
 		
 		const FPlatformUserId UserToAssign = DeviceMapper.GetPlatformUserForNewlyConnectedDevice();
 		
+PRAGMA_DISABLE_DEPRECATION_WARNINGS
 		FInputDeviceScope InputScope(
 			nullptr,
 			UE::GameInput::GetDeviceScopeClassName(),
 			DeviceMapper.GetUserIndexForPlatformUser(UserToAssign),
 			FString(UE::GameInput::GetHardwareDeviceIdentifierName(Device)));
+PRAGMA_ENABLE_DEPRECATION_WARNINGS
 		
 		{
 			FInputDeviceDescriptor Desc;
diff --git a/Engine/Plugins/Runtime/CommonUI/Source/CommonInput/Private/CommonInputPreprocessor.cpp b/Engine/Plugins/Runtime/CommonUI/Source/CommonInput/Private/CommonInputPreprocessor.cpp
--- a/Engine/Plugins/Runtime/CommonUI/Source/CommonInput/Private/CommonInputPreprocessor.cpp
+++ b/Engine/Plugins/Runtime/CommonUI/Source/CommonInput/Private/CommonInputPreprocessor.cpp
@@ -10,6 +10,7 @@
 #include "Engine/LocalPlayer.h"
 #include "Framework/Application/SlateApplication.h"
 #include "Framework/Application/SlateUser.h"
+#include "GameFramework/InputDeviceSubsystem.h"
 #include "ICommonInputModule.h"
 
 #include "Engine/Engine.h"
@@ -239,15 +240,20 @@ void FCommonInputPreprocessor::RefreshCurrentInputMethod(ECommonInputType InputM
 	    && ICommonInputModule::GetSettings().GetEnableAutomaticGamepadTypeDetection()
 	    && UCommonInputPlatformSettings::Get()->CanChangeGamepadType())
 	{
-		if (const FInputDeviceScope* DeviceScope = FInputDeviceScope::GetCurrent())
+		const ULocalPlayer* LP = InputSubsystem.GetLocalPlayer<ULocalPlayer>();
+		const UInputDeviceSubsystem* DeviceSubsystem = UInputDeviceSubsystem::Get();
+		
+		if (LP && DeviceSubsystem)
 		{
-			if ((DeviceScope->InputDeviceName != LastSeenGamepadInputDeviceName) || (DeviceScope->HardwareDeviceIdentifier != LastSeenGamepadHardwareDeviceIdentifier))
+			const FHardwareDeviceIdentifier DeviceScope = DeviceSubsystem->GetMostRecentlyUsedHardwareDevice(LP->GetPlatformUserId());
+			const FString HardwareDeviceIdentifierString = DeviceScope.HardwareDeviceIdentifier.ToString();
+			if ((DeviceScope.InputClassName != LastSeenGamepadInputDeviceName) || (HardwareDeviceIdentifierString != LastSeenGamepadHardwareDeviceIdentifier))
 			{
-				LastSeenGamepadInputDeviceName = DeviceScope->InputDeviceName;
-				LastSeenGamepadHardwareDeviceIdentifier = DeviceScope->HardwareDeviceIdentifier;
+				LastSeenGamepadInputDeviceName = DeviceScope.InputClassName;
+				LastSeenGamepadHardwareDeviceIdentifier = HardwareDeviceIdentifierString;
 
 				const FName GamepadInputType = InputSubsystem.GetCurrentGamepadName();
-				const FName BestGamepadType = UCommonInputPlatformSettings::Get()->GetBestGamepadNameForHardware(GamepadInputType, DeviceScope->InputDeviceName, DeviceScope->HardwareDeviceIdentifier);
+				const FName BestGamepadType = UCommonInputPlatformSettings::Get()->GetBestGamepadNameForHardware(GamepadInputType, DeviceScope.InputClassName, HardwareDeviceIdentifierString);
 				if (BestGamepadType != GamepadInputType)
 				{
 					UE_LOG(LogCommonInput, Log, TEXT("UCommonInputSubsystem: Autodetect changed GamepadInputType to %s"), *BestGamepadType.ToString());
diff --git a/Engine/Plugins/Runtime/GameInput/Source/GameInputBase/Private/Processors/GameInputDeviceProcessor_Raw.cpp b/Engine/Plugins/Runtime/GameInput/Source/GameInputBase/Private/Processors/GameInputDeviceProcessor_Raw.cpp
--- a/Engine/Plugins/Runtime/GameInput/Source/GameInputBase/Private/Processors/GameInputDeviceProcessor_Raw.cpp
+++ b/Engine/Plugins/Runtime/GameInput/Source/GameInputBase/Private/Processors/GameInputDeviceProcessor_Raw.cpp
@@ -500,12 +500,14 @@ void FGameInputRawDeviceProcessor::EvaluatePackedUInt8ButtonChordsAsAnalog(
 
 	IPlatformInputDeviceMapper& DeviceMapper = IPlatformInputDeviceMapper::Get();
 
+PRAGMA_DISABLE_DEPRECATION_WARNINGS
 	FInputDeviceScope InputScope(
 		nullptr,
 		UE::GameInput::GetDeviceScopeClassName(),
 		DeviceMapper.GetUserIndexForPlatformUser(Params.PlatformUserId),
 		FString(UE::GameInput::GetHardwareDeviceIdentifierName(Params.Device)));
-	
+PRAGMA_ENABLE_DEPRECATION_WARNINGS
+
 	if (CurrentValue != PreviousValue)
 	{
 		const float CurrentValueFloat = RawValueNibbleValueToFloat(CurrentValue);
@@ -553,7 +555,9 @@ void FGameInputRawDeviceProcessor::EvaluatePackedUInt8ButtonChordsAsButtonEvents
 
 	IPlatformInputDeviceMapper& DeviceMapper = IPlatformInputDeviceMapper::Get();
 
+PRAGMA_DISABLE_DEPRECATION_WARNINGS
 	const FInputDeviceScope InputScope(nullptr, FName("GameInput"), DeviceMapper.GetUserIndexForPlatformUser(Params.PlatformUserId), FString(UE::GameInput::GetHardwareDeviceIdentifierName(Params.Device)));
+PRAGMA_ENABLE_DEPRECATION_WARNINGS
 
 	auto ForEachValidKey = [&ButtonNames](TFunction<void(const FGamepadKeyNames::Type&)>&& Predicate)
 	{
diff --git a/Engine/Plugins/Runtime/GameInput/Source/GameInputBase/Private/Processors/IGameInputDeviceProcessor.cpp b/Engine/Plugins/Runtime/GameInput/Source/GameInputBase/Private/Processors/IGameInputDeviceProcessor.cpp
--- a/Engine/Plugins/Runtime/GameInput/Source/GameInputBase/Private/Processors/IGameInputDeviceProcessor.cpp
+++ b/Engine/Plugins/Runtime/GameInput/Source/GameInputBase/Private/Processors/IGameInputDeviceProcessor.cpp
@@ -55,12 +55,13 @@ void IGameInputDeviceProcessor::OnControllerAnalog(const FGameInputEventParams&
 		{
 			if (bSetDeviceScope)
 			{
+PRAGMA_DISABLE_DEPRECATION_WARNINGS
 				FInputDeviceScope InputScope(
 					nullptr,
 					UE::GameInput::GetDeviceScopeClassName(),
 					IPlatformInputDeviceMapper::Get().GetUserIndexForPlatformUser(Params.PlatformUserId),
 					FString(UE::GameInput::GetHardwareDeviceIdentifierName(Params.Device)));
-				
+PRAGMA_ENABLE_DEPRECATION_WARNINGS
 				Params.MessageHandler->OnControllerAnalog(GamePadKey, Params.PlatformUserId, Params.InputDeviceId, NewAxisValueNormalized);
 			}
 			else
@@ -115,7 +116,9 @@ void IGameInputDeviceProcessor::EvaluateButtonStates(
 			// Check for button state change
 			if (0 != (ActionMask & BitMask))
 			{
+PRAGMA_DISABLE_DEPRECATION_WARNINGS
 				FInputDeviceScope InputScope(nullptr, UE::GameInput::GetDeviceScopeClassName(), DeviceMapper.GetUserIndexForPlatformUser(Params.PlatformUserId), FString(UE::GameInput::GetHardwareDeviceIdentifierName(Params.Device)));
+PRAGMA_ENABLE_DEPRECATION_WARNINGS
 
 				if (0 != (CurrentButtonHeldMask & BitMask))
 				{
@@ -134,7 +137,9 @@ void IGameInputDeviceProcessor::EvaluateButtonStates(
 			// Check for repeat key
 			if (0 != (RepeatMask & BitMask) && CurrentTime > RepeatTime[n])
 			{
+PRAGMA_DISABLE_DEPRECATION_WARNINGS
 				FInputDeviceScope InputScope(nullptr, UE::GameInput::GetDeviceScopeClassName(), DeviceMapper.GetUserIndexForPlatformUser(Params.PlatformUserId), FString(UE::GameInput::GetHardwareDeviceIdentifierName(Params.Device)));
+PRAGMA_ENABLE_DEPRECATION_WARNINGS
 
 				RepeatTime[n] = CurrentTime + ButtonRepeatDelay;
 				UE_LOG(LogGameInput, Verbose, TEXT("[FGameInputDevice::EvaluateButtonStates] (PlatformUserId = %d, InputDeviceId = %d) - Button '%s' Repeat Pressed"), Params.PlatformUserId.GetInternalId(), Params.InputDeviceId.GetId(), *ButtonKey.ToString());
diff --git a/Engine/Plugins/Runtime/Steam/SteamController/Source/SteamController/Private/SteamController.cpp b/Engine/Plugins/Runtime/Steam/SteamController/Source/SteamController/Private/SteamController.cpp
--- a/Engine/Plugins/Runtime/Steam/SteamController/Source/SteamController/Private/SteamController.cpp
+++ b/Engine/Plugins/Runtime/Steam/SteamController/Source/SteamController/Private/SteamController.cpp
@@ -140,7 +140,9 @@ public:
 			// Doesn't seem to be a good way to get a non-localized string, so use generic name for scope
 			static FName SystemName(TEXT("SteamController"));
 			static FString ControllerName(TEXT("SteamController"));
+PRAGMA_DISABLE_DEPRECATION_WARNINGS
 			FInputDeviceScope InputScope(this, SystemName, i, ControllerName);
+PRAGMA_ENABLE_DEPRECATION_WARNINGS
 
 			FPlatformUserId UserId = FPlatformMisc::GetPlatformUserForUserIndex(i);
 			FInputDeviceId DeviceId = INPUTDEVICEID_NONE;
diff --git a/Engine/Plugins/Runtime/Windows/WinDualShock/Source/WinDualShock/Private/WinDualShock.cpp b/Engine/Plugins/Runtime/Windows/WinDualShock/Source/WinDualShock/Private/WinDualShock.cpp
--- a/Engine/Plugins/Runtime/Windows/WinDualShock/Source/WinDualShock/Private/WinDualShock.cpp
+++ b/Engine/Plugins/Runtime/Windows/WinDualShock/Source/WinDualShock/Private/WinDualShock.cpp
@@ -141,7 +141,9 @@ public:
 			// Don't send events until we have the controller type as we need to set the device scope up.
 			if (Controllers.GetControllerTypeIdentifier(UserIndex) != ESonyControllerType::None)
 			{
-				FInputDeviceScope InputScope(this, InputClassName, UserIndex, Controllers.GetControllerTypeIdentifierName(UserIndex));
+PRAGMA_DISABLE_DEPRECATION_WARNINGS
+			FInputDeviceScope InputScope(this, InputClassName, UserIndex, Controllers.GetControllerTypeIdentifierName(UserIndex));
+PRAGMA_ENABLE_DEPRECATION_WARNINGS
 				Controllers.SendControllerEvents(UserIndex, MessageHandler);
 			}
 		}
diff --git a/Engine/Plugins/Runtime/Windows/XInputDevice/Source/XInputDevice/Private/XInputInterface.cpp b/Engine/Plugins/Runtime/Windows/XInputDevice/Source/XInputDevice/Private/XInputInterface.cpp
--- a/Engine/Plugins/Runtime/Windows/XInputDevice/Source/XInputDevice/Private/XInputInterface.cpp
+++ b/Engine/Plugins/Runtime/Windows/XInputDevice/Source/XInputDevice/Private/XInputInterface.cpp
@@ -234,9 +234,11 @@ void XInputInterface::SendControllerEvents()
 		
 	for ( int32 ControllerIndex = 0; ControllerIndex < MAX_NUM_XINPUT_CONTROLLERS; ++ControllerIndex )
 	{
+		PRAGMA_DISABLE_DEPRECATION_WARNINGS
 		// Set input scope, there doesn't seem to be a reliable way to differentiate 360 vs Xbox one controllers so use generic name
 		FInputDeviceScope InputScope(this, XInputInterfaceName, ControllerIndex, XInputControllerIdentifier);
-
+		PRAGMA_ENABLE_DEPRECATION_WARNINGS
+		
 		FControllerState& ControllerState = ControllerStates[ControllerIndex];
 
 		const bool bWasConnected = bWereConnected[ControllerIndex];
diff --git a/Engine/Source/Runtime/ApplicationCore/Private/Android/AndroidInputInterface.cpp b/Engine/Source/Runtime/ApplicationCore/Private/Android/AndroidInputInterface.cpp
--- a/Engine/Source/Runtime/ApplicationCore/Private/Android/AndroidInputInterface.cpp
+++ b/Engine/Source/Runtime/ApplicationCore/Private/Android/AndroidInputInterface.cpp
@@ -1233,7 +1233,9 @@ void FAndroidInputInterface::SendControllerEvents()
 		// This is a temp solution with a hardcoded string which can be mapped to FHardwareDeviceIdentifier::DefaultMobileTouch.
 		// TODO: Future improvement is needed to acquire them by values in the Android input device info.
 		FName DeviceName = TouchDeviceInfo ? TouchDeviceInfo->Name : InputClassName_DefaultMobileTouch;
+PRAGMA_DISABLE_DEPRECATION_WARNINGS
 		FInputDeviceScope InputScope(nullptr, DeviceName, TouchInputDeviceId.GetId(), HardwareDeviceIdentifier_DefaultMobileTouch);
+PRAGMA_ENABLE_DEPRECATION_WARNINGS
 
 		IPlatformInputDeviceMapper& DeviceMapper = IPlatformInputDeviceMapper::Get();
 		FPlatformUserId PlatformUserId = DeviceMapper.GetUserForInputDevice(TouchInputDeviceId);
@@ -1271,7 +1273,9 @@ void FAndroidInputInterface::SendControllerEvents()
 			///The FInputDeviceScope::HardwareDeviceIdentifier has to be one of values in UInputPlatformSettings::HardwareDevices.
 			// This is a temp solution with a hardcoded string which can be mapped to FHardwareDeviceIdentifier::DefaultGamepad.
 			// TODO: Future improvement is needed to acquire them by values in the Android input device info.
+PRAGMA_DISABLE_DEPRECATION_WARNINGS
 			FInputDeviceScope InputScope(nullptr, DeviceInfo.Name, InputDeviceId.GetId(), HardwareDeviceIdentifier_DefaultGamepad);
+PRAGMA_ENABLE_DEPRECATION_WARNINGS
 
 			FAndroidGamepadDeviceMapping* DeviceData = GameControllerDataMap.Find(InputDeviceId);
 			if (DeviceData == nullptr || !DeviceData->DeviceConfig.IsValid())
diff --git a/Engine/Source/Runtime/ApplicationCore/Private/Apple/AppleControllerInterface.cpp b/Engine/Source/Runtime/ApplicationCore/Private/Apple/AppleControllerInterface.cpp
--- a/Engine/Source/Runtime/ApplicationCore/Private/Apple/AppleControllerInterface.cpp
+++ b/Engine/Source/Runtime/ApplicationCore/Private/Apple/AppleControllerInterface.cpp
@@ -494,7 +494,9 @@ void FAppleControllerInterface::SendControllerEvents()
 		///The FInputDeviceScope::HardwareDeviceIdentifier has to be one of values in UInputPlatformSettings::HardwareDevices.
 		// This is a temp solution with a hardcoded string which can be mapped to FHardwareDeviceIdentifier::DefaultGamepad.
 		// TODO: Future improvement is needed to acquire them by values in the iOS input device info.
+PRAGMA_DISABLE_DEPRECATION_WARNINGS
 		FInputDeviceScope InputScope(nullptr, GetControllerTypeName(Controller.ControllerType), DeviceId.GetId(), HardwareDeviceIdentifier_DefaultGamepad);
+PRAGMA_ENABLE_DEPRECATION_WARNINGS
 
 		uint32 ButtonIndex = PrevButtonIndex;	
 		PrevButtonIndex = CurrentButtonIndex;
diff --git a/Engine/Source/Runtime/ApplicationCore/Private/GenericPlatform/GenericApplication.cpp b/Engine/Source/Runtime/ApplicationCore/Private/GenericPlatform/GenericApplication.cpp
--- a/Engine/Source/Runtime/ApplicationCore/Private/GenericPlatform/GenericApplication.cpp
+++ b/Engine/Source/Runtime/ApplicationCore/Private/GenericPlatform/GenericApplication.cpp
@@ -46,6 +46,20 @@ const FGamepadKeyNames::Type FGamepadKeyNames::RightStickDown("Gamepad_RightStic
 const FGamepadKeyNames::Type FGamepadKeyNames::RightStickRight("Gamepad_RightStick_Right");
 const FGamepadKeyNames::Type FGamepadKeyNames::RightStickLeft("Gamepad_RightStick_Left");
 
+// The FInputDeviceRegistry should be used instead of the FInputDeviceScope for setting 
+// device metadata like the name and type of device. The FInputDeviceRegistry improves on
+// the FInputDeviceScope by being thread safe, as well as only being required once upon
+// device connection, instead of being scoped around every single input event.
+//
+// If your project does depend on the FInputDeviceScope stack, you can set this in your *.Target.cs file:
+//
+// 	GlobalDefinitions.Add("UE_USE_LEGACY_INPUT_DEVICE_SCOPE=1");
+#ifndef UE_USE_LEGACY_INPUT_DEVICE_SCOPE
+	#define UE_USE_LEGACY_INPUT_DEVICE_SCOPE 0
+#endif
+
+PRAGMA_DISABLE_DEPRECATION_WARNINGS
+
 namespace UE::InputDeviceScope::Private
 {
 using TArrayInputDeviceScope = TArray<FInputDeviceScope*, TInlineAllocator<16>>; 
@@ -62,17 +76,21 @@ FInputDeviceScope::FInputDeviceScope(IInputDevice* InInputDevice, FName InInputD
 	, HardwareDeviceHandle(InHardwareDeviceHandle)
 	, HardwareDeviceIdentifier(InHardwareDeviceIdentifier)
 {
+#if UE_USE_LEGACY_INPUT_DEVICE_SCOPE
 	// Add to scope stack
 	UE::InputDeviceScope::Private::GetScopeStack().Add(this);
+#endif
 }
 
 FInputDeviceScope::~FInputDeviceScope()
 {
+#if UE_USE_LEGACY_INPUT_DEVICE_SCOPE
 	UE::InputDeviceScope::Private::TArrayInputDeviceScope& ScopeStack = UE::InputDeviceScope::Private::GetScopeStack();
 
 	// This should always be the top of the stack
 	ensureMsgf((ScopeStack.Num() > 0 && ScopeStack.Last() == this), TEXT("FInputDeviceScope was not destroyed in correct order!"));
 	ScopeStack.Remove(this);
+#endif
 }
 
 const FInputDeviceScope* FInputDeviceScope::GetCurrent()
@@ -85,6 +103,8 @@ const FInputDeviceScope* FInputDeviceScope::GetCurrent()
 	return nullptr;
 }
 
+PRAGMA_ENABLE_DEPRECATION_WARNINGS
+
 float GDebugSafeZoneRatio = 1.0f;
 float GDebugActionZoneRatio = 1.0f;
 
diff --git a/Engine/Source/Runtime/ApplicationCore/Private/GenericPlatform/GenericPlatformInputDeviceMapper.cpp b/Engine/Source/Runtime/ApplicationCore/Private/GenericPlatform/GenericPlatformInputDeviceMapper.cpp
--- a/Engine/Source/Runtime/ApplicationCore/Private/GenericPlatform/GenericPlatformInputDeviceMapper.cpp
+++ b/Engine/Source/Runtime/ApplicationCore/Private/GenericPlatform/GenericPlatformInputDeviceMapper.cpp
@@ -6,6 +6,7 @@
 #include "Misc/ConfigCacheIni.h"
 #include "GenericPlatform/GenericApplicationMessageHandler.h"
 #include "GenericPlatform/InputDeviceMappingPolicy.h"
+#include "GenericPlatform/InputDeviceRegistry.h"
 
 ///////////////////////////////////////////////////////////////////////////
 // IPlatformInputDeviceMapper
@@ -361,8 +362,35 @@ bool IPlatformInputDeviceMapper::Internal_MapInputDeviceToUser(FInputDeviceId De
 	InputDeviceState.OwningPlatformUser = UserId;
 	InputDeviceState.ConnectionState = ConnectionState;
 
-	// If there is a valid FInputDeviceScope at the time of connection, then the hardware info is available.
-	if (FInputDeviceScope::GetCurrent() != nullptr)
+#if 0	// todo_bh: Not sure if we actually need this here since we have updated all platforms already
+PRAGMA_DISABLE_DEPRECATION_WARNINGS
+	// Forward compatibility for FInputDeviceScope with the newer FInputDeviceRegistry system. 
+	// If we have an input device scope, and there is not yet a device descriptor set for this device,
+	// then populate it with the FInputDeviceScope data if possible.
+	if (IsInGameThread())
+	{
+		if(!FInputDeviceRegistry::FindDescriptor(DeviceId).IsSet())
+		{
+			if (const FInputDeviceScope* Scope = FInputDeviceScope::GetCurrent())
+			{
+				UE_LOG(LogInputDeviceMapper, Log, TEXT("Registering device based on the FInputDeviceScope %d"), DeviceId.GetId());
+			
+				FInputDeviceRegistry::RegisterDevice(
+					FInputDeviceDescriptor
+					{
+						.HardwareDeviceHandle     = DeviceId,
+						.InputDeviceName          = Scope->InputDeviceName,
+						.HardwareDeviceIdentifier = FName(*Scope->HardwareDeviceIdentifier),
+					});
+			}
+		}	
+	}
+PRAGMA_ENABLE_DEPRECATION_WARNINGS
+#endif
+
+	// If there is a valid FInputDeviceScope at the time of connection, or the device registry is aware of this device,
+	// then the hardware info is available.
+	if (FInputDeviceRegistry::FindDescriptor(DeviceId).IsSet())
 	{
 		OnInputDeviceHardwareInfoAvailable.Broadcast(ConnectionState, UserId, DeviceId);	
 	}
diff --git a/Engine/Source/Runtime/ApplicationCore/Private/IOS/IOSInputInterface.cpp b/Engine/Source/Runtime/ApplicationCore/Private/IOS/IOSInputInterface.cpp
--- a/Engine/Source/Runtime/ApplicationCore/Private/IOS/IOSInputInterface.cpp
+++ b/Engine/Source/Runtime/ApplicationCore/Private/IOS/IOSInputInterface.cpp
@@ -199,8 +199,10 @@ void FIOSInputInterface::ProcessTouchesAndKeys(uint32 ControllerId, const TArray
 		///The FInputDeviceScope::HardwareDeviceIdentifier has to be one of values in UInputPlatformSettings::HardwareDevices.
 		// This is a temp solution with a hardcoded string which can be mapped to FHardwareDeviceIdentifier::DefaultMobileTouch.
 		// TODO: Future improvement is needed to acquire them by values in the iOS input device info.
+PRAGMA_DISABLE_DEPRECATION_WARNINGS
 		FInputDeviceScope InputScope(nullptr, InputClassName_DefaultMobileTouch, 0, HardwareDeviceIdentifier_DefaultMobileTouch);
-		
+PRAGMA_ENABLE_DEPRECATION_WARNINGS
+
         // send input to handler
         if (Touch.Type == TouchBegan)
         {
diff --git a/Engine/Source/Runtime/ApplicationCore/Private/Windows/WindowsApplication.cpp b/Engine/Source/Runtime/ApplicationCore/Private/Windows/WindowsApplication.cpp
--- a/Engine/Source/Runtime/ApplicationCore/Private/Windows/WindowsApplication.cpp
+++ b/Engine/Source/Runtime/ApplicationCore/Private/Windows/WindowsApplication.cpp
@@ -3227,7 +3227,9 @@ void FWindowsApplication::CheckForShiftUpEvents(const int32 KeyCode)
 	{
 		ModifierKeyState[ModifierKeyIndex] = false;
 		
+PRAGMA_DISABLE_DEPRECATION_WARNINGS
 		FInputDeviceScope InputScope(nullptr, UE::WindowsInput::InputClassName, IPlatformInputDeviceMapper::Get().GetDefaultInputDevice().GetId(), UE::WindowsInput::KBMInputHardwareName);
+PRAGMA_ENABLE_DEPRECATION_WARNINGS
 		MessageHandler->OnKeyUp( KeyCode, 0, false );
 	}
 }
@@ -3286,7 +3288,9 @@ int32 FWindowsApplication::ProcessDeferredMessage( const FDeferredWindowsMessage
 
 				// LPARAM bit 30 will be ZERO for new presses, or ONE if this is a repeat
 				const bool bIsRepeat = ( lParam & 0x40000000 ) != 0;
-				FInputDeviceScope InputScope(nullptr, UE::WindowsInput::InputClassName, IPlatformInputDeviceMapper::Get().GetDefaultInputDevice().GetId(), UE::WindowsInput::KBMInputHardwareName); 
+PRAGMA_DISABLE_DEPRECATION_WARNINGS
+			FInputDeviceScope InputScope(nullptr, UE::WindowsInput::InputClassName, IPlatformInputDeviceMapper::Get().GetDefaultInputDevice().GetId(), UE::WindowsInput::KBMInputHardwareName);
+PRAGMA_ENABLE_DEPRECATION_WARNINGS
 				MessageHandler->OnKeyChar( Character, bIsRepeat );
 
 				// Note: always return 0 to handle the message.  Win32 beeps if WM_CHAR is not handled...
@@ -3366,7 +3370,9 @@ int32 FWindowsApplication::ProcessDeferredMessage( const FDeferredWindowsMessage
 
 				// Get the character code from the virtual key pressed.  If 0, no translation from virtual key to character exists
 				uint32 CharCode = ::MapVirtualKey( Win32Key, MAPVK_VK_TO_CHAR );
-				FInputDeviceScope InputScope(nullptr, UE::WindowsInput::InputClassName, IPlatformInputDeviceMapper::Get().GetDefaultInputDevice().GetId(), UE::WindowsInput::KBMInputHardwareName); 
+PRAGMA_DISABLE_DEPRECATION_WARNINGS
+			FInputDeviceScope InputScope(nullptr, UE::WindowsInput::InputClassName, IPlatformInputDeviceMapper::Get().GetDefaultInputDevice().GetId(), UE::WindowsInput::KBMInputHardwareName);
+PRAGMA_ENABLE_DEPRECATION_WARNINGS
 				const bool Result = MessageHandler->OnKeyDown( ActualKey, CharCode, bIsRepeat );
 
 				// Always return 0 to handle the message or else windows will beep
@@ -3445,7 +3451,9 @@ int32 FWindowsApplication::ProcessDeferredMessage( const FDeferredWindowsMessage
 
 				// Key up events are never repeats
 				const bool bIsRepeat = false;
-				FInputDeviceScope InputScope(nullptr, UE::WindowsInput::InputClassName, IPlatformInputDeviceMapper::Get().GetDefaultInputDevice().GetId(), UE::WindowsInput::KBMInputHardwareName); 
+PRAGMA_DISABLE_DEPRECATION_WARNINGS
+			FInputDeviceScope InputScope(nullptr, UE::WindowsInput::InputClassName, IPlatformInputDeviceMapper::Get().GetDefaultInputDevice().GetId(), UE::WindowsInput::KBMInputHardwareName);
+PRAGMA_ENABLE_DEPRECATION_WARNINGS
 				const bool Result = MessageHandler->OnKeyUp( ActualKey, CharCode, bIsRepeat );
 
 				// Note that we allow system keys to pass through to DefWndProc here, so that core features
@@ -3533,7 +3541,9 @@ int32 FWindowsApplication::ProcessDeferredMessage( const FDeferredWindowsMessage
 					check(0);
 				}
 
-				FInputDeviceScope InputScope(nullptr, UE::WindowsInput::InputClassName, IPlatformInputDeviceMapper::Get().GetDefaultInputDevice().GetId(), UE::WindowsInput::KBMInputHardwareName); 
+PRAGMA_DISABLE_DEPRECATION_WARNINGS
+			FInputDeviceScope InputScope(nullptr, UE::WindowsInput::InputClassName, IPlatformInputDeviceMapper::Get().GetDefaultInputDevice().GetId(), UE::WindowsInput::KBMInputHardwareName);
+PRAGMA_ENABLE_DEPRECATION_WARNINGS
 				if (bMouseUp)
 				{
 					return MessageHandler->OnMouseUp( MouseButton, CursorPos ) ? 0 : 1;
@@ -3553,7 +3563,9 @@ int32 FWindowsApplication::ProcessDeferredMessage( const FDeferredWindowsMessage
 		// Mouse Movement
 		case WM_INPUT:
 			{
-				FInputDeviceScope InputScope(nullptr, UE::WindowsInput::InputClassName, IPlatformInputDeviceMapper::Get().GetDefaultInputDevice().GetId(), UE::WindowsInput::KBMInputHardwareName); 
+PRAGMA_DISABLE_DEPRECATION_WARNINGS
+			FInputDeviceScope InputScope(nullptr, UE::WindowsInput::InputClassName, IPlatformInputDeviceMapper::Get().GetDefaultInputDevice().GetId(), UE::WindowsInput::KBMInputHardwareName);
+PRAGMA_ENABLE_DEPRECATION_WARNINGS
 				if( DeferredMessage.RawInputFlags == MOUSE_MOVE_RELATIVE )
 				{
 					MessageHandler->OnRawMouseMove(DeferredMessage.X, DeferredMessage.Y);
@@ -3575,7 +3587,9 @@ int32 FWindowsApplication::ProcessDeferredMessage( const FDeferredWindowsMessage
 				BOOL Result = false;
 				if (!bUsingHighPrecisionMouseInput)
 				{
-					FInputDeviceScope InputScope(nullptr, UE::WindowsInput::InputClassName, IPlatformInputDeviceMapper::Get().GetDefaultInputDevice().GetId(), UE::WindowsInput::KBMInputHardwareName); 
+PRAGMA_DISABLE_DEPRECATION_WARNINGS
+			FInputDeviceScope InputScope(nullptr, UE::WindowsInput::InputClassName, IPlatformInputDeviceMapper::Get().GetDefaultInputDevice().GetId(), UE::WindowsInput::KBMInputHardwareName);
+PRAGMA_ENABLE_DEPRECATION_WARNINGS
 					Result = MessageHandler->OnMouseMove();
 				}
 
@@ -3594,7 +3608,9 @@ int32 FWindowsApplication::ProcessDeferredMessage( const FDeferredWindowsMessage
 
 				const FVector2D CursorPos(CursorPoint.x, CursorPoint.y);
 
-				FInputDeviceScope InputScope(nullptr, UE::WindowsInput::InputClassName, IPlatformInputDeviceMapper::Get().GetDefaultInputDevice().GetId(), UE::WindowsInput::KBMInputHardwareName); 
+PRAGMA_DISABLE_DEPRECATION_WARNINGS
+			FInputDeviceScope InputScope(nullptr, UE::WindowsInput::InputClassName, IPlatformInputDeviceMapper::Get().GetDefaultInputDevice().GetId(), UE::WindowsInput::KBMInputHardwareName);
+PRAGMA_ENABLE_DEPRECATION_WARNINGS
 				const BOOL Result = MessageHandler->OnMouseWheel( static_cast<float>( WheelDelta ) * SpinFactor, CursorPos );
 				return Result ? 0 : 1;
 			}
@@ -3628,7 +3644,9 @@ int32 FWindowsApplication::ProcessDeferredMessage( const FDeferredWindowsMessage
 					TUniquePtr<TOUCHINPUT[]> Inputs = MakeUnique<TOUCHINPUT[]>( InputCount );
 					if ( GetTouchInputInfo( (HTOUCHINPUT)lParam, InputCount, Inputs.Get(), sizeof(TOUCHINPUT) ) )
 					{
-						FInputDeviceScope InputScope(nullptr, UE::WindowsInput::InputClassName, IPlatformInputDeviceMapper::Get().GetDefaultInputDevice().GetId(), UE::WindowsInput::TouchInputHardwareName);
+PRAGMA_DISABLE_DEPRECATION_WARNINGS
+			FInputDeviceScope InputScope(nullptr, UE::WindowsInput::InputClassName, IPlatformInputDeviceMapper::Get().GetDefaultInputDevice().GetId(), UE::WindowsInput::TouchInputHardwareName);
+PRAGMA_ENABLE_DEPRECATION_WARNINGS
 
 						for ( uint32 i = 0; i < InputCount; i++ )
 						{
@@ -3702,7 +3720,9 @@ int32 FWindowsApplication::ProcessDeferredMessage( const FDeferredWindowsMessage
 
 		case WM_POINTERDOWN:
 			{
-				FInputDeviceScope InputScope(nullptr, UE::WindowsInput::InputClassName, IPlatformInputDeviceMapper::Get().GetDefaultInputDevice().GetId(), UE::WindowsInput::TouchInputHardwareName);
+PRAGMA_DISABLE_DEPRECATION_WARNINGS
+			FInputDeviceScope InputScope(nullptr, UE::WindowsInput::InputClassName, IPlatformInputDeviceMapper::Get().GetDefaultInputDevice().GetId(), UE::WindowsInput::TouchInputHardwareName);
+PRAGMA_ENABLE_DEPRECATION_WARNINGS
 
 				int32 PointerId = GET_POINTERID_WPARAM(wParam);
 
@@ -3758,7 +3778,9 @@ int32 FWindowsApplication::ProcessDeferredMessage( const FDeferredWindowsMessage
 
 		case WM_POINTERUPDATE:
 			{
-				FInputDeviceScope InputScope(nullptr, UE::WindowsInput::InputClassName, IPlatformInputDeviceMapper::Get().GetDefaultInputDevice().GetId(), UE::WindowsInput::TouchInputHardwareName);
+PRAGMA_DISABLE_DEPRECATION_WARNINGS
+			FInputDeviceScope InputScope(nullptr, UE::WindowsInput::InputClassName, IPlatformInputDeviceMapper::Get().GetDefaultInputDevice().GetId(), UE::WindowsInput::TouchInputHardwareName);
+PRAGMA_ENABLE_DEPRECATION_WARNINGS
 
 				int32 PointerId = GET_POINTERID_WPARAM(wParam);
 		
@@ -3821,7 +3843,9 @@ int32 FWindowsApplication::ProcessDeferredMessage( const FDeferredWindowsMessage
 
 		case WM_POINTERUP:
 			{
-				FInputDeviceScope InputScope(nullptr, UE::WindowsInput::InputClassName, IPlatformInputDeviceMapper::Get().GetDefaultInputDevice().GetId(), UE::WindowsInput::TouchInputHardwareName);
+PRAGMA_DISABLE_DEPRECATION_WARNINGS
+			FInputDeviceScope InputScope(nullptr, UE::WindowsInput::InputClassName, IPlatformInputDeviceMapper::Get().GetDefaultInputDevice().GetId(), UE::WindowsInput::TouchInputHardwareName);
+PRAGMA_ENABLE_DEPRECATION_WARNINGS
 
 				int32 PointerId = GET_POINTERID_WPARAM(wParam);
 
@@ -4818,7 +4842,9 @@ void FWindowsApplication::ProcessRawMouseInput(const RAWMOUSE& Mouse, RawMouseIn
 void FWindowsApplication::ProcessWorkerInputs()
 {
 	TRACE_CPUPROFILER_EVENT_SCOPE(ProcessWorkerInputs);
+PRAGMA_DISABLE_DEPRECATION_WARNINGS
 	FInputDeviceScope InputScope(nullptr, UE::WindowsInput::InputClassName, IPlatformInputDeviceMapper::Get().GetDefaultInputDevice().GetId(), UE::WindowsInput::KBMInputHardwareName);
+PRAGMA_ENABLE_DEPRECATION_WARNINGS
 
 	int32 xPosRelative = 0;
 	int32 yPosRelative = 0;
diff --git a/Engine/Source/Runtime/ApplicationCore/Public/GenericPlatform/GenericApplicationMessageHandler.h b/Engine/Source/Runtime/ApplicationCore/Public/GenericPlatform/GenericApplicationMessageHandler.h
--- a/Engine/Source/Runtime/ApplicationCore/Public/GenericPlatform/GenericApplicationMessageHandler.h
+++ b/Engine/Source/Runtime/ApplicationCore/Public/GenericPlatform/GenericApplicationMessageHandler.h
@@ -9,6 +9,7 @@
 #include "Templates/SharedPointer.h"
 #include "Misc/Optional.h"
 #include "GenericPlatform/GenericPlatformInputDeviceMapper.h"
+#include "GenericPlatform/InputDeviceRegistry.h"
 
 class FGenericWindow;
 
@@ -180,10 +181,17 @@ private:
 	TOptional<float> MaxHeight;
 };
 
-/** 
- * Context scope that indicates which IInputDevice is currently being handled. 
+class UE_DEPRECATED(5.8, "Use FInputDeviceRegistry::Get() instead") FInputDeviceScope;
+
+/**
+ * Context scope that indicates which IInputDevice is currently being handled.
  * This can be used to determine hardware-specific information when handling input from FGenericApplicationMessageHandler subclasses.
  * This is generally set during SendControllerEvents or Tick and is only valid on the game thread.
+ *
+ * @deprecated FInputDeviceScope is being superseded by FInputDeviceRegistry, which provides the same
+ * hardware metadata in a thread-safe manner keyed on FInputDeviceId. All call sites should migrate to
+ * FInputDeviceRegistry::FindDescriptor(DeviceId) and FInputDeviceScope will be removed once the
+ * migration is complete.
  */
 class FInputDeviceScope
 {
@@ -205,11 +213,13 @@ public:
 	APPLICATIONCORE_API ~FInputDeviceScope();
 
 	/** Cannot be copied/moved */
+	PRAGMA_DISABLE_DEPRECATION_WARNINGS
 	FInputDeviceScope() = delete;
 	FInputDeviceScope(const FInputDeviceScope&) = delete;
 	FInputDeviceScope& operator=(const FInputDeviceScope&) = delete;
 	FInputDeviceScope(FInputDeviceScope&&) = delete;
 	FInputDeviceScope& operator=(FInputDeviceScope&&) = delete;
+	PRAGMA_ENABLE_DEPRECATION_WARNINGS
 
 	/** Returns the currently active InputDeviceScope. This is only valid to call on the game thread and may return null */
 	static APPLICATIONCORE_API const FInputDeviceScope* GetCurrent();
diff --git a/Engine/Source/Runtime/Engine/Private/GameFramework/InputDeviceSubsystem.cpp b/Engine/Source/Runtime/Engine/Private/GameFramework/InputDeviceSubsystem.cpp
--- a/Engine/Source/Runtime/Engine/Private/GameFramework/InputDeviceSubsystem.cpp
+++ b/Engine/Source/Runtime/Engine/Private/GameFramework/InputDeviceSubsystem.cpp
@@ -4,6 +4,7 @@
 #include "GameFramework/InputDeviceProperties.h"
 #include "Framework/Application/IInputProcessor.h"
 #include "Framework/Application/SlateApplication.h"			// For RegisterInputPreProcessor
+#include "GenericPlatform/InputDeviceRegistry.h"
 #include "HAL/IConsoleManager.h"
 #include "Misc/App.h"	// For FApp::GetDeltaTime()
 
@@ -48,19 +49,16 @@ class FInputDeviceSubsystemProcessor : public IInputProcessor
 		
 		if (UInputDeviceSubsystem* SubSystem = UInputDeviceSubsystem::Get())
 		{
-			if (const FInputDeviceScope* Scope = FInputDeviceScope::GetCurrent())
+			if (const TOptional<FInputDeviceDescriptor> Desc = FInputDeviceRegistry::FindDescriptor(InDeviceId))
 			{
-				// TODO: Refactor FInputDeviceScope to use FName's instead of a FString for HardwareDeviceIdentifier
-				// Look up a hardware device ID from the config that matches this identifier
-				const FName DeviceScopeName(*Scope->HardwareDeviceIdentifier);
-				if (const FHardwareDeviceIdentifier* Hardware = UInputPlatformSettings::Get()->GetHardwareDeviceForClassName(DeviceScopeName))
+				if (const FHardwareDeviceIdentifier* Hardware = UInputPlatformSettings::Get()->GetHardwareDeviceForClassName(Desc->HardwareDeviceIdentifier))
 				{
 					SubSystem->SetMostRecentlyUsedHardwareDevice(InDeviceId, *Hardware);
 				}
 				// If there isn't one specified in the config file, we can just use the device name and scope name to make one.
 				else
 				{
-					SubSystem->SetMostRecentlyUsedHardwareDevice(InDeviceId, { Scope->InputDeviceName, DeviceScopeName });
+					SubSystem->SetMostRecentlyUsedHardwareDevice(InDeviceId, { Desc->InputDeviceName, Desc->HardwareDeviceIdentifier });
 				}
 			}
 			// If there isn't a recent input device scope, then we can check if the key was from a keyboard and mouse
@@ -571,18 +569,21 @@ void UInputDeviceSubsystem::HandleDeviceConnectionChange(EInputDeviceConnectionS
 		return;
 	}
 
-	// Not all systems will have a device scope available right away upon connection yet, and we need that to proceed 
-	const FInputDeviceScope* Scope = FInputDeviceScope::GetCurrent();
-	if (!Scope)
+	// Not all systems will have a device descriptor available right away upon connection yet, and we need that to proceed
+	const TOptional<FInputDeviceDescriptor> Desc = FInputDeviceRegistry::FindDescriptor(DeviceId);
+	if (!Desc)
 	{
 		return;
 	}
-	
-	const FName DeviceScopeName(*Scope->HardwareDeviceIdentifier);
-	if (const FHardwareDeviceIdentifier* Hardware = UInputPlatformSettings::Get()->GetHardwareDeviceForClassName(DeviceScopeName))
+
+	if (const FHardwareDeviceIdentifier* Hardware = UInputPlatformSettings::Get()->GetHardwareDeviceForClassName(Desc->HardwareDeviceIdentifier))
 	{
 		SetMostRecentlyUsedHardwareDevice(DeviceId, *Hardware);
 	}
+	else
+	{
+		SetMostRecentlyUsedHardwareDevice(DeviceId, { Desc->InputDeviceName, Desc->HardwareDeviceIdentifier });
+	}
 }
 
 void UInputDeviceSubsystem::SetMostRecentlyUsedHardwareDevice(const FInputDeviceId InDeviceId, const FHardwareDeviceIdentifier& InHardwareId)
