Clean up Linux splash screen.

#rnx
#Synced-CL 52211267

diff --git a/Engine/Source/Runtime/ApplicationCore/Private/Linux/LinuxPlatformSplash.cpp b/Engine/Source/Runtime/ApplicationCore/Private/Linux/LinuxPlatformSplash.cpp
--- a/Engine/Source/Runtime/ApplicationCore/Private/Linux/LinuxPlatformSplash.cpp
+++ b/Engine/Source/Runtime/ApplicationCore/Private/Linux/LinuxPlatformSplash.cpp
@@ -93,10 +93,6 @@ static FLinuxSplashState *GSplashState = nullptr;
 FLinuxSplashState::~FLinuxSplashState()
 {
 #if WITH_EDITOR
-	// Just in case SDL's renderer steps on GL state...
-	SDL_Window *CurrentWindow = SDL_GL_GetCurrentWindow();
-	SDL_GLContext CurrentContext = SDL_GL_GetCurrentContext();
-
 	if (SplashTexture)
 	{
 		SDL_DestroyTexture(SplashTexture);
@@ -143,11 +139,6 @@ FLinuxSplashState::~FLinuxSplashState()
 	{
 		FT_Done_FreeType(FontLibrary);
 	}
-
-	if (CurrentWindow && CurrentWindow != SplashWindow)  // put back any old GL state...
-	{
-		SDL_GL_MakeCurrent(CurrentWindow, CurrentContext);
-	}
 #endif // WITH_EDITOR
 
 	// do not deinit SDL here
@@ -476,15 +467,9 @@ bool FLinuxSplashState::InitSplashResources(const FText &AppName, const FString
 	if (SplashWidth <= 0 || SplashHeight <= 0)
 	{
 		UE_LOG(LogHAL, Warning, TEXT("Invalid splash image dimensions."));
-		return -1;
+		return false;
 	}
 
-#if WITH_EDITOR
-	// Just in case SDL's renderer steps on GL state...
-	SDL_Window *CurrentWindow = SDL_GL_GetCurrentWindow();
-	SDL_GLContext CurrentContext = SDL_GL_GetCurrentContext();
-#endif
-
 	// on modern X11, your windows might turn gray if they don't pump the event queue fast enough.
 	// But this is because they opt-in to an optional window manager protocol by default; legacy
 	// apps and those that know they'll be slow to respond to events--like splash screens--can
@@ -615,11 +600,6 @@ bool FLinuxSplashState::InitSplashResources(const FText &AppName, const FString
 	SDL_ShowWindow(SplashWindow);
 	Redraw();
 
-	if (CurrentWindow)  // put back any old GL state...
-	{
-		SDL_GL_MakeCurrent(CurrentWindow, CurrentContext);
-	}
-
 #else
 
 	SDL_Surface *WindowSurface = SDL_GetWindowSurface( SplashWindow );
@@ -630,7 +610,6 @@ bool FLinuxSplashState::InitSplashResources(const FText &AppName, const FString
 	}
 
 	SDL_BlitSurface(SplashSurface, NULL, WindowSurface, NULL);
-	SDL_DestroySurface(WindowSurface);
 
 	SDL_ShowWindow(SplashWindow);
 	Pump();
@@ -665,19 +644,10 @@ void FLinuxSplashState::Redraw()
 #if WITH_EDITOR
 	if (bNeedsRedraw || bStringsChanged)
 	{
-		// Just in case SDL's renderer steps on GL state...
-		SDL_Window *CurrentWindow = SDL_GL_GetCurrentWindow();
-		SDL_GLContext CurrentContext = SDL_GL_GetCurrentContext();
-
 		RenderStrings();
 		SDL_RenderTexture(SplashRenderer, SplashTexture, NULL, NULL);
 		SDL_RenderPresent(SplashRenderer);
 		bNeedsRedraw = false;
-
-		if (CurrentWindow)  // put back any old GL state...
-		{
-			SDL_GL_MakeCurrent(CurrentWindow, CurrentContext);
-		}
 	}
 #else
 	if (bNeedsRedraw)
