Don't warn on TSF error when cleaning up window focus associations

The warning was added as part of CL52085173 to detect TSF call errors, but it turns out that this TSF call can sometimes fail when popup menus are involved.
In some cases the focus changes trigger IME deactivation after the menu has been closed but before the matching FGenericWindow gets deleted, resulting in Hwnd being invalid.

The whole KnownWindows looping is a hack anyway, and the AssociateFocus result was likely left unchecked for that reason.
Let it fail silently, and the entry will be cleaned up from KnownWindows before the next loop.

#rnx
#Synced-CL 52253472

diff --git a/Engine/Source/Runtime/ApplicationCore/Private/Windows/WindowsTextInputMethodSystem.cpp b/Engine/Source/Runtime/ApplicationCore/Private/Windows/WindowsTextInputMethodSystem.cpp
--- a/Engine/Source/Runtime/ApplicationCore/Private/Windows/WindowsTextInputMethodSystem.cpp
+++ b/Engine/Source/Runtime/ApplicationCore/Private/Windows/WindowsTextInputMethodSystem.cpp
@@ -916,13 +916,7 @@ void FWindowsTextInputMethodSystem::DeactivateContext(const TSharedRef<ITextInpu
 				{
 					// TSF Implementation
 					ITfDocumentMgr* Unused;
-					Result = TSFThreadManager->AssociateFocus(Hwnd, TSFDisabledDocumentManager, &Unused);
-					if (FAILED(Result))
-					{
-						TCHAR ErrorMsg[1024];
-						FPlatformMisc::GetSystemErrorMessage(ErrorMsg, 1024, Result);
-						UE_LOG(LogWindowsTextInputMethodSystem, Error, TEXT("Deactivating a context failed while associating focus on the disabled TSF document manager. %s (0x%08x)"), ErrorMsg, Result);
-					}
+					TSFThreadManager->AssociateFocus(Hwnd, TSFDisabledDocumentManager, &Unused);
 
 					// IMM Implementation
 					::ImmAssociateContext(Hwnd, nullptr);
