[iOS/Mac] Fix virtual key codes triggered for Gamepad left and right Thumbstick on iOS and Mac.
- fix inversion of left/right gamepad keys for left stick.
- fix right stick logic to check proper gamepad key.

#tests Tested on iOS device.
#rnx

diff --git a/Engine/Source/Runtime/ApplicationCore/Private/Apple/AppleControllerInterface.cpp b/Engine/Source/Runtime/ApplicationCore/Private/Apple/AppleControllerInterface.cpp
--- a/Engine/Source/Runtime/ApplicationCore/Private/Apple/AppleControllerInterface.cpp
+++ b/Engine/Source/Runtime/ApplicationCore/Private/Apple/AppleControllerInterface.cpp
@@ -445,7 +445,7 @@ void FAppleControllerInterface::HandleDisconnect(GCController* Controller)
 					HandleAnalogGamepad(Pair.Key, ControllerIndex);
 				}
 			}
-			HandleVirtualButtonGamepad(FGamepadKeyNames::LeftStickRight, FGamepadKeyNames::LeftStickLeft, ControllerIndex);
+			HandleVirtualButtonGamepad(FGamepadKeyNames::LeftStickLeft, FGamepadKeyNames::LeftStickRight, ControllerIndex);
 			HandleVirtualButtonGamepad(FGamepadKeyNames::LeftStickDown, FGamepadKeyNames::LeftStickUp, ControllerIndex);
 			HandleVirtualButtonGamepad(FGamepadKeyNames::RightStickLeft, FGamepadKeyNames::RightStickRight, ControllerIndex);
 			HandleVirtualButtonGamepad(FGamepadKeyNames::RightStickDown, FGamepadKeyNames::RightStickUp, ControllerIndex);
