Port crvars commandline override from Android to iOS ConfigRules

#rnx

#Synced-CL 52432871

diff --git a/Engine/Source/Runtime/ApplicationCore/Private/IOS/IOSConfigRules.cpp b/Engine/Source/Runtime/ApplicationCore/Private/IOS/IOSConfigRules.cpp
--- a/Engine/Source/Runtime/ApplicationCore/Private/IOS/IOSConfigRules.cpp
+++ b/Engine/Source/Runtime/ApplicationCore/Private/IOS/IOSConfigRules.cpp
@@ -1,7 +1,9 @@
 // Copyright Epic Games, Inc. All Rights Reserved.
 
 #include "IOS/IOSConfigRules.h"
+#include "Misc/CommandLine.h"
 #include "Misc/FileHelper.h"
+#include "Misc/Parse.h"
 #include "zlib.h"
 #include <CommonCrypto/CommonCryptor.h>
 #include <CommonCrypto/CommonKeyDerivation.h>
@@ -201,6 +203,37 @@ bool FIOSConfigRules::Init(TMap<FString,FString>&& PredefinedVariables)
 	{
 		ConfigRuleVariablesMap = PredefinedVariables;
 	}
+
+#if !UE_BUILD_SHIPPING
+	// Allow commandline overrides of config rule variables
+	FString CrVarsValue;
+	if (FParse::Value(FCommandLine::Get(), TEXT("crvars="), CrVarsValue, false))
+	{
+		TArray<FString> Pairs;
+		CrVarsValue.ParseIntoArray(Pairs, TEXT(","));
+		for (const FString& Pair : Pairs)
+		{
+			FString Key, Value;
+			if (Pair.Split(TEXT("="), &Key, &Value))
+			{
+				Value.TrimQuotesInline();
+				CONF_LOG(LogConfigRules, Log, TEXT("crvars override: set %s = %s"), *Key, *Value);
+				ConfigRuleVariablesMap.Add(Key, Value);
+			}
+			else if (Pair.StartsWith(TEXT("-")))
+			{
+				Key = Pair.RightChop(1);
+				CONF_LOG(LogConfigRules, Log, TEXT("crvars override: remove %s"), *Key);
+				ConfigRuleVariablesMap.Remove(Key);
+			}
+			else
+			{
+				CONF_LOG(LogConfigRules, Log, TEXT("crvars ignoring malformed parameter: %s"), *Pair);
+			}
+		}
+	}
+#endif
+
 	return FoundVersion > -1;
 }
 
