Remove mouse cursor optimization on Windows to fix editor inconsistencies

#Synced-CL 52463319

diff --git a/Engine/Source/Runtime/ApplicationCore/Private/Windows/WindowsApplication.cpp b/Engine/Source/Runtime/ApplicationCore/Private/Windows/WindowsApplication.cpp
--- a/Engine/Source/Runtime/ApplicationCore/Private/Windows/WindowsApplication.cpp
+++ b/Engine/Source/Runtime/ApplicationCore/Private/Windows/WindowsApplication.cpp
@@ -1488,24 +1488,13 @@ FWindowsApplication::~FWindowsApplication()
 		TextInputMethodSystem->Terminate();
 	}
 
-	if (bAsyncCursorTaskActive)
-	{
-		const UE::FEventCountToken Token = AsyncCursorEventCount.PrepareWait();
-		if (bAsyncCursorTaskActive)
-		{
-			AsyncCursorEventCount.Wait(Token);
-		}
-	}
-
-	AsyncCursorTask.Wait();
-
 	::CoUninitialize();
 	OleUninitialize();
 }
 
 TSharedRef< FGenericWindow > FWindowsApplication::MakeWindow() 
 { 
-	return FWindowsWindow::Make(); 
+	return FWindowsWindow::Make();
 }
 
 void FWindowsApplication::InitializeWindow( const TSharedRef< FGenericWindow >& InWindow, const TSharedRef< FGenericWindowDefinition >& InDefinition, const TSharedPtr< FGenericWindow >& InParent, const bool bShowImmediately )
@@ -1611,54 +1600,17 @@ bool FWindowsApplication::IsCursorDirectlyOverSlateWindow() const
 {
 	QUICK_SCOPE_CYCLE_COUNTER(STAT_IsCursorDirectlyOverSlateWindow);
 
-	FRWScopeLock Lock(AsyncCursorRWLock, SLT_ReadOnly);
-	if (bAsyncCursorPosValid)
+	POINT CursorPos;
+	if (::GetCursorPos(&CursorPos))
 	{
-		TSharedPtr< FWindowsWindow > SlateWindowUnderCursor = FindWindowByHWND(Windows, AsyncCursorWindowHandle);
+		const HWND hWnd = ::WindowFromPoint(CursorPos);
+		const TSharedPtr<FWindowsWindow> SlateWindowUnderCursor = FindWindowByHWND(Windows, hWnd);
 		return SlateWindowUnderCursor.IsValid();
 	}
 
 	return false;
 }
 
-bool FWindowsApplication::StaticGetCursorPosition(POINT& OutCursorPos)
-{
-	return WindowsApplication ? WindowsApplication->GetCursorPosition(OutCursorPos) : false;
-}
-
-bool FWindowsApplication::GetCursorPosition(POINT& OutCursorPos)
-{
-	FRWScopeLock Lock(AsyncCursorRWLock, SLT_ReadOnly);
-	OutCursorPos = AsyncCursorPos;
-	return bAsyncCursorPosValid;
-}
-
-void FWindowsApplication::StaticForceUpdateAsyncCursor()
-{
-	if (WindowsApplication)
-	{
-		WindowsApplication->ForceUpdateAsyncCursor();
-	}
-}
-
-void FWindowsApplication::ForceUpdateAsyncCursor()
-{
-	// Querying outside of the lock can cause reordering:
-	// - Query for cursor info here
-	// - Async task queries for cursor info and updates members under lock
-	// - Members are updated here under lock with stale info
-	// The reason we don't query under lock here is that Win32 calls from the game thread can enter our WndProc and access AsyncCursorRWLock recursively and hang.
-	// We could add lock tracking to avoid grabbing the recursive lock if allowing a little reordering turns out to be fatal.
-	POINT CursorPos;
-	const BOOL bCursorPosValid = ::GetCursorPos(&CursorPos);
-	const HWND CursorWindowHandle = bCursorPosValid ? ::WindowFromPoint(CursorPos) : NULL;
-
-	FRWScopeLock Lock(AsyncCursorRWLock, SLT_Write);
-	bAsyncCursorPosValid = bCursorPosValid;
-	AsyncCursorPos = CursorPos;
-	AsyncCursorWindowHandle = CursorWindowHandle;
-}
-
 void FWindowsApplication::SetCapture( const TSharedPtr< FGenericWindow >& InWindow )
 {
 	HWND DesiredCaptureHandle = (InWindow.IsValid() ? (HWND)InWindow->GetOSWindowHandle() : NULL);
@@ -4719,7 +4671,6 @@ void FWindowsApplication::PostPumpMessages()
 {
 	WindowsApplication->UpdateHighPrecisionMouseMode();
 	WindowsApplication->ProcessWorkerInputs();
-	WindowsApplication->CacheCursorInfo();
 }
 
 void FWindowsApplication::ProcessRawMouseInput(const RAWMOUSE& Mouse, RawMouseInputResult& OutResult)
@@ -4910,23 +4861,20 @@ PRAGMA_ENABLE_DEPRECATION_WARNINGS
 			TSharedPtr< FWindowsWindow > Window;
 
 			POINT ButtonCursorPoint;
-
+			if (::GetCursorPos(&ButtonCursorPoint))
 			{
-				FRWScopeLock Lock(AsyncCursorRWLock, SLT_ReadOnly);
-				if (bAsyncCursorPosValid)
+				const HWND hWnd = ::WindowFromPoint(ButtonCursorPoint);
+				if (hWnd)
 				{
-					ButtonCursorPoint = AsyncCursorPos;
-					if (AsyncCursorWindowHandle)
-					{
-						Window = FindWindowByHWND(Windows, AsyncCursorWindowHandle);
-					}
-				}
-				else
-				{
-					ButtonCursorPoint.x = 0;
-					ButtonCursorPoint.y = 0;
+					Window = FindWindowByHWND(Windows, hWnd);
 				}
 			}
+			else
+			{
+				UE_LOG(LogWindowsDesktop, Error, TEXT("Could not GetCursorPos for raw input. Error code: %u"), GetLastError());
+				ButtonCursorPoint.x = 0;
+				ButtonCursorPoint.y = 0;
+			}
 
 			const FVector2D ButtonCursorPos = FVector2D(ButtonCursorPoint.x, ButtonCursorPoint.y);
 			const double CurrentTime = FPlatformTime::Seconds();
@@ -4973,43 +4921,6 @@ void FWindowsApplication::ProcessMouseButton(double CurrentTime, const TSharedPt
 	}
 }
 
-void FWindowsApplication::CacheCursorInfo()
-{
-	TRACE_CPUPROFILER_EVENT_SCOPE(CacheCursorInfo);
-
-	if (bAsyncCursorTaskActive)
-	{
-		const UE::FEventCountToken Token = AsyncCursorEventCount.PrepareWait();
-		if (bAsyncCursorTaskActive)
-		{
-			AsyncCursorEventCount.Wait(Token);
-		}
-	}
-
-	bAsyncCursorTaskActive = true;
-
-	AsyncCursorTask = UE::Tasks::Launch(UE_SOURCE_LOCATION,
-		[&CursorRWLock = AsyncCursorRWLock,
-		&bCursorPosValid = bAsyncCursorPosValid,
-		&CursorPos = AsyncCursorPos,
-		&CursorWindowHandle = AsyncCursorWindowHandle,
-		&bCursorTaskActive = bAsyncCursorTaskActive,
-		&CursorEventCount = AsyncCursorEventCount]
-		{
-			TRACE_CPUPROFILER_EVENT_SCOPE(AsyncCursorTask);
-
-			{
-				FRWScopeLock Lock(CursorRWLock, SLT_Write);
-				bCursorPosValid = ::GetCursorPos(&CursorPos);
-				CursorWindowHandle = bCursorPosValid ? ::WindowFromPoint(CursorPos) : NULL;
-			}
-
-			bCursorTaskActive = false;
-			CursorEventCount.Notify();
-		}
-	);
-}
-
 // Restore the windowsx.h macro for IsMaximized
 #pragma pop_macro("IsMaximized")
 
diff --git a/Engine/Source/Runtime/ApplicationCore/Private/Windows/WindowsCursor.cpp b/Engine/Source/Runtime/ApplicationCore/Private/Windows/WindowsCursor.cpp
--- a/Engine/Source/Runtime/ApplicationCore/Private/Windows/WindowsCursor.cpp
+++ b/Engine/Source/Runtime/ApplicationCore/Private/Windows/WindowsCursor.cpp
@@ -7,7 +7,6 @@
 #include "Misc/CoreMisc.h"
 #include "Math/Vector2D.h"
 #include "Math/Color.h"
-#include "Windows/WindowsApplication.h"
 #include "Windows/WindowsWindow.h"
 #include "HAL/IConsoleManager.h"
 #include "HAL/PlatformProcess.h"
@@ -230,21 +229,13 @@ void* FWindowsCursor::CreateCursorFromRGBABuffer(const FColor* Pixels, int32 Wid
 
 FVector2D FWindowsCursor::GetPosition() const
 {
-	POINT OutCursorPos;
-	if (!FWindowsApplication::StaticGetCursorPosition(OutCursorPos))
-	{
-		::GetCursorPos(&OutCursorPos);
-	}
-
-	return FVector2D(OutCursorPos.x, OutCursorPos.y);
+	POINT CursorPos;
+	return ::GetCursorPos(&CursorPos) ? FVector2D(CursorPos.x, CursorPos.y) : FVector2D::ZeroVector;
 }
 
 void FWindowsCursor::SetPosition(const int32 X, const int32 Y)
 {
-	if (::SetCursorPos(X, Y))
-	{
-		FWindowsApplication::StaticForceUpdateAsyncCursor();
-	}
+	::SetCursorPos(X, Y);
 }
 
 void FWindowsCursor::SetType( const EMouseCursor::Type InNewCursor )
diff --git a/Engine/Source/Runtime/ApplicationCore/Public/Windows/WindowsApplication.h b/Engine/Source/Runtime/ApplicationCore/Public/Windows/WindowsApplication.h
--- a/Engine/Source/Runtime/ApplicationCore/Public/Windows/WindowsApplication.h
+++ b/Engine/Source/Runtime/ApplicationCore/Public/Windows/WindowsApplication.h
@@ -8,7 +8,6 @@
 #include "HAL/IConsoleManager.h"
 #include "GenericPlatform/GenericApplication.h"
 #include "GenericPlatform/IInputInterface.h"
-#include "Tasks/Task.h"
 #include "Windows/AllowWindowsPlatformTypes.h"
 	#include <Ole2.h>
 	#include <oleidl.h>
@@ -316,9 +315,6 @@ public:
 	/** Process raw mouse inputs gathered from a helper thread/window. */
 	void ProcessWorkerInputs();
 
-	/** Cache Win32 calls that can take a while. */
-	void CacheCursorInfo();
-
 	/** Virtual destructor. */
 	APPLICATIONCORE_API virtual ~FWindowsApplication();
 
@@ -380,10 +376,6 @@ public:
 	APPLICATIONCORE_API virtual bool IsGamepadAttached() const override;
 	APPLICATIONCORE_API virtual FModifierKeysState GetModifierKeys() const override;
 	APPLICATIONCORE_API virtual bool IsCursorDirectlyOverSlateWindow() const override;
-	static bool StaticGetCursorPosition(POINT& OutSursorPos);
-	bool GetCursorPosition(POINT& OutSursorPos);
-	static void StaticForceUpdateAsyncCursor();
-	void ForceUpdateAsyncCursor();
 	APPLICATIONCORE_API virtual FPlatformRect GetWorkArea( const FPlatformRect& CurrentWindow ) const override;
 	APPLICATIONCORE_API virtual void GetInitialDisplayMetrics( FDisplayMetrics& OutDisplayMetrics ) const override;
 	APPLICATIONCORE_API virtual EWindowTitleAlignment::Type GetWindowTitleAlignment() const override;
@@ -542,14 +534,6 @@ private:
 	HWND HighPrecisionMouseInputWindowHandle;
 	HWND CaptureWindowHandle;
 
-	UE::FEventCount AsyncCursorEventCount;
-	std::atomic<bool> bAsyncCursorTaskActive = false;
-	UE::Tasks::FTask AsyncCursorTask;
-	mutable FRWLock AsyncCursorRWLock;
-	BOOL bAsyncCursorPosValid{};
-	POINT AsyncCursorPos{};
-	HWND AsyncCursorWindowHandle = NULL;
-
 	bool bIsMouseAttached;
 
 	bool bForceActivateByMouse;
