Native iOS gesture recognizer integration.
Introduced EGesturePhase parameter in FGenericApplicationMessageHandler::OnTouchGesture().
Added CVar "Input.EnableNativeGestureRecognition" to enable/disable this feature. The default value is false. It's set to true on iOS.

#Synced-CL 52530459

diff --git a/Engine/Config/IOS/BaseIOSEngine.ini b/Engine/Config/IOS/BaseIOSEngine.ini
--- a/Engine/Config/IOS/BaseIOSEngine.ini
+++ b/Engine/Config/IOS/BaseIOSEngine.ini
@@ -66,6 +66,7 @@ CheckForForegroundStaleDownloadsWithInterval=1
 ForegroundStaleDownloadTimeout=60
 
 [ConsoleVariables]
+Input.EnableNativeGestureRecognition=true
 Slate.CacheRenderData=0
 r.HairStrands.Simulation=0
 r.HairStrands.Strands=0
diff --git a/Engine/Source/Editor/TextureEditor/Private/Models/TextureEditorViewportClient.cpp b/Engine/Source/Editor/TextureEditor/Private/Models/TextureEditorViewportClient.cpp
--- a/Engine/Source/Editor/TextureEditor/Private/Models/TextureEditorViewportClient.cpp
+++ b/Engine/Source/Editor/TextureEditor/Private/Models/TextureEditorViewportClient.cpp
@@ -626,7 +626,7 @@ EMouseCursor::Type FTextureEditorViewportClient::GetCursor(FViewport* Viewport,
 	return ShouldUseMousePanning(Viewport) ? EMouseCursor::GrabHandClosed : EMouseCursor::Default;
 }
 
-bool FTextureEditorViewportClient::InputGesture(FViewport* Viewport, const FInputDeviceId DeviceId, EGestureEvent GestureType, const FVector2D& GestureDelta, bool bIsDirectionInvertedFromDevice, const uint64 Timestamp)
+bool FTextureEditorViewportClient::InputGesture(FViewport* Viewport, const FInputDeviceId DeviceId, EGestureEvent GestureType, EGesturePhase GesturePhase, const FVector2D& GestureDelta, bool bIsDirectionInvertedFromDevice, const uint64 Timestamp)
 {
 	const bool LeftMouseButtonDown = Viewport->KeyState(EKeys::LeftMouseButton);
 	const bool RightMouseButtonDown = Viewport->KeyState(EKeys::RightMouseButton);
diff --git a/Engine/Source/Editor/TextureEditor/Private/Models/TextureEditorViewportClient.h b/Engine/Source/Editor/TextureEditor/Private/Models/TextureEditorViewportClient.h
--- a/Engine/Source/Editor/TextureEditor/Private/Models/TextureEditorViewportClient.h
+++ b/Engine/Source/Editor/TextureEditor/Private/Models/TextureEditorViewportClient.h
@@ -26,7 +26,7 @@ public:
 	virtual void Draw(FViewport* Viewport, FCanvas* Canvas) override;
 	virtual bool InputKey(const FInputKeyEventArgs& InEventArgs) override;
 	virtual bool InputAxis(const FInputKeyEventArgs& InEventArgs) override;
-	virtual bool InputGesture(FViewport* Viewport, const FInputDeviceId DeviceId, EGestureEvent GestureType, const FVector2D& GestureDelta, bool bIsDirectionInvertedFromDevice, const uint64 Timestamp) override;
+	virtual bool InputGesture(FViewport* Viewport, const FInputDeviceId DeviceId, EGestureEvent GestureType, EGesturePhase GesturePhase, const FVector2D& GestureDelta, bool bIsDirectionInvertedFromDevice, const uint64 Timestamp) override;
 	virtual UWorld* GetWorld() const override { return nullptr; }
 	virtual EMouseCursor::Type GetCursor(FViewport* Viewport, int32 X, int32 Y) override;
 
diff --git a/Engine/Source/Editor/UnrealEd/Private/EditorViewportClient.cpp b/Engine/Source/Editor/UnrealEd/Private/EditorViewportClient.cpp
--- a/Engine/Source/Editor/UnrealEd/Private/EditorViewportClient.cpp
+++ b/Engine/Source/Editor/UnrealEd/Private/EditorViewportClient.cpp
@@ -5400,7 +5400,7 @@ static float AdjustGestureCameraRotation(float Delta, float AdjustLimit, float D
 	return bIsUsingTrackpad ? Delta : FMath::Clamp(Delta, -DeltaCutoff, DeltaCutoff);
 }
 
-bool FEditorViewportClient::InputGesture(FViewport* InViewport, const FInputDeviceId DeviceId, EGestureEvent GestureType, const FVector2D& GestureDelta, bool bIsDirectionInvertedFromDevice, const uint64 Timestamp)
+bool FEditorViewportClient::InputGesture(FViewport* InViewport, const FInputDeviceId DeviceId, EGestureEvent GestureType, EGesturePhase GesturePhase, const FVector2D& GestureDelta, bool bIsDirectionInvertedFromDevice, const uint64 Timestamp)
 {
 	if (bDisableInput)
 	{
diff --git a/Engine/Source/Editor/UnrealEd/Public/EditorViewportClient.h b/Engine/Source/Editor/UnrealEd/Public/EditorViewportClient.h
--- a/Engine/Source/Editor/UnrealEd/Public/EditorViewportClient.h
+++ b/Engine/Source/Editor/UnrealEd/Public/EditorViewportClient.h
@@ -626,7 +626,7 @@ public:
 	UNREALED_API virtual void RequestInvalidateHitProxy(FViewport* Viewport) override;
 	UNREALED_API virtual bool InputKey(const FInputKeyEventArgs& EventArgs) override;
 	UNREALED_API virtual bool InputAxis(const FInputKeyEventArgs& Args) override;
-	UNREALED_API virtual bool InputGesture(FViewport* Viewport, const FInputDeviceId DeviceId, EGestureEvent GestureType, const FVector2D& GestureDelta, bool bIsDirectionInvertedFromDevice, const uint64 Timestamp) override;
+	UNREALED_API virtual bool InputGesture(FViewport* Viewport, const FInputDeviceId DeviceId, EGestureEvent GestureType, EGesturePhase GesturePhase, const FVector2D& GestureDelta, bool bIsDirectionInvertedFromDevice, const uint64 Timestamp) override;
 	UNREALED_API virtual void ReceivedFocus(FViewport* Viewport) override;
 	UNREALED_API virtual void MouseEnter(FViewport* Viewport,int32 x, int32 y) override;
 	UNREALED_API virtual void MouseMove(FViewport* Viewport,int32 x, int32 y) override;
diff --git a/Engine/Source/Runtime/ApplicationCore/Private/GenericPlatform/GenericPlatformApplicationMisc.cpp b/Engine/Source/Runtime/ApplicationCore/Private/GenericPlatform/GenericPlatformApplicationMisc.cpp
--- a/Engine/Source/Runtime/ApplicationCore/Private/GenericPlatform/GenericPlatformApplicationMisc.cpp
+++ b/Engine/Source/Runtime/ApplicationCore/Private/GenericPlatform/GenericPlatformApplicationMisc.cpp
@@ -40,6 +40,18 @@ FAutoConsoleVariableRef FGenericPlatformApplicationMisc::CVarAllowVirtualKeyboar
 	ECVF_Default
 );
 
+TAutoConsoleVariable<bool> FGenericPlatformApplicationMisc::CVarEnableNativeGestureRecognition(
+	TEXT("Input.EnableNativeGestureRecognition"),
+	false,
+	TEXT("Use native platform gesture recognizers instead of the software gesture detection. Enabled via platform ini on supported platforms."),
+	ECVF_Default
+);
+
+bool FGenericPlatformApplicationMisc::AreNativeGesturesEnabled()
+{
+	return CVarEnableNativeGestureRecognition.GetValueOnAnyThread();
+}
+
 void FGenericPlatformApplicationMisc::PreInit()
 {
 }
diff --git a/Engine/Source/Runtime/ApplicationCore/Private/IOS/IOSInputInterface.cpp b/Engine/Source/Runtime/ApplicationCore/Private/IOS/IOSInputInterface.cpp
--- a/Engine/Source/Runtime/ApplicationCore/Private/IOS/IOSInputInterface.cpp
+++ b/Engine/Source/Runtime/ApplicationCore/Private/IOS/IOSInputInterface.cpp
@@ -3,6 +3,8 @@
 #include "IOS/IOSInputInterface.h"
 #include "IOS/IOSAppDelegate.h"
 #include "IOS/IOSApplication.h"
+#include "GenericPlatform/GenericPlatformInputDeviceMapper.h"
+#include "HAL/PlatformApplicationMisc.h"
 #include "Misc/ConfigCacheIni.h"
 #include "Misc/ScopeLock.h"
 #include "HAL/IConsoleManager.h"
@@ -15,8 +17,6 @@
 #include "Apple/AppleMouseController.h"
 #include "IOS/IOSInputDefinitions.h"
 
-DECLARE_LOG_CATEGORY_EXTERN(LogIOSInput, Log, All);
-
 #define APPLE_CONTROLLER_DEBUG 0
 
 static TAutoConsoleVariable<float> CVarHapticsKickHeavy(TEXT("ios.VibrationHapticsKickHeavyValue"), 0.65f, TEXT("Vibation values higher than this will kick a haptics heavy Impact"));
@@ -29,6 +29,8 @@ constexpr EIOSEventType operator+(EIOSEventType type, int Index) { return (EIOSE
 // protects the input stack used on 2 threads
 static FCriticalSection CriticalSection;
 static TArray<TouchInput> TouchInputStack;
+static TArray<GestureInput> GestureInputStack;
+static TArray<GestureInput> PendingDiscreteGestureReleases;
 
 FName FIOSInputInterface::InputClassName_DefaultMobileTouch;
 FString FIOSInputInterface::HardwareDeviceIdentifier_DefaultMobileTouch;
@@ -227,6 +229,58 @@ PRAGMA_ENABLE_DEPRECATION_WARNINGS
     }
 }
 
+void FIOSInputInterface::ProcessGestures(uint32 ControllerId, const TArray<GestureInput>& InGestureInputStack)
+{
+	// Remap the controller ID to FInputDeviceId, matching ProcessTouchesAndKeys.
+	FPlatformUserId UserId = FGenericPlatformMisc::GetPlatformUserForUserIndex(ControllerId);
+	FInputDeviceId DeviceId = INPUTDEVICEID_NONE;
+	IPlatformInputDeviceMapper::Get().RemapControllerIdToPlatformUserAndDevice(ControllerId, UserId, DeviceId);
+
+	// Inject deferred releases for discrete gestures (Tap, Flick) from the previous
+	// frame so EvaluateKeyMapState had one full tick to read the non-zero press value.
+	for (const GestureInput& Release : PendingDiscreteGestureReleases)
+	{
+		MessageHandler->OnNativeGestureEnd(Release.Type, Release.Delta, Release.Value, false, UserId, DeviceId);
+	}
+	PendingDiscreteGestureReleases.Reset();
+
+	for (const GestureInput& Gesture : InGestureInputStack)
+	{
+		if (Gesture.bIsStart)
+		{
+			if (Gesture.bIsEnd)
+			{
+				NSLog(@"Warning: Gesture event (type %d) has both bIsStart and bIsEnd set; treating as Begin.", (int)Gesture.Type);
+			}
+			MessageHandler->OnNativeGestureBegin(Gesture.Type, Gesture.Delta, Gesture.Value, false, UserId, DeviceId);
+		}
+		else if (Gesture.bIsEnd)
+		{
+			MessageHandler->OnNativeGestureEnd(Gesture.Type, Gesture.Delta, Gesture.Value, false, UserId, DeviceId);
+		}
+		else
+		{
+			// Treat gesture events with false values of bIsStart and bIsEnd as Update.
+			MessageHandler->OnNativeGestureUpdate(Gesture.Type, Gesture.Delta, Gesture.Value, false, UserId, DeviceId);
+		}
+
+		// For discrete (one-shot) gestures, defer the release to next frame.
+		// The press sets RawValue to the gesture value; the deferred release
+		// clears it after EvaluateKeyMapState has had one tick to read it.
+		const bool bIsDiscreteGesture = (Gesture.Type == EGestureEvent::Tap || Gesture.Type == EGestureEvent::Flick);
+		if (bIsDiscreteGesture && Gesture.bIsStart)
+		{
+			GestureInput DeferredRelease;
+			DeferredRelease.Type = Gesture.Type;
+			DeferredRelease.Delta = FVector2D::ZeroVector;
+			DeferredRelease.Value = 0.0f;
+			DeferredRelease.bIsStart = false;
+			DeferredRelease.bIsEnd = true;
+			PendingDiscreteGestureReleases.Add(DeferredRelease);
+		}
+	}
+}
+
 void FIOSInputInterface::SendControllerEvents()
 {
 	KeyboardController->SendControllerEvents();
@@ -243,6 +297,15 @@ void FIOSInputInterface::SendControllerEvents()
     ProcessTouchesAndKeys(0, LocalTouchInputStack);
 #endif
 
+	// Not sure if I should put TVOS as an exception for gesture handling.
+	// It won't be harmful to keep TVOS in the list - If without gesture support, anyways the input stack will be empty.
+	TArray<GestureInput> LocalGestureInputStack;
+    {
+        FScopeLock Lock(&CriticalSection);
+		Exchange(LocalGestureInputStack, GestureInputStack);
+    }
+	// On iOS, gestures always go to player 0, matching touch input above.
+	ProcessGestures(0, LocalGestureInputStack);
     
 #if !PLATFORM_TVOS // @todo tvos: This needs to come from the Microcontroller rotation
     if (!bPauseMotion)
@@ -351,6 +414,17 @@ void FIOSInputInterface::QueueTouchInput(const TArray<TouchInput>& InTouchEvents
     TouchInputStack.Append(InTouchEvents);
 }
 
+void FIOSInputInterface::QueueGestureInput(const TArray<GestureInput>& InGestureInputs)
+{
+	if (!FGenericPlatformApplicationMisc::AreNativeGesturesEnabled())
+	{
+		return;
+	}
+
+	FScopeLock Lock(&CriticalSection);
+	GestureInputStack.Append(InGestureInputs);
+}
+
 void FIOSInputInterface::EnableMotionData(bool bEnable)
 {
     bPauseMotion = !bEnable;
diff --git a/Engine/Source/Runtime/ApplicationCore/Private/IOS/IOSView.cpp b/Engine/Source/Runtime/ApplicationCore/Private/IOS/IOSView.cpp
--- a/Engine/Source/Runtime/ApplicationCore/Private/IOS/IOSView.cpp
+++ b/Engine/Source/Runtime/ApplicationCore/Private/IOS/IOSView.cpp
@@ -232,6 +232,10 @@ MTL::Device* GMetalDevice = nullptr;
 	
 #endif
 
+#if !PLATFORM_TVOS
+	[self SetupGestureRecognizers];
+#endif
+
 	return self;
 }
 
@@ -1111,6 +1115,179 @@ self.accessibilityElements = @[Window.accessibilityContainer];
 #endif
 }
 
+#if !PLATFORM_TVOS
+- (void)SetupGestureRecognizers
+{
+    // Pan — continuous single-finger drag
+    UIPanGestureRecognizer* Pan = [[UIPanGestureRecognizer alloc] initWithTarget:self action:@selector(HandlePanGesture:)];
+    Pan.minimumNumberOfTouches = 1;
+    Pan.maximumNumberOfTouches = 1;
+    Pan.delegate = self;
+    [self addGestureRecognizer:Pan];
+    [Pan release];
+
+    // Pinch — two-finger scale
+    UIPinchGestureRecognizer* Pinch = [[UIPinchGestureRecognizer alloc] initWithTarget:self action:@selector(HandlePinchGesture:)];
+    Pinch.delegate = self;
+    [self addGestureRecognizer:Pinch];
+    [Pinch release];
+
+    // Rotate — two-finger rotation
+    UIRotationGestureRecognizer* Rotation = [[UIRotationGestureRecognizer alloc] initWithTarget:self action:@selector(HandleRotateGesture:)];
+    Rotation.delegate = self;
+    [self addGestureRecognizer:Rotation];
+    [Rotation release];
+
+    // Tap — single tap
+    UITapGestureRecognizer* Tap = [[UITapGestureRecognizer alloc] initWithTarget:self action:@selector(HandleTapGesture:)];
+    Tap.numberOfTapsRequired = 1;
+    Tap.delegate = self;
+    [self addGestureRecognizer:Tap];
+    [Tap release];
+
+    // Long Press — continuous (Began → Changed → Ended)
+    UILongPressGestureRecognizer* LongPress = [[UILongPressGestureRecognizer alloc] initWithTarget:self action:@selector(HandleLongPressGesture:)];
+    LongPress.delegate = self;
+    [self addGestureRecognizer:LongPress];
+    [LongPress release];
+
+}
+
+// Allow gesture recognizers to work simultaneously with raw touches
+- (BOOL)gestureRecognizer:(UIGestureRecognizer*)gestureRecognizer
+    shouldRecognizeSimultaneouslyWithGestureRecognizer:(UIGestureRecognizer*)otherGestureRecognizer
+{
+    return YES;
+}
+
+- (void)HandlePanGesture:(UIPanGestureRecognizer*)Recognizer
+{
+    CGPoint Translation = [Recognizer translationInView:self];
+    // Apply same view size scaling as touch events
+    CGPoint ViewSizeModifier = CGPointMake(_ViewSize.width / self.frame.size.width,
+                                           _ViewSize.height / self.frame.size.height);
+
+    TArray<GestureInput> Gestures;
+    GestureInput G;
+    G.Type = EGestureEvent::Pan;
+    G.Delta = FVector2D(Translation.x * ViewSizeModifier.x, Translation.y * ViewSizeModifier.y);
+    G.Value = 0.0f;
+    G.bIsStart = (Recognizer.state == UIGestureRecognizerStateBegan);
+    G.bIsEnd   = (Recognizer.state == UIGestureRecognizerStateEnded ||
+                  Recognizer.state == UIGestureRecognizerStateCancelled);
+    Gestures.Add(G);
+
+    // Detect flick: if pan ended with sufficient velocity, emit a one-shot Flick gesture.
+    // Angle convention matches FGestureRecognizer: right=0, up=90, left=180, down=270.
+    if (G.bIsEnd)
+    {
+        CGPoint Velocity = [Recognizer velocityInView:self];
+        CGFloat SpeedSq = Velocity.x * Velocity.x + Velocity.y * Velocity.y;
+        const CGFloat FlickThresholdSq = 500.0 * 500.0; // 500 pts/s
+        if (SpeedSq > FlickThresholdSq)
+        {
+            float Angle = (float)FRotator::ClampAxis(FMath::RadiansToDegrees(FMath::Atan2(-Velocity.y, Velocity.x)));
+
+            GestureInput Flick;
+            Flick.Type = EGestureEvent::Flick;
+            Flick.Delta = FVector2D(Velocity.x * ViewSizeModifier.x, Velocity.y * ViewSizeModifier.y);
+            Flick.Value = Angle;
+            Flick.bIsStart = true;
+            Flick.bIsEnd = false;
+            Gestures.Add(Flick);
+        }
+    }
+
+    FIOSInputInterface::QueueGestureInput(Gestures);
+
+    // Reset translation so delta is per-callback, not cumulative
+    [Recognizer setTranslation:CGPointZero inView:self];
+}
+
+- (void)HandlePinchGesture:(UIPinchGestureRecognizer*)Recognizer
+{
+    TArray<GestureInput> Gestures;
+    GestureInput G;
+    G.Type = EGestureEvent::Magnify;
+    float ScaleDelta = (float)(Recognizer.scale - 1.0);  // Convert from 1-centered scale to 0-centered delta to match Mac's magnification values
+    G.Delta = FVector2D(ScaleDelta, 0.0f);
+    G.Value = ScaleDelta;
+    G.bIsStart = (Recognizer.state == UIGestureRecognizerStateBegan);
+    G.bIsEnd   = (Recognizer.state == UIGestureRecognizerStateEnded ||
+                  Recognizer.state == UIGestureRecognizerStateCancelled);
+    Gestures.Add(G);
+    FIOSInputInterface::QueueGestureInput(Gestures);
+
+    // Reset scale for delta-based reporting
+    Recognizer.scale = 1.0;
+}
+
+- (void)HandleRotateGesture:(UIRotationGestureRecognizer*)Recognizer
+{
+    TArray<GestureInput> Gestures;
+    GestureInput G;
+    G.Type = EGestureEvent::Rotate;
+    G.Delta = FVector2D((float)Recognizer.rotation, 0.0f);  // Encode rotation in Delta.X so it survives Slate routing
+    G.Value = (float)Recognizer.rotation;  // radians, positive = counter-clockwise
+    G.bIsStart = (Recognizer.state == UIGestureRecognizerStateBegan);
+    G.bIsEnd   = (Recognizer.state == UIGestureRecognizerStateEnded ||
+                  Recognizer.state == UIGestureRecognizerStateCancelled);
+    Gestures.Add(G);
+    FIOSInputInterface::QueueGestureInput(Gestures);
+
+    // Reset rotation for delta-based reporting
+    Recognizer.rotation = 0.0;
+}
+
+- (void)HandleTapGesture:(UITapGestureRecognizer*)Recognizer
+{
+	// UITapGestureRecognizer is discrete — it only fires once at StateRecognized.
+	// We queue only the press here; ProcessGestures() defers the release to the
+	// next frame so EvaluateKeyMapState has one tick to read the non-zero value.
+	if (Recognizer.state != UIGestureRecognizerStateRecognized)
+	{
+		return;
+	}
+
+	CGPoint Location = [Recognizer locationInView:self];
+	CGPoint ViewSizeModifier = CGPointMake(_ViewSize.width / self.frame.size.width,
+										   _ViewSize.height / self.frame.size.height);
+	FVector2D ScreenLocation(Location.x * ViewSizeModifier.x, Location.y * ViewSizeModifier.y);
+
+	GestureInput Tap;
+	Tap.Type = EGestureEvent::Tap;
+	Tap.Delta = ScreenLocation;
+	Tap.Value = 0.0f;
+	Tap.bIsStart = true;
+	Tap.bIsEnd = false;
+
+	TArray<GestureInput> Gestures;
+	Gestures.Add(Tap);
+	FIOSInputInterface::QueueGestureInput(Gestures);
+}
+
+- (void)HandleLongPressGesture:(UILongPressGestureRecognizer*)Recognizer
+{
+	CGPoint Location = [Recognizer locationInView:self];
+	CGPoint ViewSizeModifier = CGPointMake(_ViewSize.width / self.frame.size.width,
+										   _ViewSize.height / self.frame.size.height);
+	FVector2D ScreenLocation(Location.x * ViewSizeModifier.x, Location.y * ViewSizeModifier.y);
+
+	bool bIsEnd = (Recognizer.state == UIGestureRecognizerStateEnded ||
+				  Recognizer.state == UIGestureRecognizerStateCancelled);
+
+	GestureInput G;
+	G.Type = EGestureEvent::LongPress;
+	G.Delta = bIsEnd ? FVector2D::ZeroVector : ScreenLocation;
+	G.Value = 0.0f;
+	G.bIsStart = (Recognizer.state == UIGestureRecognizerStateBegan);
+	G.bIsEnd = bIsEnd;
+
+	TArray<GestureInput> Gestures;
+	Gestures.Add(G);
+	FIOSInputInterface::QueueGestureInput(Gestures);
+}
+#endif // !PLATFORM_TVOS
 
 @end
 
diff --git a/Engine/Source/Runtime/ApplicationCore/Private/Mac/MacApplication.cpp b/Engine/Source/Runtime/ApplicationCore/Private/Mac/MacApplication.cpp
--- a/Engine/Source/Runtime/ApplicationCore/Private/Mac/MacApplication.cpp
+++ b/Engine/Source/Runtime/ApplicationCore/Private/Mac/MacApplication.cpp
@@ -495,13 +495,17 @@ FPlatformRect FMacApplication::GetWorkArea(const FPlatformRect& CurrentWindow) c
 #if WITH_EDITOR
 void FMacApplication::SendAnalytics(IAnalyticsProvider* Provider)
 {
-	static_assert(((uint32)EGestureEvent::Count) == 6, "If the number of gestures changes you need to add more entries below!");
+	static_assert(((uint32)EGestureEvent::Count) == 9, "If the number of gestures changes you need to add more entries below!");
 
 	TArray<FAnalyticsEventAttribute> GestureAttributes;
 	GestureAttributes.Add(FAnalyticsEventAttribute(FString("Scroll"),	GestureUsage[(int32)EGestureEvent::Scroll]));
 	GestureAttributes.Add(FAnalyticsEventAttribute(FString("Magnify"),	GestureUsage[(int32)EGestureEvent::Magnify]));
 	GestureAttributes.Add(FAnalyticsEventAttribute(FString("Swipe"),	GestureUsage[(int32)EGestureEvent::Swipe]));
 	GestureAttributes.Add(FAnalyticsEventAttribute(FString("Rotate"),	GestureUsage[(int32)EGestureEvent::Rotate]));
+	GestureAttributes.Add(FAnalyticsEventAttribute(FString("Pan"),	GestureUsage[(int32)EGestureEvent::Pan]));
+	GestureAttributes.Add(FAnalyticsEventAttribute(FString("Tap"),	GestureUsage[(int32)EGestureEvent::Tap]));
+	GestureAttributes.Add(FAnalyticsEventAttribute(FString("LongPress"),	GestureUsage[(int32)EGestureEvent::LongPress]));
+	GestureAttributes.Add(FAnalyticsEventAttribute(FString("Flick"),	GestureUsage[(int32)EGestureEvent::Flick]));
 
 	Provider->RecordEvent(FString("Mac.Gesture.Usage"), GestureAttributes);
 
diff --git a/Engine/Source/Runtime/ApplicationCore/Public/GenericPlatform/GenericApplicationMessageHandler.h b/Engine/Source/Runtime/ApplicationCore/Public/GenericPlatform/GenericApplicationMessageHandler.h
--- a/Engine/Source/Runtime/ApplicationCore/Public/GenericPlatform/GenericApplicationMessageHandler.h
+++ b/Engine/Source/Runtime/ApplicationCore/Public/GenericPlatform/GenericApplicationMessageHandler.h
@@ -146,6 +146,14 @@ namespace EDropEffect
 }
 
 
+enum class EGesturePhase : uint8
+{
+	None,
+	Update,
+	Started,
+	Ended
+};
+
 enum class EGestureEvent : uint8
 {
 	None,
@@ -154,6 +162,9 @@ enum class EGestureEvent : uint8
 	Swipe,
 	Rotate,
 	LongPress,
+	Pan,
+	Tap,
+	Flick,
 	Count
 };
 
@@ -349,6 +360,21 @@ public:
     {
     }
 
+	virtual bool OnNativeGestureBegin(EGestureEvent GestureType, const FVector2D& Delta, float Value, bool bIsDirectionInvertedFromDevice, FPlatformUserId PlatformUserId, FInputDeviceId DeviceId)
+	{
+		return false;
+	}
+
+	virtual bool OnNativeGestureUpdate(EGestureEvent GestureType, const FVector2D& Delta, float Value, bool bIsDirectionInvertedFromDevice, FPlatformUserId PlatformUserId, FInputDeviceId DeviceId)
+	{
+		return false;
+	}
+
+	virtual bool OnNativeGestureEnd(EGestureEvent GestureType, const FVector2D& Delta, float Value, bool bIsDirectionInvertedFromDevice, FPlatformUserId PlatformUserId, FInputDeviceId DeviceId)
+	{
+		return false;
+	}
+
 	virtual bool OnTouchStarted( const TSharedPtr< FGenericWindow >& Window, const FVector2D& Location, float Force, int32 TouchIndex, FPlatformUserId PlatformUserId, FInputDeviceId DeviceId )
 	{
 		if (!ShouldUsePlatformUserId())
diff --git a/Engine/Source/Runtime/ApplicationCore/Public/GenericPlatform/GenericPlatformApplicationMisc.h b/Engine/Source/Runtime/ApplicationCore/Public/GenericPlatform/GenericPlatformApplicationMisc.h
--- a/Engine/Source/Runtime/ApplicationCore/Public/GenericPlatform/GenericPlatformApplicationMisc.h
+++ b/Engine/Source/Runtime/ApplicationCore/Public/GenericPlatform/GenericPlatformApplicationMisc.h
@@ -388,6 +388,8 @@ struct FGenericPlatformApplicationMisc
 		return false;
 	}
 
+	/** Returns true when native platform gesture recognizers should be used instead of software gesture detection. */
+	static APPLICATIONCORE_API bool AreNativeGesturesEnabled();
 
 protected:
 	static APPLICATIONCORE_API bool CachedPhysicalScreenData;
@@ -400,4 +402,6 @@ protected:
 	static APPLICATIONCORE_API FVector2f CachedPhysicalScreenDimensions;
 	static APPLICATIONCORE_API float CachedPhysicalScreenSize;
 
+private:
+	static APPLICATIONCORE_API TAutoConsoleVariable<bool> CVarEnableNativeGestureRecognition;
 };
diff --git a/Engine/Source/Runtime/ApplicationCore/Public/IOS/IOSInputInterface.h b/Engine/Source/Runtime/ApplicationCore/Public/IOS/IOSInputInterface.h
--- a/Engine/Source/Runtime/ApplicationCore/Public/IOS/IOSInputInterface.h
+++ b/Engine/Source/Runtime/ApplicationCore/Public/IOS/IOSInputInterface.h
@@ -74,6 +74,15 @@ struct TouchInput
 	float Force;
 };
 
+struct GestureInput
+{
+    EGestureEvent Type;
+    FVector2D     Delta;       // gesture delta (pan offset, swipe direction)
+    float         Value;       // scalar value (pinch scale, rotation radians)
+    bool          bIsStart;    // UIGestureRecognizerStateBegan
+    bool          bIsEnd;      // UIGestureRecognizerStateEnded/Cancelled
+};
+
 enum class EIOSEventType : int32
 {
     Invalid = 0,
@@ -125,6 +134,7 @@ public:
 	virtual void SetForceFeedbackChannelValues(int32 ControllerId, const FForceFeedbackValues &values) override;
 
 	static APPLICATIONCORE_API void QueueTouchInput(const TArray<TouchInput>& InTouchEvents);
+	static APPLICATIONCORE_API void QueueGestureInput(const TArray<GestureInput>& InGestureInputs);
 
 	void SetGamepadsAllowed(bool bAllowed) { bAllowControllers = bAllowed; }
 	void SetGamepadsBlockDeviceFeedback(bool bBlock) { bControllersBlockDeviceFeedback = bBlock; }
@@ -181,6 +191,8 @@ private:
 
 private:
 	void ProcessTouchesAndKeys(uint32 ControllerId, const TArray<TouchInput>& InTouchInputStack);
+
+	void ProcessGestures(uint32 ControllerId, const TArray<GestureInput>& InGestureInputStack);
 	
 	// can the remote be rotated to landscape
 	bool bAllowRemoteRotation;
diff --git a/Engine/Source/Runtime/ApplicationCore/Public/IOS/IOSView.h b/Engine/Source/Runtime/ApplicationCore/Public/IOS/IOSView.h
--- a/Engine/Source/Runtime/ApplicationCore/Public/IOS/IOSView.h
+++ b/Engine/Source/Runtime/ApplicationCore/Public/IOS/IOSView.h
@@ -42,7 +42,11 @@ namespace MTL
 #ifndef SWIFT_IMPORT
 APPLICATIONCORE_API
 #endif
+#if !PLATFORM_TVOS
+@interface FIOSView : UIView  <UIKeyInput, UITextInput, UIGestureRecognizerDelegate>
+#else
 @interface FIOSView : UIView  <UIKeyInput, UITextInput>
+#endif
 {
 @public
 	// are we initialized?
@@ -128,6 +132,18 @@ APPLICATIONCORE_API
 +(void)StartupEmbeddedUnreal;
 #endif
 
+#if !PLATFORM_TVOS
+// Gesture recognizer setup (called during view init)
+- (void)SetupGestureRecognizers;
+
+// Gesture callbacks
+- (void)HandlePanGesture:(UIPanGestureRecognizer*)Recognizer;
+- (void)HandlePinchGesture:(UIPinchGestureRecognizer*)Recognizer;
+- (void)HandleRotateGesture:(UIRotationGestureRecognizer*)Recognizer;
+- (void)HandleTapGesture:(UITapGestureRecognizer*)Recognizer;
+- (void)HandleLongPressGesture:(UILongPressGestureRecognizer*)Recognizer;
+#endif
+
 @end
 
 
diff --git a/Engine/Source/Runtime/Engine/Classes/Engine/GameViewportClient.h b/Engine/Source/Runtime/Engine/Classes/Engine/GameViewportClient.h
--- a/Engine/Source/Runtime/Engine/Classes/Engine/GameViewportClient.h
+++ b/Engine/Source/Runtime/Engine/Classes/Engine/GameViewportClient.h
@@ -174,6 +174,7 @@ public:
 	ENGINE_API virtual bool InputChar(FViewport* Viewport,int32 ControllerId, TCHAR Character) override;
 	ENGINE_API virtual bool InputTouch(FViewport* const InViewport, const FTouchId TouchId, const ETouchType::Type Type, const FVector2D& TouchLocation, const float Force, const uint64 Timestamp) override;
 	ENGINE_API virtual bool InputMotion(FViewport* Viewport, const FInputDeviceId DeviceId, const FVector& Tilt, const FVector& RotationRate, const FVector& Gravity, const FVector& Acceleration, const uint64 Timestamp) override;
+	ENGINE_API virtual bool InputGesture(FViewport* Viewport, const FInputDeviceId DeviceId, EGestureEvent GestureType, EGesturePhase GesturePhase, const FVector2D& GestureDelta, bool bIsDirectionInvertedFromDevice, const uint64 Timestamp) override;
 	ENGINE_API virtual EMouseCursor::Type GetCursor(FViewport* Viewport, int32 X, int32 Y ) override;
 	ENGINE_API virtual TOptional<TSharedRef<SWidget>> MapCursor(FViewport* Viewport, const FCursorReply& CursorReply) override;
 	ENGINE_API virtual void Precache() override;
diff --git a/Engine/Source/Runtime/Engine/Classes/GameFramework/PlayerController.h b/Engine/Source/Runtime/Engine/Classes/GameFramework/PlayerController.h
--- a/Engine/Source/Runtime/Engine/Classes/GameFramework/PlayerController.h
+++ b/Engine/Source/Runtime/Engine/Classes/GameFramework/PlayerController.h
@@ -53,6 +53,8 @@ struct FForceFeedbackEffectHistoryEntry;
 struct FInputKeyParams;
 struct FPlatformUserId;
 class UAsyncPhysicsInputComponent;
+enum class EGestureEvent : uint8;
+enum class EGesturePhase : uint8;
 
 /** Default delegate that provides an implementation for those that don't have special needs other than a toggle */
 DECLARE_DELEGATE_RetVal(bool, FCanUnpause);
@@ -1825,6 +1827,9 @@ public:
 	/** Handles motion control */
 	ENGINE_API virtual bool InputMotion(const FInputDeviceId DeviceId, const FVector& Tilt, const FVector& RotationRate, const FVector& Gravity, const FVector& Acceleration, const uint64 Timestamp);
 
+	/** Handles a gesture input event, maps EGestureEvent to FKey and forwards to PlayerInput */
+	ENGINE_API virtual bool InputGesture(const FInputDeviceId DeviceId, EGestureEvent GestureType, EGesturePhase GesturePhase, const FVector2D& GestureDelta, bool bIsDirectionInvertedFromDevice, const uint64 Timestamp);
+
 	/** Associate a new UPlayer with this PlayerController. */
 	ENGINE_API virtual void SetPlayer(UPlayer* InPlayer);
 
diff --git a/Engine/Source/Runtime/Engine/Classes/GameFramework/PlayerInput.h b/Engine/Source/Runtime/Engine/Classes/GameFramework/PlayerInput.h
--- a/Engine/Source/Runtime/Engine/Classes/GameFramework/PlayerInput.h
+++ b/Engine/Source/Runtime/Engine/Classes/GameFramework/PlayerInput.h
@@ -665,6 +665,9 @@ public:
 		const float Value,
 		const uint64 Timestamp);
 
+	/** Handles a gesture input event with 2D value (e.g. pan delta). Returns true. */
+	ENGINE_API bool InputGesture(const FInputKeyEventArgs& Params);
+
 	/** Manually update the GestureRecognizer AnchorDistance using the current locations of the touches */
 	ENGINE_API void UpdatePinchStartDistance();
 
diff --git a/Engine/Source/Runtime/Engine/Classes/Kismet/KismetInputLibrary.h b/Engine/Source/Runtime/Engine/Classes/Kismet/KismetInputLibrary.h
--- a/Engine/Source/Runtime/Engine/Classes/Kismet/KismetInputLibrary.h
+++ b/Engine/Source/Runtime/Engine/Classes/Kismet/KismetInputLibrary.h
@@ -20,7 +20,10 @@ enum class ESlateGesture : uint8
 	Magnify,
 	Swipe,
 	Rotate,
-	LongPress
+	LongPress,
+	Pan,
+	Tap,
+	Flick
 };
 
 /** A structure which captures the application's modifier key states (shift, alt, ctrl, etc.) */
diff --git a/Engine/Source/Runtime/Engine/Private/GameViewportClient.cpp b/Engine/Source/Runtime/Engine/Private/GameViewportClient.cpp
--- a/Engine/Source/Runtime/Engine/Private/GameViewportClient.cpp
+++ b/Engine/Source/Runtime/Engine/Private/GameViewportClient.cpp
@@ -947,6 +947,24 @@ bool UGameViewportClient::InputMotion(FViewport* InViewport, const FInputDeviceI
 	return bResult;
 }
 
+bool UGameViewportClient::InputGesture(FViewport* InViewport, const FInputDeviceId DeviceId, EGestureEvent GestureType, EGesturePhase GesturePhase, const FVector2D& GestureDelta, bool bIsDirectionInvertedFromDevice, const uint64 Timestamp)
+{
+	if (IgnoreInput())
+	{
+		return false;
+	}
+
+	bool bResult = false;
+
+	ULocalPlayer* const TargetPlayer = GEngine->GetLocalPlayerFromInputDevice(this, DeviceId);
+	if (TargetPlayer && TargetPlayer->PlayerController)
+	{
+		bResult = TargetPlayer->PlayerController->InputGesture(DeviceId, GestureType, GesturePhase, GestureDelta, bIsDirectionInvertedFromDevice, Timestamp);
+	}
+
+	return bResult;
+}
+
 void UGameViewportClient::SetIsSimulateInEditorViewport(bool bInIsSimulateInEditorViewport)
 {
 	if (FPlatformMisc::DesktopTouchScreen() && GetUseMouseForTouch())
diff --git a/Engine/Source/Runtime/Engine/Private/KismetInputLibrary.cpp b/Engine/Source/Runtime/Engine/Private/KismetInputLibrary.cpp
--- a/Engine/Source/Runtime/Engine/Private/KismetInputLibrary.cpp
+++ b/Engine/Source/Runtime/Engine/Private/KismetInputLibrary.cpp
@@ -322,6 +322,9 @@ ESlateGesture UKismetInputLibrary::PointerEvent_GetGestureType(const FPointerEve
 	static_assert( (int32)EGestureEvent::Swipe == (int32)ESlateGesture::Swipe, "EGestureEvent == ESlateGesture" );
 	static_assert( (int32)EGestureEvent::Rotate == (int32)ESlateGesture::Rotate, "EGestureEvent == ESlateGesture" );
 	static_assert( (int32)EGestureEvent::LongPress == (int32)ESlateGesture::LongPress, "EGestureEvent == ESlateGesture" );
+	static_assert( (int32)EGestureEvent::Pan == (int32)ESlateGesture::Pan, "EGestureEvent == ESlateGesture" );
+	static_assert( (int32)EGestureEvent::Tap == (int32)ESlateGesture::Tap, "EGestureEvent == ESlateGesture" );
+	static_assert( (int32)EGestureEvent::Flick == (int32)ESlateGesture::Flick, "EGestureEvent == ESlateGesture" );
 
 	switch ( Input.GetGestureType() )
 	{
@@ -335,6 +338,12 @@ ESlateGesture UKismetInputLibrary::PointerEvent_GetGestureType(const FPointerEve
 		return ESlateGesture::Rotate;
 	case EGestureEvent::LongPress:
 		return ESlateGesture::LongPress;
+	case EGestureEvent::Pan:
+		return ESlateGesture::Pan;
+	case EGestureEvent::Tap:
+		return ESlateGesture::Tap;
+	case EGestureEvent::Flick:
+		return ESlateGesture::Flick;
 	case EGestureEvent::None:
 	default:
 		return ESlateGesture::None;
diff --git a/Engine/Source/Runtime/Engine/Private/PlayerController.cpp b/Engine/Source/Runtime/Engine/Private/PlayerController.cpp
--- a/Engine/Source/Runtime/Engine/Private/PlayerController.cpp
+++ b/Engine/Source/Runtime/Engine/Private/PlayerController.cpp
@@ -2610,6 +2610,38 @@ bool APlayerController::InputMotion(const FInputDeviceId DeviceId, const FVector
 	return bResult;
 }
 
+bool APlayerController::InputGesture(const FInputDeviceId DeviceId, EGestureEvent GestureType, EGesturePhase GesturePhase, const FVector2D& GestureDelta, bool bIsDirectionInvertedFromDevice, const uint64 Timestamp)
+{
+	bool bResult = false;
+
+	if (PlayerInput)
+	{
+		FKey GestureKey;
+		switch (GestureType)
+		{
+		case EGestureEvent::Pan:      GestureKey = EKeys::Gesture_Pan; break;
+		case EGestureEvent::Magnify:  GestureKey = EKeys::Gesture_Pinch; break;
+		case EGestureEvent::Rotate:   GestureKey = EKeys::Gesture_Rotate; break;
+		case EGestureEvent::Tap:      GestureKey = EKeys::Gesture_Tap; break;
+		case EGestureEvent::LongPress: GestureKey = EKeys::Gesture_LongPress; break;
+		case EGestureEvent::Flick:    GestureKey = EKeys::Gesture_Flick; break;
+		default: return false;
+		}
+
+		EInputEvent InputEvent;
+		switch (GesturePhase)
+		{
+		case EGesturePhase::Started: InputEvent = IE_Pressed; break;
+		case EGesturePhase::Update:  InputEvent = IE_Repeat; break;
+		case EGesturePhase::Ended:   InputEvent = IE_Released; break;
+		default: return false;
+		}
+		bResult = PlayerInput->InputGesture(FInputKeyEventArgs(DeviceId, GestureKey, InputEvent, GestureDelta, Timestamp));
+	}
+
+	return bResult;
+}
+
 void APlayerController::SetMotionControlsEnabled(bool bEnabled)
 {
 	bEnableMotionControls = bEnabled;
diff --git a/Engine/Source/Runtime/Engine/Private/Slate/SceneViewport.cpp b/Engine/Source/Runtime/Engine/Private/Slate/SceneViewport.cpp
--- a/Engine/Source/Runtime/Engine/Private/Slate/SceneViewport.cpp
+++ b/Engine/Source/Runtime/Engine/Private/Slate/SceneViewport.cpp
@@ -1081,7 +1081,7 @@ FReply FSceneViewport::OnTouchGesture( const FGeometry& MyGeometry, const FPoint
 			FSlateApplication::Get().SetKeyboardFocus(ViewportWidget.Pin());
 		}
 
-		if(!ViewportClient->InputGesture(this, GestureEvent.GetInputDeviceId(), GestureEvent.GetGestureType(), GestureEvent.GetGestureDelta(), GestureEvent.IsDirectionInvertedFromDevice(), GestureEvent.GetEventTimestamp()))
+		if(!ViewportClient->InputGesture(this, GestureEvent.GetInputDeviceId(), GestureEvent.GetGestureType(), GestureEvent.GetGesturePhase(), GestureEvent.GetGestureDelta(), GestureEvent.IsDirectionInvertedFromDevice(), GestureEvent.GetEventTimestamp()))
 		{
 			CurrentReplyState = FReply::Unhandled();
 		}
diff --git a/Engine/Source/Runtime/Engine/Private/UserInterface/PlayerInput.cpp b/Engine/Source/Runtime/Engine/Private/UserInterface/PlayerInput.cpp
--- a/Engine/Source/Runtime/Engine/Private/UserInterface/PlayerInput.cpp
+++ b/Engine/Source/Runtime/Engine/Private/UserInterface/PlayerInput.cpp
@@ -6,6 +6,7 @@
 
 #include "GameFramework/PlayerInput.h"
 #include "GameFramework/PlayerInputDebugging.h"
+#include "HAL/PlatformApplicationMisc.h"
 #include "Engine/World.h"
 #include "Misc/App.h"
 #include "GlobalRenderResources.h"
@@ -651,6 +652,26 @@ bool UPlayerInput::InputGesture(
 	return true;
 }
 
+bool UPlayerInput::InputGesture(const FInputKeyEventArgs& Params)
+{
+	FKeyState& KeyState = KeyStateMap.FindOrAdd(Params.Key);
+
+	KeyState.Value.X = KeyState.RawValue.X = KeyState.RawValueAccumulator.X = Params.AmountDepressed2D.X;
+	KeyState.Value.Y = KeyState.RawValue.Y = KeyState.RawValueAccumulator.Y = Params.AmountDepressed2D.Y;
+	KeyState.EventAccumulator[Params.Event].Add(++EventCount);
+
+	UE_LOG(LogPlayerInput, VeryVerbose, TEXT("[%hs] Frame: %llu %s (outer: %s) received gesture: Key: %s Value: %s Event:  %s"),
+		__func__,
+		GFrameCounter,
+		*GetNameSafe(this),
+		*GetNameSafe(GetOuter()),
+		*Params.Key.GetFName().ToString(),
+		*Params.AmountDepressed2D.ToString(),
+		UE::Input::LexToString(Params.Event));
+
+	return true;
+}
+
 void UPlayerInput::UpdatePinchStartDistance()
 {
 	GestureRecognizer.SetAnchorDistanceSquared(FVector2D(TouchStates[ETouchIndex::Touch1].ViewportRelativeLocation),
@@ -2357,8 +2378,12 @@ void UPlayerInput::Tick(float DeltaTime)
 
 	if (GetDefault<UInputSettings>()->bEnableGestureRecognizer)
 	{
-		SCOPE_CYCLE_COUNTER(STAT_PC_GestureRecognition);
-		GestureRecognizer.DetectGestures(TouchStates, this, DeltaTime);
+		// Skip the software gesture recognizer when native platform gesture recognizers are active.
+		if (!FGenericPlatformApplicationMisc::AreNativeGesturesEnabled())
+		{
+			SCOPE_CYCLE_COUNTER(STAT_PC_GestureRecognition);
+			GestureRecognizer.DetectGestures(TouchStates, this, DeltaTime);
+		}
 	}
 }
 
diff --git a/Engine/Source/Runtime/Engine/Public/InputKeyEventArgs.h b/Engine/Source/Runtime/Engine/Public/InputKeyEventArgs.h
--- a/Engine/Source/Runtime/Engine/Public/InputKeyEventArgs.h
+++ b/Engine/Source/Runtime/Engine/Public/InputKeyEventArgs.h
@@ -4,6 +4,7 @@
 
 #include "Engine/EngineBaseTypes.h"
 #include "InputCoreTypes.h"
+#include "Math/Vector2D.h"
 
 class FViewport;
 
@@ -104,7 +105,34 @@ public:
 		DeltaTime = InDeltaTime;
 		NumSamples = InNumSamples;
 	}
-	
+
+	/**
+	 * Construct a FInputKeyEventArgs for a 2D gesture input event.
+	 *
+	 * @param InInputDevice		The input device that triggered this event
+	 * @param InKey				The gesture key (e.g. EKeys::Gesture_Pan)
+	 * @param InEvent			The input event type (IE_Pressed, IE_Repeat, IE_Released)
+	 * @param InAmountDepressed2D	The 2D gesture delta value
+	 * @param InEventTimestamp	The timestamp of when this input event was originally polled
+	 */
+	FInputKeyEventArgs(
+		const FInputDeviceId InInputDevice,
+		const FKey& InKey,
+		const EInputEvent InEvent,
+		const FVector2D& InAmountDepressed2D,
+		const uint64 InEventTimestamp)
+		: FInputKeyEventArgs(
+			/*Viewport*/ nullptr,
+			InInputDevice,
+			InKey,
+			InEvent,
+			/*AmountDepressed*/ 0.0f,
+			/*bIsTouchEvent*/ false,
+			InEventTimestamp)
+	{
+		AmountDepressed2D = InAmountDepressed2D;
+	}
+
 	static ENGINE_API FInputKeyEventArgs CreateSimulated(
 		const FKey& InKey,
 		const EInputEvent InEvent,
@@ -159,11 +187,17 @@ public:
 	
 	/**
 	 * The value that this input event represents.
-	 * 
+	 *
 	 * For analog keys, the depression percent.
 	 */
 	float AmountDepressed;
 
+	/**
+	 * 2D value for this input event (e.g. gesture pan delta).
+	 * When set, this takes precedence over AmountDepressed for 2D input.
+	 */
+	FVector2D AmountDepressed2D = FVector2D::ZeroVector;
+
 	/** The time between the previous frame and the current one */
 	float DeltaTime = 1.0f / 60.f;
 
diff --git a/Engine/Source/Runtime/Engine/Public/ViewportClient.h b/Engine/Source/Runtime/Engine/Public/ViewportClient.h
--- a/Engine/Source/Runtime/Engine/Public/ViewportClient.h
+++ b/Engine/Source/Runtime/Engine/Public/ViewportClient.h
@@ -14,6 +14,7 @@ class FWindowActivateEvent;
 class SWidget;
 enum class EFocusCause : uint8;
 enum class EGestureEvent : uint8;
+enum class EGesturePhase : uint8;
 struct FInputKeyEventArgs;
 struct FKey;
 struct FStatHitchesData;
@@ -122,6 +123,7 @@ public:
 	 * @param	Viewport - The viewport on which the gesture was performed.
 	 * @param	DeviceId - The input device on which the gesture was performed.
 	 * @param	GestureType - The type of the gesture.
+	 * @param	GesturePhase - The phase of the gesture.
 	 * @param	GestureDelta - Delta values of this gesture from the previous instance of it.
 	 * @param	Timestamp - The timestamp of when this input event was originally polled in the form of CPU cycles.
 	 * @return	True to consume the gesture event, false to pass it on.
@@ -130,6 +132,7 @@ public:
 		FViewport* Viewport,
 		const FInputDeviceId DeviceId,
 		EGestureEvent GestureType,
+		EGesturePhase GesturePhase,
 		const FVector2D& GestureDelta,
 		bool bIsDirectionInvertedFromDevice,
 		const uint64 Timestamp)
diff --git a/Engine/Source/Runtime/InputCore/Classes/InputCoreTypes.h b/Engine/Source/Runtime/InputCore/Classes/InputCoreTypes.h
--- a/Engine/Source/Runtime/InputCore/Classes/InputCoreTypes.h
+++ b/Engine/Source/Runtime/InputCore/Classes/InputCoreTypes.h
@@ -547,6 +547,9 @@ struct EKeys
 	static INPUTCORE_API const FKey Gesture_Pinch;
 	static INPUTCORE_API const FKey Gesture_Flick;
 	static INPUTCORE_API const FKey Gesture_Rotate;
+	static INPUTCORE_API const FKey Gesture_Pan;
+	static INPUTCORE_API const FKey Gesture_Tap;
+	static INPUTCORE_API const FKey Gesture_LongPress;
 
 	// Steam Controller Specific
 	static INPUTCORE_API const FKey Steam_Touch_0;
diff --git a/Engine/Source/Runtime/InputCore/Private/InputCoreTypes.cpp b/Engine/Source/Runtime/InputCore/Private/InputCoreTypes.cpp
--- a/Engine/Source/Runtime/InputCore/Private/InputCoreTypes.cpp
+++ b/Engine/Source/Runtime/InputCore/Private/InputCoreTypes.cpp
@@ -219,6 +219,9 @@ const FKey EKeys::TouchKeys[NUM_TOUCH_KEYS];
 const FKey EKeys::Gesture_Pinch("Gesture_Pinch");
 const FKey EKeys::Gesture_Flick("Gesture_Flick");
 const FKey EKeys::Gesture_Rotate("Gesture_Rotate");
+const FKey EKeys::Gesture_Pan("Gesture_Pan");
+const FKey EKeys::Gesture_Tap("Gesture_Tap");
+const FKey EKeys::Gesture_LongPress("Gesture_LongPress");
 
 // Steam Controller Specific
 const FKey EKeys::Steam_Touch_0("Steam_Touch_0");
@@ -747,9 +750,12 @@ void EKeys::Initialize()
 	// Gestures
 	AddMenuCategoryDisplayInfo("Gesture", LOCTEXT("GestureSubCateogry", "Gesture"), TEXT("GraphEditor.KeyEvent_16x"));
 
-	AddKey(FKeyDetails(EKeys::Gesture_Pinch, LOCTEXT("Gesture_Pinch", "Pinch"), FKeyDetails::Gesture, "Gesture"));
-	AddKey(FKeyDetails(EKeys::Gesture_Flick, LOCTEXT("Gesture_Flick", "Flick"), FKeyDetails::Gesture, "Gesture"));
-	AddKey(FKeyDetails(EKeys::Gesture_Rotate, LOCTEXT("Gesture_Rotate", "Rotate"), FKeyDetails::Gesture, "Gesture"));
+	AddKey(FKeyDetails(EKeys::Gesture_Pinch, LOCTEXT("Gesture_Pinch", "Pinch"), FKeyDetails::Gesture | FKeyDetails::Axis1D, "Gesture"));
+	AddKey(FKeyDetails(EKeys::Gesture_Flick, LOCTEXT("Gesture_Flick", "Flick"), FKeyDetails::Gesture | FKeyDetails::Axis1D, "Gesture"));
+	AddKey(FKeyDetails(EKeys::Gesture_Rotate, LOCTEXT("Gesture_Rotate", "Rotate"), FKeyDetails::Gesture | FKeyDetails::Axis1D, "Gesture"));
+	AddKey(FKeyDetails(EKeys::Gesture_Pan, LOCTEXT("Gesture_Pan", "Pan"), FKeyDetails::Gesture | FKeyDetails::Axis2D, "Gesture"));
+	AddKey(FKeyDetails(EKeys::Gesture_Tap, LOCTEXT("Gesture_Tap", "Tap"), FKeyDetails::Gesture | FKeyDetails::Axis2D, "Gesture"));
+	AddKey(FKeyDetails(EKeys::Gesture_LongPress, LOCTEXT("Gesture_LongPress", "Long Press"), FKeyDetails::Gesture | FKeyDetails::Axis2D, "Gesture"));
 
 	// PS4-specific
 	AddMenuCategoryDisplayInfo("Special Gamepad", LOCTEXT("SpecialGamepadSubCategory", "SpecialGamepad"), TEXT("GraphEditor.PadEvent_16x"));
diff --git a/Engine/Source/Runtime/Slate/Private/Framework/Application/SlateApplication.cpp b/Engine/Source/Runtime/Slate/Private/Framework/Application/SlateApplication.cpp
--- a/Engine/Source/Runtime/Slate/Private/Framework/Application/SlateApplication.cpp
+++ b/Engine/Source/Runtime/Slate/Private/Framework/Application/SlateApplication.cpp
@@ -1714,9 +1714,13 @@ void FSlateApplication::TickPlatform(float DeltaTime)
 		bSynthesizedCursorMove = bSynthesizedCursorMoveThisFrame;
 
 		// Generate any simulated gestures that we've detected.
-		ForEachUser([this](FSlateUser& User) {
-			User.GetGestureDetector().GenerateGestures(*this, SimulateGestures);
-		});
+		// Skip the software gesture detector when native platform gesture recognizers are active.
+		if (!FGenericPlatformApplicationMisc::AreNativeGesturesEnabled())
+		{
+			ForEachUser([this](FSlateUser& User) {
+				User.GetGestureDetector().GenerateGestures(*this, SimulateGestures);
+			});
+		}
 	}
 }
 
@@ -6176,7 +6180,9 @@ bool FSlateApplication::ProcessMouseWheelOrGestureEvent( const FPointerEvent& In
 			bShouldProcessEvent = true;
 			break;
 		default:
-			bShouldProcessEvent = InGestureEvent->GetGestureDelta() != FVector2f::ZeroVector;
+			bShouldProcessEvent = (InGestureEvent->GetGestureDelta() != FVector2f::ZeroVector) ||
+				(InGestureEvent->GetGesturePhase() == EGesturePhase::Started) ||
+				(InGestureEvent->GetGesturePhase() == EGesturePhase::Ended);
 			break;
 		}
 	}
@@ -6673,33 +6679,58 @@ bool FSlateApplication::OnControllerButtonReleased(FGamepadKeyNames::Type KeyNam
 	return false;
 }
 
-bool FSlateApplication::OnTouchGesture( EGestureEvent GestureType, const FVector2D &Delta, const float MouseWheelDelta, bool bIsDirectionInvertedFromDevice )
+bool FSlateApplication::HandleTouchGesture( FInputDeviceId DeviceId, EGestureEvent GestureType, EGesturePhase GesturePhase, const FVector2D &Delta, const float Value, bool bIsDirectionInvertedFromDevice )
 {
 	const FVector2f CurrentCursorPosition = GetCursorPos();
-	
+
 	FPointerEvent GestureEvent(
+		DeviceId,
 		CurrentCursorPosition,
 		CurrentCursorPosition,
 		PressedMouseButtons,
 		PlatformApplication->GetModifierKeys(),
 		GestureType,
+		GesturePhase,
 		Delta,
 		bIsDirectionInvertedFromDevice
 	);
 	
 	FPointerEvent MouseWheelEvent(
+		DeviceId,
 		CursorPointerIndex,
 		CurrentCursorPosition,
 		CurrentCursorPosition,
 		PressedMouseButtons,
 		EKeys::Invalid,
-		MouseWheelDelta,
+		Value,
 		PlatformApplication->GetModifierKeys()
 	);
 	
 	return ProcessMouseWheelOrGestureEvent( MouseWheelEvent, &GestureEvent );
 }
 
+bool FSlateApplication::OnTouchGesture( EGestureEvent GestureType, const FVector2D &Delta, const float MouseWheelDelta, bool bIsDirectionInvertedFromDevice )
+{
+	// Software gesture path (e.g. Mac trackpad via GestureRecognizer) has no device ID.
+	// Use device 0 to preserve the original UserIndex=0 routing behavior.
+	return HandleTouchGesture(FInputDeviceId::CreateFromInternalId(0), GestureType, EGesturePhase::Update, Delta, MouseWheelDelta, bIsDirectionInvertedFromDevice);
+}
+
+bool FSlateApplication::OnNativeGestureBegin(EGestureEvent GestureType, const FVector2D& Delta, float Value, bool bIsDirectionInvertedFromDevice, FPlatformUserId PlatformUserId, FInputDeviceId DeviceId)
+{
+	return HandleTouchGesture(DeviceId, GestureType, EGesturePhase::Started, Delta, Value, bIsDirectionInvertedFromDevice);
+}
+
+bool FSlateApplication::OnNativeGestureUpdate(EGestureEvent GestureType, const FVector2D& Delta, float Value, bool bIsDirectionInvertedFromDevice, FPlatformUserId PlatformUserId, FInputDeviceId DeviceId)
+{
+	return HandleTouchGesture(DeviceId, GestureType, EGesturePhase::Update, Delta, Value, bIsDirectionInvertedFromDevice);
+}
+
+bool FSlateApplication::OnNativeGestureEnd(EGestureEvent GestureType, const FVector2D& Delta, float Value, bool bIsDirectionInvertedFromDevice, FPlatformUserId PlatformUserId, FInputDeviceId DeviceId)
+{
+	return HandleTouchGesture(DeviceId, GestureType, EGesturePhase::Ended, Delta, Value, bIsDirectionInvertedFromDevice);
+}
+
 bool ValidateTouchIndex(int32 TouchIndex)
 {
 	return true;
diff --git a/Engine/Source/Runtime/Slate/Public/Framework/Application/SlateApplication.h b/Engine/Source/Runtime/Slate/Public/Framework/Application/SlateApplication.h
--- a/Engine/Source/Runtime/Slate/Public/Framework/Application/SlateApplication.h
+++ b/Engine/Source/Runtime/Slate/Public/Framework/Application/SlateApplication.h
@@ -1684,6 +1684,9 @@ public:
 	SLATE_API virtual bool OnRawMouseMove( const int32 X, const int32 Y ) override;
 	SLATE_API virtual bool OnCursorSet() override;
 	SLATE_API virtual bool OnTouchGesture( EGestureEvent GestureType, const FVector2D& Delta, float WheelDelta, bool bIsDirectionInvertedFromDevice ) override;
+	SLATE_API virtual bool OnNativeGestureBegin( EGestureEvent GestureType, const FVector2D& Delta, float Value, bool bIsDirectionInvertedFromDevice, FPlatformUserId PlatformUserId, FInputDeviceId DeviceId ) override;
+	SLATE_API virtual bool OnNativeGestureUpdate( EGestureEvent GestureType, const FVector2D& Delta, float Value, bool bIsDirectionInvertedFromDevice, FPlatformUserId PlatformUserId, FInputDeviceId DeviceId ) override;
+	SLATE_API virtual bool OnNativeGestureEnd( EGestureEvent GestureType, const FVector2D& Delta, float Value, bool bIsDirectionInvertedFromDevice, FPlatformUserId PlatformUserId, FInputDeviceId DeviceId ) override;
 	SLATE_API virtual bool OnTouchStarted( const TSharedPtr< FGenericWindow >& Window, const FVector2D& Location, float Force, int32 TouchIndex, FPlatformUserId PlatformUserId, FInputDeviceId DeviceId ) override;
 	SLATE_API virtual bool OnTouchMoved( const FVector2D& Location, float Force, int32 TouchIndex, FPlatformUserId PlatformUserId, FInputDeviceId DeviceID ) override;
 	SLATE_API virtual bool OnTouchEnded( const FVector2D& Location, int32 TouchIndex, FPlatformUserId PlatformUserId, FInputDeviceId DeviceID ) override;
@@ -1933,6 +1936,8 @@ private:
 	 * desktop changes (e.g. resolution change or monitors re-arranged)
 	 */
 	SLATE_API void OnVirtualDesktopSizeChanged(const FDisplayMetrics& NewDisplayMetric);
+	
+	SLATE_API bool HandleTouchGesture(FInputDeviceId DeviceId, EGestureEvent GestureType, EGesturePhase GesturePhase, const FVector2D &Delta, const float Value, bool bIsDirectionInvertedFromDevice);
 
 	/** Application singleton */
 	static SLATE_API TSharedPtr< FSlateApplication > CurrentApplication;
diff --git a/Engine/Source/Runtime/SlateCore/Public/Input/Events.h b/Engine/Source/Runtime/SlateCore/Public/Input/Events.h
--- a/Engine/Source/Runtime/SlateCore/Public/Input/Events.h
+++ b/Engine/Source/Runtime/SlateCore/Public/Input/Events.h
@@ -710,6 +710,7 @@ public:
 		, Force(1.0f)
 		, bIsTouchEvent(false)
 		, GestureType(EGestureEvent::None)
+		, GesturePhase(EGesturePhase::None)
 		, WheelOrGestureDelta(0.0f, 0)
 		, bIsDirectionInvertedFromDevice(false)
 		, bIsTouchForceChanged(false)
@@ -737,6 +738,7 @@ public:
 		, Force(1.0f)
 		, bIsTouchEvent(false)
 		, GestureType(EGestureEvent::None)
+		, GesturePhase(EGesturePhase::None)
 		, WheelOrGestureDelta(0.0f, InWheelDelta)
 		, bIsDirectionInvertedFromDevice(false)
 		, bIsTouchForceChanged(false)
@@ -764,6 +766,7 @@ public:
 		, Force(1.0f)
 		, bIsTouchEvent(false)
 		, GestureType(EGestureEvent::None)
+		, GesturePhase(EGesturePhase::None)
 		, WheelOrGestureDelta(0.0f, InWheelDelta)
 		, bIsDirectionInvertedFromDevice(false)
 		, bIsTouchForceChanged(false)
@@ -792,6 +795,7 @@ public:
 		, Force(1.0f)
 		, bIsTouchEvent(false)
 		, GestureType(EGestureEvent::None)
+		, GesturePhase(EGesturePhase::None)
 		, WheelOrGestureDelta(0.0f, InWheelDelta)
 		, bIsDirectionInvertedFromDevice(false)
 		, bIsTouchForceChanged(false)
@@ -822,6 +826,7 @@ public:
 		, Force(1.0f)
 		, bIsTouchEvent(false)
 		, GestureType(EGestureEvent::None)
+		, GesturePhase(EGesturePhase::None)
 		, WheelOrGestureDelta(0.0f, 0.0f)
 		, bIsDirectionInvertedFromDevice(false)
 		, bIsTouchForceChanged(false)
@@ -847,6 +852,7 @@ public:
 		, Force(1.0f)
 		, bIsTouchEvent(false)
 		, GestureType(EGestureEvent::None)
+		, GesturePhase(EGesturePhase::None)
 		, WheelOrGestureDelta(0.0f, 0.0f)
 		, bIsDirectionInvertedFromDevice(false)
 		, bIsTouchForceChanged(false)
@@ -877,6 +883,7 @@ public:
 		, Force(InForce)
 		, bIsTouchEvent(true)
 		, GestureType(EGestureEvent::None)
+		, GesturePhase(EGesturePhase::None)
 		, WheelOrGestureDelta(0.0f, 0.0f)
 		, bIsDirectionInvertedFromDevice(false)
 		, bIsTouchForceChanged(bInIsForceChanged)
@@ -971,6 +978,7 @@ public:
 		, Force(InForce)
 		, bIsTouchEvent(true)
 		, GestureType(EGestureEvent::None)
+		, GesturePhase(EGesturePhase::None)
 		, WheelOrGestureDelta(0.0f, 0.0f)
 		, bIsDirectionInvertedFromDevice(false)
 		, bIsTouchForceChanged(bInIsForceChanged)
@@ -989,6 +997,7 @@ public:
 		const TSet<FKey>& InPressedButtons,
 		const FModifierKeysState& InModifierKeys,
 		EGestureEvent InGestureType,
+		EGesturePhase InGesturePhase,
 		const UE::Slate::FDeprecateVector2DParameter& InGestureDelta,
 		bool bInIsDirectionInvertedFromDevice
 	)
@@ -1002,6 +1011,36 @@ public:
 		, Force(1.0f)
 		, bIsTouchEvent(InGestureType != EGestureEvent::None)
 		, GestureType(InGestureType)
+		, GesturePhase(InGesturePhase)
+		, WheelOrGestureDelta(InGestureDelta)
+		, bIsDirectionInvertedFromDevice(bInIsDirectionInvertedFromDevice)
+		, bIsTouchForceChanged(false)
+		, bIsTouchFirstMove(false)
+	{ }
+
+	/** A constructor for gesture events with a specific input device */
+	FPointerEvent(
+		FInputDeviceId InDeviceId,
+		const UE::Slate::FDeprecateVector2DParameter& InScreenSpacePosition,
+		const UE::Slate::FDeprecateVector2DParameter& InLastScreenSpacePosition,
+		const TSet<FKey>& InPressedButtons,
+		const FModifierKeysState& InModifierKeys,
+		EGestureEvent InGestureType,
+		EGesturePhase InGesturePhase,
+		const UE::Slate::FDeprecateVector2DParameter& InGestureDelta,
+		bool bInIsDirectionInvertedFromDevice
+	)
+		: FInputEvent(InModifierKeys, InDeviceId, false)
+		, ScreenSpacePosition(InScreenSpacePosition)
+		, LastScreenSpacePosition(InLastScreenSpacePosition)
+		, CursorDelta(FVector2f(LastScreenSpacePosition) - FVector2f(ScreenSpacePosition))
+		, PressedButtons(InPressedButtons)
+		, PointerIndex(0)
+		, TouchpadIndex(0)
+		, Force(1.0f)
+		, bIsTouchEvent(InGestureType != EGestureEvent::None)
+		, GestureType(InGestureType)
+		, GesturePhase(InGesturePhase)
 		, WheelOrGestureDelta(InGestureDelta)
 		, bIsDirectionInvertedFromDevice(bInIsDirectionInvertedFromDevice)
 		, bIsTouchForceChanged(false)
@@ -1062,6 +1101,9 @@ public:
 	/** Returns the type of touch gesture */
 	EGestureEvent GetGestureType() const { return GestureType; }
 
+	/** Returns the phase of the gesture (Started, Update, Ended) */
+	EGesturePhase GetGesturePhase() const { return GesturePhase; }
+
 	/** Returns the change in gesture value since the last gesture event of the same type. */
 	const UE::Slate::FDeprecateVector2DResult& GetGestureDelta() const { return WheelOrGestureDelta; }
 
@@ -1097,6 +1139,7 @@ private:
 	float Force;
 	bool bIsTouchEvent;
 	EGestureEvent GestureType;
+	EGesturePhase GesturePhase;
 	UE::Slate::FDeprecateVector2DResult WheelOrGestureDelta;
 	bool bIsDirectionInvertedFromDevice;
 	bool bIsTouchForceChanged;
