Fixed the iOS compilation warning "pointer is missing a nullability type specifier".

#Synced-CL 52536323

diff --git a/Engine/Source/Runtime/ApplicationCore/Public/IOS/IOSView.h b/Engine/Source/Runtime/ApplicationCore/Public/IOS/IOSView.h
--- a/Engine/Source/Runtime/ApplicationCore/Public/IOS/IOSView.h
+++ b/Engine/Source/Runtime/ApplicationCore/Public/IOS/IOSView.h
@@ -137,11 +137,11 @@ APPLICATIONCORE_API
 - (void)SetupGestureRecognizers;
 
 // Gesture callbacks
-- (void)HandlePanGesture:(UIPanGestureRecognizer*)Recognizer;
-- (void)HandlePinchGesture:(UIPinchGestureRecognizer*)Recognizer;
-- (void)HandleRotateGesture:(UIRotationGestureRecognizer*)Recognizer;
-- (void)HandleTapGesture:(UITapGestureRecognizer*)Recognizer;
-- (void)HandleLongPressGesture:(UILongPressGestureRecognizer*)Recognizer;
+- (void)HandlePanGesture:(nonnull UIPanGestureRecognizer*)Recognizer;
+- (void)HandlePinchGesture:(nonnull UIPinchGestureRecognizer*)Recognizer;
+- (void)HandleRotateGesture:(nonnull UIRotationGestureRecognizer*)Recognizer;
+- (void)HandleTapGesture:(nonnull UITapGestureRecognizer*)Recognizer;
+- (void)HandleLongPressGesture:(nonnull UILongPressGestureRecognizer*)Recognizer;
 #endif
 
 @end
