Fix unreachable code warnings

#rnx
#Synced-CL 52553919

diff --git a/Engine/Source/Developer/Zen/Private/BuildServerInterface.cpp b/Engine/Source/Developer/Zen/Private/BuildServerInterface.cpp
--- a/Engine/Source/Developer/Zen/Private/BuildServerInterface.cpp
+++ b/Engine/Source/Developer/Zen/Private/BuildServerInterface.cpp
@@ -1101,7 +1101,6 @@ FBuildServiceInstance::AcquireAccessToken(EDesktopLoginInteractionLevel Interact
 		FailedLoginAttempts++;
 		return false;
 	}
-	return false;
 }
 
 void
diff --git a/Engine/Source/Programs/AutoRTFMTests/Private/LongJumpTests.cpp b/Engine/Source/Programs/AutoRTFMTests/Private/LongJumpTests.cpp
--- a/Engine/Source/Programs/AutoRTFMTests/Private/LongJumpTests.cpp
+++ b/Engine/Source/Programs/AutoRTFMTests/Private/LongJumpTests.cpp
@@ -44,7 +44,9 @@ TEST_CASE("LongJump")
 		/* Try */ [&]
 		{
 			LongJump.Throw();
+			PRAGMA_DISABLE_UNREACHABLE_CODE_WARNINGS
 			bPostThrow = true;
+			PRAGMA_RESTORE_UNREACHABLE_CODE_WARNINGS
 		},
 		/* Catch */ [&]
 		{
diff --git a/Engine/Source/Runtime/ApplicationCore/Private/Windows/WindowsApplication.cpp b/Engine/Source/Runtime/ApplicationCore/Private/Windows/WindowsApplication.cpp
--- a/Engine/Source/Runtime/ApplicationCore/Private/Windows/WindowsApplication.cpp
+++ b/Engine/Source/Runtime/ApplicationCore/Private/Windows/WindowsApplication.cpp
@@ -2334,7 +2334,6 @@ int32 FWindowsApplication::ProcessMessage( HWND hwnd, uint32 msg, WPARAM wParam,
 				if (!bConsumeAltSpace && (HIWORD(lParam) & 0x2000) != 0 && wParam == VK_SPACE)
 				{
 					// Do not handle Alt+Space so that it passes through and opens the window system menu
-					break;
 				}
 				else
 				{
@@ -2390,7 +2389,6 @@ int32 FWindowsApplication::ProcessMessage( HWND hwnd, uint32 msg, WPARAM wParam,
 				// Handled
 				return 0;
 			}
-			break;
 
 		case WM_POINTERDOWN:
 		case WM_POINTERUPDATE:
@@ -2618,7 +2616,6 @@ int32 FWindowsApplication::ProcessMessage( HWND hwnd, uint32 msg, WPARAM wParam,
 
 						return 0;
 					}
-					return 0;
 				}
 			}
 			break;
@@ -2647,7 +2644,6 @@ int32 FWindowsApplication::ProcessMessage( HWND hwnd, uint32 msg, WPARAM wParam,
 
 				return 0;
 			}
-			break;
 
 		case WM_SIZING:
 			{
@@ -2825,7 +2821,6 @@ int32 FWindowsApplication::ProcessMessage( HWND hwnd, uint32 msg, WPARAM wParam,
 
 				return 0;
 			}
-			break;
 
 		case WM_NCHITTEST:
 			{
@@ -2892,7 +2887,6 @@ int32 FWindowsApplication::ProcessMessage( HWND hwnd, uint32 msg, WPARAM wParam,
 				DeferMessage( CurrentNativeEventWindowPtr, hwnd, msg, wParam, lParam );
 				return 0;
 			}
-			break;
 
 		case WM_SETTINGCHANGE:
 			{
@@ -2924,7 +2918,6 @@ int32 FWindowsApplication::ProcessMessage( HWND hwnd, uint32 msg, WPARAM wParam,
 				// Return non-zero to indicate that we'll handle the erasing ourselves
 				return 1;
 			}
-			break;
 
 		case WM_NCACTIVATE:
 			{
@@ -2959,14 +2952,12 @@ int32 FWindowsApplication::ProcessMessage( HWND hwnd, uint32 msg, WPARAM wParam,
 #endif
 				return 0;
 			}
-			break;
 
 		case WM_CLOSE:
 			{
 				DeferMessage( CurrentNativeEventWindowPtr, hwnd, msg, wParam, lParam );
 				return 0;
 			}
-			break;
 
 		case WM_ENDSESSION:
 			{
@@ -2978,7 +2969,6 @@ int32 FWindowsApplication::ProcessMessage( HWND hwnd, uint32 msg, WPARAM wParam,
 				}
 				return IntCastChecked<int32>(DefWindowProc(hwnd, msg, wParam, lParam));
 			}
-			break;
 
 		case WM_SYSCOMMAND:
 			{
@@ -3017,7 +3007,6 @@ int32 FWindowsApplication::ProcessMessage( HWND hwnd, uint32 msg, WPARAM wParam,
 						}
 						return 1;
 					}
-					break;
 				default:
 					if( !( MessageHandler->ShouldProcessUserInputMessages( CurrentNativeEventWindow ) && IsInputMessage( msg ) ) )
 					{
@@ -3058,7 +3047,6 @@ int32 FWindowsApplication::ProcessMessage( HWND hwnd, uint32 msg, WPARAM wParam,
 				MinMaxInfo->ptMaxTrackSize.y = FMath::RoundToInt( SizeLimits.GetMaxHeight().Get((float)MinMaxInfo->ptMaxTrackSize.y) * DPIScaleFactor ) + BorderHeight;
 				return 0;
 			}
-			break;
 			
 		case WM_NCLBUTTONDOWN:
 		case WM_NCRBUTTONDOWN:
@@ -3120,7 +3108,6 @@ int32 FWindowsApplication::ProcessMessage( HWND hwnd, uint32 msg, WPARAM wParam,
 				// Slate wants all keys and messages.
 				return DLGC_WANTALLKEYS;
 			}
-			break;
 		
 		case WM_CREATE:
 			return 0;
@@ -3233,7 +3220,6 @@ int32 FWindowsApplication::ProcessDeferredMessage( const FDeferredWindowsMessage
 				}
 				return 0;
 			}
-			break;
 			// Character
 		case WM_CHAR:
 			{
@@ -3250,7 +3236,6 @@ PRAGMA_ENABLE_DEPRECATION_WARNINGS
 				// Note: always return 0 to handle the message.  Win32 beeps if WM_CHAR is not handled...
 				return 0;
 			}
-			break;
 
 
 			// Key down
@@ -3512,7 +3497,6 @@ PRAGMA_ENABLE_DEPRECATION_WARNINGS
 				}
 				return 0;
 			}
-			break;
 
 		// Mouse Movement
 		case WM_INPUT:
@@ -3532,7 +3516,6 @@ PRAGMA_ENABLE_DEPRECATION_WARNINGS
 
 				return 0;
 			}
-			break;
 
 		// Mouse Movement
 		case WM_NCMOUSEMOVE:
@@ -3549,7 +3532,6 @@ PRAGMA_ENABLE_DEPRECATION_WARNINGS
 
 				return Result ? 0 : 1;
 			}
-			break;
 			// Mouse Wheel
 		case WM_MOUSEWHEEL:
 			{
@@ -3568,7 +3550,6 @@ PRAGMA_ENABLE_DEPRECATION_WARNINGS
 				const BOOL Result = MessageHandler->OnMouseWheel( static_cast<float>( WheelDelta ) * SpinFactor, CursorPos );
 				return Result ? 0 : 1;
 			}
-			break;
 
 			// Mouse Cursor
 		case WM_SETCURSOR:
@@ -3588,7 +3569,6 @@ PRAGMA_ENABLE_DEPRECATION_WARNINGS
 				}
 				return  MessageHandler->OnCursorSet() ? 1 : 0;
 			}
-			break;
 
 		case WM_TOUCH:
 			{
@@ -3727,7 +3707,6 @@ PRAGMA_ENABLE_DEPRECATION_WARNINGS
 			
 				return 0;
 			}
-			break;
 
 
 		case WM_POINTERUPDATE:
@@ -3793,7 +3772,6 @@ PRAGMA_ENABLE_DEPRECATION_WARNINGS
 
 				return 0;
 			}
-			break;
 
 		case WM_POINTERUP:
 			{
@@ -3844,7 +3822,6 @@ PRAGMA_ENABLE_DEPRECATION_WARNINGS
 
 				return 0;
 			}
-			break;
 
 			// Window focus and activation
 		case WM_MOUSEACTIVATE:
@@ -3855,7 +3832,6 @@ PRAGMA_ENABLE_DEPRECATION_WARNINGS
 				bForceActivateByMouse = !(LOWORD(lParam) & HTCLIENT);
 				return 0;
 			}
-			break;
 
 			// Window focus and activation
 		case WM_ACTIVATE:
@@ -3887,7 +3863,6 @@ PRAGMA_ENABLE_DEPRECATION_WARNINGS
 
 				return 1;
 			}
-			break;
 
 		case WM_ACTIVATEAPP:
 			UpdateAllModifierKeyStates();
@@ -3937,7 +3912,6 @@ PRAGMA_ENABLE_DEPRECATION_WARNINGS
 				}
 				return 0;
 			}
-			break;
 
 		case WM_SHOWWINDOW:
 			{
diff --git a/Engine/Source/Runtime/ApplicationCore/Private/Windows/WindowsTextInputMethodSystem.cpp b/Engine/Source/Runtime/ApplicationCore/Private/Windows/WindowsTextInputMethodSystem.cpp
--- a/Engine/Source/Runtime/ApplicationCore/Private/Windows/WindowsTextInputMethodSystem.cpp
+++ b/Engine/Source/Runtime/ApplicationCore/Private/Windows/WindowsTextInputMethodSystem.cpp
@@ -990,7 +990,6 @@ int32 FWindowsTextInputMethodSystem::ProcessMessage(HWND hwnd, uint32 msg, WPARA
 
 			return IntCastChecked<int32>(DefWindowProc(hwnd, msg, wParam, lParam));
 		}
-		break;
 	case WM_IME_SETCONTEXT:
 		{
 			if(ActiveContext.IsValid())
@@ -1006,17 +1005,14 @@ int32 FWindowsTextInputMethodSystem::ProcessMessage(HWND hwnd, uint32 msg, WPARA
 
 			return IntCastChecked<int32>(DefWindowProc(hwnd, msg, wParam, lParam));
 		}
-		break;
 	case WM_IME_NOTIFY:
 		{
 			return IntCastChecked<int32>(DefWindowProc(hwnd, msg, wParam, lParam));
 		}
-		break;
 	case WM_IME_REQUEST:
 		{
 			return IntCastChecked<int32>(DefWindowProc(hwnd, msg, wParam, lParam));
 		}
-		break;
 	case WM_IME_STARTCOMPOSITION:
 		{
 			if(ActiveContext.IsValid())
@@ -1028,7 +1024,6 @@ int32 FWindowsTextInputMethodSystem::ProcessMessage(HWND hwnd, uint32 msg, WPARA
 
 			return IntCastChecked<int32>(DefWindowProc(hwnd, msg, wParam, lParam));
 		}
-		break;
 	case WM_IME_COMPOSITION:
 		{
 			if(ActiveContext.IsValid())
@@ -1147,7 +1142,6 @@ int32 FWindowsTextInputMethodSystem::ProcessMessage(HWND hwnd, uint32 msg, WPARA
 
 			return IntCastChecked<int32>(DefWindowProc(hwnd, msg, wParam, lParam));
 		}
-		break;
 	case WM_IME_ENDCOMPOSITION:
 		{
 			// On composition end, notify context of the end.
@@ -1160,21 +1154,18 @@ int32 FWindowsTextInputMethodSystem::ProcessMessage(HWND hwnd, uint32 msg, WPARA
 
 			return IntCastChecked<int32>(DefWindowProc(hwnd, msg, wParam, lParam));
 		}
-		break;
 	case WM_IME_CHAR:
 		{
 			// Suppress sending a WM_CHAR for this WM_IME_CHAR - the composition windows messages will have handled this.
 			UE_LOG(LogWindowsTextInputMethodSystem, Verbose, TEXT("Ignoring WM_IME_CHAR message."));
 			return 0;
 		}
-		break;
 	default:
 		{
 			UE_LOG(LogWindowsTextInputMethodSystem, Warning, TEXT("Unexpected windows message received for processing."));
 
 			return IntCastChecked<int32>(DefWindowProc(hwnd, msg, wParam, lParam));
 		}
-		break;
 	}
 }
 
diff --git a/Engine/Source/Runtime/AutoRTFM/Private/ExternAPI.cpp b/Engine/Source/Runtime/AutoRTFM/Private/ExternAPI.cpp
--- a/Engine/Source/Runtime/AutoRTFM/Private/ExternAPI.cpp
+++ b/Engine/Source/Runtime/AutoRTFM/Private/ExternAPI.cpp
@@ -55,7 +55,6 @@ void DefaultLog(const char* File, int Line, [[maybe_unused]] void* ProgramCounte
 		case autortfm_log_fatal:
 			fprintf(stderr, "AutoRTFM %s:%d [FATAL]: %s", File, Line, Buffer);
 			__builtin_trap();
-			break;
 	}
 }
 
diff --git a/Engine/Source/Runtime/AutoRTFM/Private/Utils.cpp b/Engine/Source/Runtime/AutoRTFM/Private/Utils.cpp
--- a/Engine/Source/Runtime/AutoRTFM/Private/Utils.cpp
+++ b/Engine/Source/Runtime/AutoRTFM/Private/Utils.cpp
@@ -84,7 +84,6 @@ UE_AUTORTFM_API UE_AUTORTFM_FORCENOINLINE void ReportError(const char* File, int
 			}
 			// The `autortfm_log_fatal` lines above probably triggered a crash already, but if we made it this far, force the crash.
 			__builtin_trap();
-			break;
 
 		case ForTheRuntime::EAutoRTFMInternalAbortActionState::Abort:
 			if (ForTheRuntime::GetEnsureOnInternalAbort())
diff --git a/Engine/Source/Runtime/Core/Private/Misc/PreLoadFile.cpp b/Engine/Source/Runtime/Core/Private/Misc/PreLoadFile.cpp
--- a/Engine/Source/Runtime/Core/Private/Misc/PreLoadFile.cpp
+++ b/Engine/Source/Runtime/Core/Private/Misc/PreLoadFile.cpp
@@ -28,7 +28,7 @@ struct FPreloadFileRegistry
 #endif
 
 FPreLoadFile::FPreLoadFile(const TCHAR* InPath)
-	: FDelayedAutoRegisterHelper(STATS ? EDelayedRegisterRunPhase::StatSystemReady : EDelayedRegisterRunPhase::FileSystemReady, [this] {
+	: FDelayedAutoRegisterHelper([] { if constexpr (STATS) { return EDelayedRegisterRunPhase::StatSystemReady; } else { return EDelayedRegisterRunPhase::FileSystemReady; }}(), [this] {
 		bSystemNoLongerTakingRequests = false;
 #if PLATFORM_CAN_ASYNC_PRELOAD_FILES
 		KickOffRead();
diff --git a/Engine/Source/Runtime/Core/Private/Windows/WindowsPlatformMisc.cpp b/Engine/Source/Runtime/Core/Private/Windows/WindowsPlatformMisc.cpp
--- a/Engine/Source/Runtime/Core/Private/Windows/WindowsPlatformMisc.cpp
+++ b/Engine/Source/Runtime/Core/Private/Windows/WindowsPlatformMisc.cpp
@@ -806,21 +806,16 @@ const TCHAR* LexToString(ETDRLevel TDRLevel)
 	{
 	case ETDRLevel::Disabled:
 		return TEXT("Disabled");
-		break;
 	case ETDRLevel::BugCheck:
 		return TEXT("BugCheck");
-		break;
 	case ETDRLevel::RecoverVGA:
 		return TEXT("RecoverVGA");
-		break;
 	case ETDRLevel::Recover:
 		return TEXT("Recover");
-		break;
 	case ETDRLevel::Invalid:
 		[[fallthrough]];
 	default:
 		return TEXT("Invalid");
-		break;
 	}
 }
 
@@ -1836,7 +1831,6 @@ private:
 						SetDialogTextInClipboard(HandleWnd);
 						break;
 					}
-					break;
 					default:
 						return false;
 				}
@@ -1854,7 +1848,6 @@ private:
 					return false;
 				}
 			}
-			break;
 			case WM_PAINT:
 			{
 				PAINTSTRUCT Paint;
