Fix/Silence unreachable code warnings
#rnx
#Synced-CL 52566575

diff --git a/Engine/Plugins/Media/AvfMedia/Source/AvfMedia/Private/Shared/AvfMediaUtils.cpp b/Engine/Plugins/Media/AvfMedia/Source/AvfMedia/Private/Shared/AvfMediaUtils.cpp
--- a/Engine/Plugins/Media/AvfMedia/Source/AvfMedia/Private/Shared/AvfMediaUtils.cpp
+++ b/Engine/Plugins/Media/AvfMedia/Source/AvfMedia/Private/Shared/AvfMediaUtils.cpp
@@ -130,8 +130,6 @@ namespace AvfMedia
 			static FString ReadPathBase = FString([[NSBundle mainBundle] bundlePath]) + TEXT("/cookeddata/");
 			return ReadPathBase + Result.ToLower();
 		}
-
-		return Result;
 	}
 #endif
 }
diff --git a/Engine/Plugins/Runtime/AR/AppleAR/AppleARKit/Source/AppleARKit/Private/AppleARKitConversion.cpp b/Engine/Plugins/Runtime/AR/AppleAR/AppleARKit/Source/AppleARKit/Private/AppleARKitConversion.cpp
--- a/Engine/Plugins/Runtime/AR/AppleAR/AppleARKit/Source/AppleARKit/Private/AppleARKitConversion.cpp
+++ b/Engine/Plugins/Runtime/AR/AppleAR/AppleARKit/Source/AppleARKit/Private/AppleARKitConversion.cpp
@@ -310,11 +310,9 @@ NSArray<ARVideoFormat*>* FAppleARKitConversion::GetSupportedVideoFormats(EARSess
 		{
 			case EARSessionType::Orientation:
 				return AROrientationTrackingConfiguration.supportedVideoFormats;
-				break;
 				
 			case EARSessionType::World:
 				return ARWorldTrackingConfiguration.supportedVideoFormats;
-				break;
 				
 			case EARSessionType::Face:
 				if (auto Implementation = GetFaceARSupport())
diff --git a/Engine/Plugins/Runtime/AR/Google/GoogleARCoreServices/Source/GoogleARCoreServices/Private/GoogleARCoreCloudARPinManager.cpp b/Engine/Plugins/Runtime/AR/Google/GoogleARCoreServices/Source/GoogleARCoreServices/Private/GoogleARCoreCloudARPinManager.cpp
--- a/Engine/Plugins/Runtime/AR/Google/GoogleARCoreServices/Source/GoogleARCoreServices/Private/GoogleARCoreCloudARPinManager.cpp
+++ b/Engine/Plugins/Runtime/AR/Google/GoogleARCoreServices/Source/GoogleARCoreServices/Private/GoogleARCoreCloudARPinManager.cpp
@@ -451,9 +451,9 @@ UCloudARPin* FGoogleARCoreCloudARPinManager::CreateAndHostCloudARPin(UARPin* Pin
 	}
 
 	return NewCloudARPin;
-#endif
-
+#else
 	return nullptr;
+#endif
 }
 
 UCloudARPin* FGoogleARCoreCloudARPinManager::ResolveAndCreateCloudARPin(FString CloudID, EARPinCloudTaskResult& OutTaskResult)
@@ -484,9 +484,9 @@ UCloudARPin* FGoogleARCoreCloudARPinManager::ResolveAndCreateCloudARPin(FString
 	ArPose_destroy(PoseHandle);
 
 	return NewCloudARPin;
-#endif
-
+#else
 	return nullptr;
+#endif
 }
 
 #else // ARCORE_USE_OLD_CLOUD_ANCHOR_ASYNC
diff --git a/Engine/Plugins/Runtime/AppleImageUtils/Source/AppleImageUtils/Private/AppleImageUtilsPlugin.cpp b/Engine/Plugins/Runtime/AppleImageUtils/Source/AppleImageUtils/Private/AppleImageUtilsPlugin.cpp
--- a/Engine/Plugins/Runtime/AppleImageUtils/Source/AppleImageUtils/Private/AppleImageUtilsPlugin.cpp
+++ b/Engine/Plugins/Runtime/AppleImageUtils/Source/AppleImageUtils/Private/AppleImageUtilsPlugin.cpp
@@ -312,8 +312,9 @@ TSharedPtr<FAppleImageUtilsConversionTaskBase, ESPMode::ThreadSafe> FAppleImageU
 		ConversionTask->MarkComplete();
 	});
 	return ConversionTask;
-#endif
+#else
 	return MakeShared<FAppleImageUtilsFailedConversionTask, ESPMode::ThreadSafe>(TEXT("ConvertToHEIF requires iOS 11.0+ or macOS 10.13.4+"));
+#endif
 }
 
 TSharedPtr<FAppleImageUtilsConversionTaskBase, ESPMode::ThreadSafe> FAppleImageUtilsPlugin::ConvertToPNG(UTexture* SourceImage, bool bWantColor, bool bUseGpu, float Scale, ETextureRotationDirection Rotate)
@@ -343,8 +344,9 @@ TSharedPtr<FAppleImageUtilsConversionTaskBase, ESPMode::ThreadSafe> FAppleImageU
 		ConversionTask->MarkComplete();
 	});
 	return ConversionTask;
-#endif
+#else
 	return MakeShared<FAppleImageUtilsFailedConversionTask, ESPMode::ThreadSafe>(TEXT("ConvertToPNG requires iOS 11.0+ or macOS 10.13.4+"));
+#endif
 }
 
 TSharedPtr<FAppleImageUtilsConversionTaskBase, ESPMode::ThreadSafe> FAppleImageUtilsPlugin::ConvertToTIFF(UTexture* SourceImage, bool bWantColor, bool bUseGpu, float Scale, ETextureRotationDirection Rotate)
@@ -374,8 +376,9 @@ TSharedPtr<FAppleImageUtilsConversionTaskBase, ESPMode::ThreadSafe> FAppleImageU
 		ConversionTask->MarkComplete();
 	});
 	return ConversionTask;
-#endif
+#else
 	return MakeShared<FAppleImageUtilsFailedConversionTask, ESPMode::ThreadSafe>(TEXT("ConvertToTIFF requires iOS 11.0+ or macOS 10.13.4+"));
+#endif
 }
 
 #if SUPPORTS_IMAGE_UTILS_1_0
diff --git a/Engine/Plugins/Runtime/AppleMoviePlayer/Source/AppleMoviePlayer/Private/AppleMovieStreamer.cpp b/Engine/Plugins/Runtime/AppleMoviePlayer/Source/AppleMoviePlayer/Private/AppleMovieStreamer.cpp
--- a/Engine/Plugins/Runtime/AppleMoviePlayer/Source/AppleMoviePlayer/Private/AppleMovieStreamer.cpp
+++ b/Engine/Plugins/Runtime/AppleMoviePlayer/Source/AppleMoviePlayer/Private/AppleMovieStreamer.cpp
@@ -69,9 +69,9 @@ static FString ConvertToNativePath(const FString& Filename, bool bForWrite)
 			return ReadPathBase + Result.ToLower();
 		}
 	}
-#endif
-	
+#else
 	return Result;
+#endif
 }
 
 //
diff --git a/Engine/Source/Programs/UnrealVirtualizationTool/Private/UnrealVirtualizationToolApp.cpp b/Engine/Source/Programs/UnrealVirtualizationTool/Private/UnrealVirtualizationToolApp.cpp
--- a/Engine/Source/Programs/UnrealVirtualizationTool/Private/UnrealVirtualizationToolApp.cpp
+++ b/Engine/Source/Programs/UnrealVirtualizationTool/Private/UnrealVirtualizationToolApp.cpp
@@ -402,26 +402,20 @@ EInitResult FUnrealVirtualizationToolApp::CreateCommandFromString(const FString&
 		case EMode::Changelist:
 			CurrentCommand = CreateCommand<FVirtualizeLegacyChangeListCommand>(LexToString(Mode), Cmdline);
 			return CurrentCommand.IsValid() ? EInitResult::Success : EInitResult::Error;
-			break;
 		case EMode::PackageList:
 			CurrentCommand = CreateCommand<FVirtualizeLegacyPackageListCommand>(LexToString(Mode), Cmdline);
 			return CurrentCommand.IsValid() ? EInitResult::Success : EInitResult::Error;
-			break;
 		case EMode::Virtualize:
 			CurrentCommand = CreateCommand<FVirtualizeCommand>(CommandName, Cmdline);
 			return CurrentCommand.IsValid() ? EInitResult::Success : EInitResult::Error;
-			break;
 		case EMode::Rehydrate:
 			CurrentCommand = CreateCommand<FRehydrateCommand>(CommandName, Cmdline);
 			return CurrentCommand.IsValid() ? EInitResult::Success : EInitResult::Error;
-			break;
 		case EMode::Unknown:
 		default:
 			UE_LOGF(LogVirtualizationTool, Error, "Unexpected value for the cmdline switch 'Mode', this is a required parameter!");
 			PrintCmdLineHelp();
 			return EInitResult::Error;
-
-			break;
 	}
 }
 
diff --git a/Engine/Source/Runtime/ApplicationCore/Private/IOS/IOSAppDelegate.cpp b/Engine/Source/Runtime/ApplicationCore/Private/IOS/IOSAppDelegate.cpp
--- a/Engine/Source/Runtime/ApplicationCore/Private/IOS/IOSAppDelegate.cpp
+++ b/Engine/Source/Runtime/ApplicationCore/Private/IOS/IOSAppDelegate.cpp
@@ -2008,11 +2008,11 @@ extern double GCStartTime;
 	switch(Orientation)
 	{
 		default:
-		case UIInterfaceOrientationUnknown : return EDeviceScreenOrientation::Unknown; break;
-		case UIInterfaceOrientationPortrait : return EDeviceScreenOrientation::Portrait; break;
-		case UIInterfaceOrientationPortraitUpsideDown : return EDeviceScreenOrientation::PortraitUpsideDown; break;
-		case UIInterfaceOrientationLandscapeLeft : return EDeviceScreenOrientation::LandscapeLeft; break;
-		case UIInterfaceOrientationLandscapeRight : return EDeviceScreenOrientation::LandscapeRight; break;
+		case UIInterfaceOrientationUnknown : return EDeviceScreenOrientation::Unknown;
+		case UIInterfaceOrientationPortrait : return EDeviceScreenOrientation::Portrait;
+		case UIInterfaceOrientationPortraitUpsideDown : return EDeviceScreenOrientation::PortraitUpsideDown;
+		case UIInterfaceOrientationLandscapeLeft : return EDeviceScreenOrientation::LandscapeLeft;
+		case UIInterfaceOrientationLandscapeRight : return EDeviceScreenOrientation::LandscapeRight;
 	}
 }
 
@@ -2021,13 +2021,13 @@ extern double GCStartTime;
 	switch(Orientation)
 	{
 		default:
-		case UIDeviceOrientationUnknown : return UIInterfaceOrientationUnknown; break;
-		case UIDeviceOrientationPortrait : return UIInterfaceOrientationPortrait; break;
-		case UIDeviceOrientationPortraitUpsideDown : return UIInterfaceOrientationPortraitUpsideDown; break;
+		case UIDeviceOrientationUnknown : return UIInterfaceOrientationUnknown;
+		case UIDeviceOrientationPortrait : return UIInterfaceOrientationPortrait;
+		case UIDeviceOrientationPortraitUpsideDown : return UIInterfaceOrientationPortraitUpsideDown;
 		// Note UIDeviceOrientation and UIInterfaceOrientation have have lansdcape left and right switched
 		// https://developer.apple.com/documentation/uikit/uiinterfaceorientation
-		case UIDeviceOrientationLandscapeLeft : return UIInterfaceOrientationLandscapeRight; break;
-		case UIDeviceOrientationLandscapeRight : return UIInterfaceOrientationLandscapeLeft; break;
+		case UIDeviceOrientationLandscapeLeft : return UIInterfaceOrientationLandscapeRight;
+		case UIDeviceOrientationLandscapeRight : return UIInterfaceOrientationLandscapeLeft;
 	}
 }
 
@@ -2036,11 +2036,11 @@ extern double GCStartTime;
 	switch(Orientation)
 	{
 		default:
-		case UIInterfaceOrientationUnknown : return UIInterfaceOrientationMaskAll; break;
-		case UIInterfaceOrientationPortrait : return UIInterfaceOrientationMaskPortrait; break;
-		case UIInterfaceOrientationPortraitUpsideDown : return UIInterfaceOrientationMaskPortraitUpsideDown; break;
-		case UIInterfaceOrientationLandscapeLeft : return UIInterfaceOrientationMaskLandscapeLeft; break;
-		case UIInterfaceOrientationLandscapeRight : return UIInterfaceOrientationMaskLandscapeRight; break;
+		case UIInterfaceOrientationUnknown : return UIInterfaceOrientationMaskAll;
+		case UIInterfaceOrientationPortrait : return UIInterfaceOrientationMaskPortrait;
+		case UIInterfaceOrientationPortraitUpsideDown : return UIInterfaceOrientationMaskPortraitUpsideDown;
+		case UIInterfaceOrientationLandscapeLeft : return UIInterfaceOrientationMaskLandscapeLeft;
+		case UIInterfaceOrientationLandscapeRight : return UIInterfaceOrientationMaskLandscapeRight;
 	}
 }
 
@@ -2049,11 +2049,11 @@ extern double GCStartTime;
 	switch(Orientation)
 	{
 		default:
-		case EDeviceScreenOrientation::Unknown : return UIInterfaceOrientationMaskAll; break;
-		case EDeviceScreenOrientation::Portrait : return UIInterfaceOrientationMaskPortrait; break;
-		case EDeviceScreenOrientation::PortraitUpsideDown : return UIInterfaceOrientationMaskPortraitUpsideDown; break;
-		case EDeviceScreenOrientation::LandscapeLeft : return UIInterfaceOrientationMaskLandscapeLeft; break;
-		case EDeviceScreenOrientation::LandscapeRight : return UIInterfaceOrientationMaskLandscapeRight; break;
+		case EDeviceScreenOrientation::Unknown : return UIInterfaceOrientationMaskAll;
+		case EDeviceScreenOrientation::Portrait : return UIInterfaceOrientationMaskPortrait;
+		case EDeviceScreenOrientation::PortraitUpsideDown : return UIInterfaceOrientationMaskPortraitUpsideDown;
+		case EDeviceScreenOrientation::LandscapeLeft : return UIInterfaceOrientationMaskLandscapeLeft;
+		case EDeviceScreenOrientation::LandscapeRight : return UIInterfaceOrientationMaskLandscapeRight;
 	}
 }
 
diff --git a/Engine/Source/Runtime/Core/Private/Apple/ApplePlatformMemory.cpp b/Engine/Source/Runtime/Core/Private/Apple/ApplePlatformMemory.cpp
--- a/Engine/Source/Runtime/Core/Private/Apple/ApplePlatformMemory.cpp
+++ b/Engine/Source/Runtime/Core/Private/Apple/ApplePlatformMemory.cpp
@@ -539,7 +539,6 @@ void* _Nullable FApplePlatformMemory::BinnedAllocFromOS(SIZE_T Size)
 	{
 		FPlatformMemory::OnOutOfMemory(ActualSizeMapped, ExpectedAlignment);
 		// unreachable
-		return nullptr;
 	}
 
 	SIZE_T Offset = (reinterpret_cast<SIZE_T>(Pointer) % ExpectedAlignment);
@@ -556,7 +555,6 @@ void* _Nullable FApplePlatformMemory::BinnedAllocFromOS(SIZE_T Size)
 		{
 			FPlatformMemory::OnOutOfMemory(SizeToNextAlignedPointer, ExpectedAlignment);
 			// unreachable
-			return nullptr;
 		}
 
 		// account for reduced mmaped size
@@ -580,7 +578,6 @@ void* _Nullable FApplePlatformMemory::BinnedAllocFromOS(SIZE_T Size)
 		{
 			FPlatformMemory::OnOutOfMemory(TailSize, ExpectedAlignment);
 			// unreachable
-			return nullptr;
 		}
 	}
 
@@ -904,8 +901,9 @@ bool FApplePlatformMemory::CanOverallocateVirtualMemory()
 #if PLATFORM_IOS || PLATFORM_TVOS
 	static bool bHasExtendedVirtualAddressingEntitlement = FIOSPlatformMisc::IsEntitlementEnabled("com.apple.developer.kernel.extended-virtual-addressing");
 	return bHasExtendedVirtualAddressingEntitlement;
-#endif
+#else
 	return true;	// 64 bit Mac process can allocate ~18 exabytes of addressable space
+#endif
 }
 
 NS_ASSUME_NONNULL_END
diff --git a/Engine/Source/Runtime/Core/Private/IOS/IOSCallstackTrace.cpp b/Engine/Source/Runtime/Core/Private/IOS/IOSCallstackTrace.cpp
--- a/Engine/Source/Runtime/Core/Private/IOS/IOSCallstackTrace.cpp
+++ b/Engine/Source/Runtime/Core/Private/IOS/IOSCallstackTrace.cpp
@@ -72,9 +72,9 @@ FORCEINLINE uint32 FBacktracer::GetBacktraceId(uint64* StackFrames, uint32 NumSt
 	// Add to queue to be processed. This might block until there is room in the
 	// queue (i.e. the processing thread has caught up processing).
 	return CallstackTracer.AddCallstack(BacktraceEntry);
-#endif
-
+#else
 	return 0;
+#endif
 }
 
 ////////////////////////////////////////////////////////////////////////////////
diff --git a/Engine/Source/Runtime/Core/Private/IOS/IOSChunkInstaller.cpp b/Engine/Source/Runtime/Core/Private/IOS/IOSChunkInstaller.cpp
--- a/Engine/Source/Runtime/Core/Private/IOS/IOSChunkInstaller.cpp
+++ b/Engine/Source/Runtime/Core/Private/IOS/IOSChunkInstaller.cpp
@@ -178,7 +178,6 @@ bool FIOSChunkInstall::DebugStartNextChunk()
     
     //cannot call this function in a submission build.  Will fail cert.
     // create a request and ask for the bundle
-    return true;
 #endif
     return true;
 }
diff --git a/Engine/Source/Runtime/Core/Private/IOS/IOSMallocZone.cpp b/Engine/Source/Runtime/Core/Private/IOS/IOSMallocZone.cpp
--- a/Engine/Source/Runtime/Core/Private/IOS/IOSMallocZone.cpp
+++ b/Engine/Source/Runtime/Core/Private/IOS/IOSMallocZone.cpp
@@ -191,6 +191,5 @@ bool FIOSMallocCrashHandler::IsOnCrashedThread( void )
 	else
 	{
 		FPlatformProcess::SleepInfinite();
-		return false;
 	}
 }
diff --git a/Engine/Source/Runtime/Core/Private/IOS/IOSPlatformFile.cpp b/Engine/Source/Runtime/Core/Private/IOS/IOSPlatformFile.cpp
--- a/Engine/Source/Runtime/Core/Private/IOS/IOSPlatformFile.cpp
+++ b/Engine/Source/Runtime/Core/Private/IOS/IOSPlatformFile.cpp
@@ -1084,8 +1084,6 @@ FString FIOSPlatformFile::ConvertToPlatformPath(const FString& Filename, bool bF
             return ReadPathBase + Result.ToLower();
 		}
 	}
-
-	return Result;
 }
 
 int32 FIOSPlatformFile::OpenFileForMemMap(const FString& NormalizedFilename, EOpenMappedFlags OpenOptions, FString& OutFinalPath)
diff --git a/Engine/Source/Runtime/Core/Private/IOS/IOSPlatformMisc.cpp b/Engine/Source/Runtime/Core/Private/IOS/IOSPlatformMisc.cpp
--- a/Engine/Source/Runtime/Core/Private/IOS/IOSPlatformMisc.cpp
+++ b/Engine/Source/Runtime/Core/Private/IOS/IOSPlatformMisc.cpp
@@ -362,10 +362,10 @@ float FIOSPlatformMisc::GetDeviceTemperatureLevel()
 #if WITH_APPLICATION_CORE && !PLATFORM_TVOS
     switch ([[IOSAppDelegate GetDelegate] GetThermalState])
     {
-        case NSProcessInfoThermalStateNominal:	return (float)FCoreDelegates::ETemperatureSeverity::Good; break;
-        case NSProcessInfoThermalStateFair:		return (float)FCoreDelegates::ETemperatureSeverity::Bad; break;
-        case NSProcessInfoThermalStateSerious:	return (float)FCoreDelegates::ETemperatureSeverity::Serious; break;
-        case NSProcessInfoThermalStateCritical:	return (float)FCoreDelegates::ETemperatureSeverity::Critical; break;
+        case NSProcessInfoThermalStateNominal:	return (float)FCoreDelegates::ETemperatureSeverity::Good;
+        case NSProcessInfoThermalStateFair:		return (float)FCoreDelegates::ETemperatureSeverity::Bad;
+        case NSProcessInfoThermalStateSerious:	return (float)FCoreDelegates::ETemperatureSeverity::Serious;
+        case NSProcessInfoThermalStateCritical:	return (float)FCoreDelegates::ETemperatureSeverity::Critical;
     }
 #endif
 	return -1.0f;
@@ -382,10 +382,10 @@ EDeviceThermalState FIOSPlatformMisc::GetDeviceThermalState()
 	// Critical/Severe indicate that serious throttling is kicking in and the device is probably noticably hot to touch.
 	switch ([[IOSAppDelegate GetDelegate]GetThermalState] )
 	{
-	case NSProcessInfoThermalStateNominal:	return EDeviceThermalState::None; break;
-	case NSProcessInfoThermalStateFair:		return EDeviceThermalState::Light; break;
-	case NSProcessInfoThermalStateSerious:	return EDeviceThermalState::Moderate; break;
-	case NSProcessInfoThermalStateCritical:	return EDeviceThermalState::Severe; break;
+	case NSProcessInfoThermalStateNominal:	return EDeviceThermalState::None;
+	case NSProcessInfoThermalStateFair:		return EDeviceThermalState::Light;
+	case NSProcessInfoThermalStateSerious:	return EDeviceThermalState::Moderate;
+	case NSProcessInfoThermalStateCritical:	return EDeviceThermalState::Severe;
 	}
 #endif
 	return EDeviceThermalState::Unsupported;
@@ -396,8 +396,9 @@ bool FIOSPlatformMisc::IsInLowPowerMode()
 #if !PLATFORM_TVOS
     bool bInLowPowerMode = [[NSProcessInfo processInfo] isLowPowerModeEnabled];
     return bInLowPowerMode;
-#endif
+#else
     return false;
+#endif
 }
 
 bool FIOSPlatformMisc::IsDesignedForIpadOnVisionOS()
@@ -430,34 +431,24 @@ UIInterfaceOrientationMask GetUIInterfaceOrientationMask(EDeviceScreenOrientatio
 		// Fallthrough...
 	case EDeviceScreenOrientation::Unknown:
 		return UIInterfaceOrientationMaskAll;
-		break;
 	case EDeviceScreenOrientation::Portrait:
 		return UIInterfaceOrientationMaskPortrait;
-		break;
 	case EDeviceScreenOrientation::PortraitUpsideDown:
 		return UIInterfaceOrientationMaskPortraitUpsideDown;
-		break;
 	case EDeviceScreenOrientation::LandscapeLeft:
 		return UIInterfaceOrientationMaskLandscapeLeft;
-		break;
 	case EDeviceScreenOrientation::LandscapeRight:
 		return UIInterfaceOrientationMaskLandscapeRight;
-		break;
 	case EDeviceScreenOrientation::FaceUp:
 		return UIInterfaceOrientationMaskAll;
-		break;
 	case EDeviceScreenOrientation::FaceDown:
 		return UIInterfaceOrientationMaskAll;
-		break;
 	case EDeviceScreenOrientation::PortraitSensor:
 		return UIInterfaceOrientationMaskPortrait;
-		break;
 	case EDeviceScreenOrientation::LandscapeSensor:
 		return UIInterfaceOrientationMaskLandscape;
-		break;
 	case EDeviceScreenOrientation::FullSensor:
 		return UIInterfaceOrientationMaskAll;
-		break;
 	}
 }
 #endif
diff --git a/Engine/Source/Runtime/CrashReportCore/Private/Mac/MacErrorReport.cpp b/Engine/Source/Runtime/CrashReportCore/Private/Mac/MacErrorReport.cpp
--- a/Engine/Source/Runtime/CrashReportCore/Private/Mac/MacErrorReport.cpp
+++ b/Engine/Source/Runtime/CrashReportCore/Private/Mac/MacErrorReport.cpp
@@ -157,8 +157,9 @@ FText FMacErrorReport::DiagnoseReport() const
 			return FText::FromString("Failed to create diagnosis information.");
 		}
 	}
-#endif
+#else
 	return FText::FromString(TEXT("Failed to load CrashDebugHelper."));
+#endif
 }
 
 #undef LOCTEXT_NAMESPACE
