Fix for regular iOS touch input where touches were being eaten by the gesture support
#Synced-CL 52622993

diff --git a/Engine/Source/Runtime/ApplicationCore/Private/IOS/IOSView.cpp b/Engine/Source/Runtime/ApplicationCore/Private/IOS/IOSView.cpp
--- a/Engine/Source/Runtime/ApplicationCore/Private/IOS/IOSView.cpp
+++ b/Engine/Source/Runtime/ApplicationCore/Private/IOS/IOSView.cpp
@@ -1122,18 +1122,21 @@ self.accessibilityElements = @[Window.accessibilityContainer];
     UIPanGestureRecognizer* Pan = [[UIPanGestureRecognizer alloc] initWithTarget:self action:@selector(HandlePanGesture:)];
     Pan.minimumNumberOfTouches = 1;
     Pan.maximumNumberOfTouches = 1;
+    Pan.cancelsTouchesInView = NO;
     Pan.delegate = self;
     [self addGestureRecognizer:Pan];
     [Pan release];
 
     // Pinch — two-finger scale
     UIPinchGestureRecognizer* Pinch = [[UIPinchGestureRecognizer alloc] initWithTarget:self action:@selector(HandlePinchGesture:)];
+    Pinch.cancelsTouchesInView = NO;
     Pinch.delegate = self;
     [self addGestureRecognizer:Pinch];
     [Pinch release];
 
     // Rotate — two-finger rotation
     UIRotationGestureRecognizer* Rotation = [[UIRotationGestureRecognizer alloc] initWithTarget:self action:@selector(HandleRotateGesture:)];
+    Rotation.cancelsTouchesInView = NO;
     Rotation.delegate = self;
     [self addGestureRecognizer:Rotation];
     [Rotation release];
@@ -1141,12 +1144,14 @@ self.accessibilityElements = @[Window.accessibilityContainer];
     // Tap — single tap
     UITapGestureRecognizer* Tap = [[UITapGestureRecognizer alloc] initWithTarget:self action:@selector(HandleTapGesture:)];
     Tap.numberOfTapsRequired = 1;
+    Tap.cancelsTouchesInView = NO;
     Tap.delegate = self;
     [self addGestureRecognizer:Tap];
     [Tap release];
 
     // Long Press — continuous (Began → Changed → Ended)
     UILongPressGestureRecognizer* LongPress = [[UILongPressGestureRecognizer alloc] initWithTarget:self action:@selector(HandleLongPressGesture:)];
+    LongPress.cancelsTouchesInView = NO;
     LongPress.delegate = self;
     [self addGestureRecognizer:LongPress];
     [LongPress release];
