Add ConfigRules support to Windows MSG build

Also add SRC_BuildFlavor support.

#Synced-CL 52669770

diff --git a/Engine/Source/Runtime/ApplicationCore/Private/Windows/WindowsConfigRules.cpp b/Engine/Source/Runtime/ApplicationCore/Private/Windows/WindowsConfigRules.cpp
--- a/Engine/Source/Runtime/ApplicationCore/Private/Windows/WindowsConfigRules.cpp
+++ b/Engine/Source/Runtime/ApplicationCore/Private/Windows/WindowsConfigRules.cpp
@@ -345,6 +345,11 @@ void FWindowsConfigRules::SetStandardPredefinedVars()
 	FString BuildVersion = FApp::GetBuildVersion();
 	FString BuildConfig = LexToString(FApp::GetBuildConfiguration());
 	FString BuildTarget = LexToString(FApp::GetBuildTargetType());
+#ifdef UE_BUILD_FLAVOR
+	const FString BuildFlavor(UE_BUILD_FLAVOR);
+#else
+	const FString BuildFlavor;
+#endif
 	FString TotalPhysicalGB = FString::FromInt(static_cast<int32>(FPlatformMemory::GetConstants().TotalPhysicalGB));
 	FString TotalVirtualGB = FString::FromInt(static_cast<int32>(FPlatformMemory::GetConstants().TotalVirtual / 1024 / 1024 / 1024));
 
@@ -352,6 +357,7 @@ void FWindowsConfigRules::SetStandardPredefinedVars()
 	PredefinedConfigRuleVars.Emplace(TEXT("SRC_BuildVersion"), BuildVersion);
 	PredefinedConfigRuleVars.Emplace(TEXT("SRC_BuildConfig"), BuildConfig);
 	PredefinedConfigRuleVars.Emplace(TEXT("SRC_BuildTarget"), BuildTarget);
+	PredefinedConfigRuleVars.Emplace(TEXT("SRC_BuildFlavor"), BuildFlavor);
 	PredefinedConfigRuleVars.Emplace(TEXT("SRC_MakeAndModel"), MakeAndModel);
 	PredefinedConfigRuleVars.Emplace(TEXT("SRC_CPUVendor"), CPUVendor);
 	PredefinedConfigRuleVars.Emplace(TEXT("SRC_CPUBrand"), CPUBrand);
